Refactor plugin to move out the animation graph portion into a separate plugin to decouple the two

diff --git a/Engine/Plugins/Experimental/AnimNext/Config/BaseAnimNext.ini b/Engine/Plugins/Experimental/AnimNext/Config/BaseAnimNext.ini
--- a/Engine/Plugins/Experimental/AnimNext/Config/BaseAnimNext.ini
+++ b/Engine/Plugins/Experimental/AnimNext/Config/BaseAnimNext.ini
@@ -7,15 +7,6 @@
 ; Traits rename
 +StructRedirects=(OldName="/Script/AnimNext.AnimNextDecoratorHandle", NewName="/Script/AnimNext.AnimNextTraitHandle")
 +StructRedirects=(OldName="/Script/AnimNext.RigUnit_AnimNextDecoratorStack", NewName="/Script/AnimNext.RigUnit_AnimNextTraitStack")
-+StructRedirects=(OldName="/Script/AnimNext.AnimNextBlendByBoolDecoratorSharedData", NewName="/Script/AnimNext.AnimNextBlendByBoolTraitSharedData")
-+StructRedirects=(OldName="/Script/AnimNext.AnimNextBlendInertializerDecoratorSharedData", NewName="/Script/AnimNext.AnimNextBlendInertializerTraitSharedData")
-+StructRedirects=(OldName="/Script/AnimNext.AnimNextBlendSmootherDecoratorSharedData", NewName="/Script/AnimNext.AnimNextBlendSmootherTraitSharedData")
-+StructRedirects=(OldName="/Script/AnimNext.AnimNextBlendSmootherPerBoneDecoratorSharedData", NewName="/Script/AnimNext.AnimNextBlendSmootherPerBoneTraitSharedData")
-+StructRedirects=(OldName="/Script/AnimNext.AnimNextBlendTwoWayDecoratorSharedData", NewName="/Script/AnimNext.AnimNextBlendTwoWayTraitSharedData")
-+StructRedirects=(OldName="/Script/AnimNext.AnimNextReferencePoseDecoratorSharedData", NewName="/Script/AnimNext.AnimNextReferencePoseTraitSharedData")
-+StructRedirects=(OldName="/Script/AnimNext.AnimNextSequencePlayerDecoratorSharedData", NewName="/Script/AnimNext.AnimNextSequencePlayerTraitSharedData")
-+StructRedirects=(OldName="/Script/AnimNext.AnimNextSubGraphHostDecoratorSharedData", NewName="/Script/AnimNext.AnimNextSubGraphHostTraitSharedData")
-+StructRedirects=(OldName="/Script/AnimNext.AnimNextSynchronizeUsingGroupsDecoratorSharedData", NewName="/Script/AnimNext.AnimNextSynchronizeUsingGroupsTraitSharedData")
 
 ; Parameter block/graph merge
 +ClassRedirects=(OldName="/Script/AnimNext.AnimNextParameterBlock", NewName="/Script/AnimNext.AnimNextGraph")
diff --git a/Engine/Plugins/Experimental/AnimNext/Source/AnimNext/Internal/Graph/AnimGraph/AnimNode_AnimNextGraph.h b/Engine/Plugins/Experimental/AnimNext/Source/AnimNext/Internal/Graph/AnimGraph/AnimNode_AnimNextGraph.h
deleted file mode 100644
diff --git a/Engine/Plugins/Experimental/AnimNext/Source/AnimNext/Internal/Scheduler/AnimNextSchedule.h b/Engine/Plugins/Experimental/AnimNext/Source/AnimNext/Internal/Scheduler/AnimNextSchedule.h
--- a/Engine/Plugins/Experimental/AnimNext/Source/AnimNext/Internal/Scheduler/AnimNextSchedule.h
+++ b/Engine/Plugins/Experimental/AnimNext/Source/AnimNext/Internal/Scheduler/AnimNextSchedule.h
@@ -266,6 +266,8 @@ private:
 	TArray<TObjectPtr<UAnimNextScheduleEntry>> Entries;
 #endif
 
+	// TODO: By having the structs below inline, it prevents us from defining them in external plugins that we know nothing about
+
 	// TEMP: Instructions derived from the entries above
 	UPROPERTY(NonTransactional)
 	TArray<FAnimNextScheduleInstruction> Instructions;
diff --git a/Engine/Plugins/Experimental/AnimNext/Source/AnimNext/Private/Graph/AnimGraph/AnimGraphParamStackScope.cpp b/Engine/Plugins/Experimental/AnimNext/Source/AnimNext/Private/Graph/AnimGraph/AnimGraphParamStackScope.cpp
--- a/Engine/Plugins/Experimental/AnimNext/Source/AnimNext/Private/Graph/AnimGraph/AnimGraphParamStackScope.cpp
+++ b/Engine/Plugins/Experimental/AnimNext/Source/AnimNext/Private/Graph/AnimGraph/AnimGraphParamStackScope.cpp
@@ -1,6 +1,6 @@
 // Copyright Epic Games, Inc. All Rights Reserved.
 
-#include "AnimGraphParamStackScope.h"
+#include "Graph/AnimGraphParamStackScope.h"
 #include "Param/ParamStack.h"
 #include "Animation/AnimNodeBase.h"
 
diff --git a/Engine/Plugins/Experimental/AnimNext/Source/AnimNext/Private/Graph/AnimGraph/AnimGraphParamStackScope.h b/Engine/Plugins/Experimental/AnimNext/Source/AnimNext/Private/Graph/AnimGraph/AnimGraphParamStackScope.h
deleted file mode 100644
diff --git a/Engine/Plugins/Experimental/AnimNext/Source/AnimNext/Private/Graph/AnimGraph/AnimNodeExposedValueHandler_AnimNextParameters.cpp b/Engine/Plugins/Experimental/AnimNext/Source/AnimNext/Private/Graph/AnimGraph/AnimNodeExposedValueHandler_AnimNextParameters.cpp
--- a/Engine/Plugins/Experimental/AnimNext/Source/AnimNext/Private/Graph/AnimGraph/AnimNodeExposedValueHandler_AnimNextParameters.cpp
+++ b/Engine/Plugins/Experimental/AnimNext/Source/AnimNext/Private/Graph/AnimGraph/AnimNodeExposedValueHandler_AnimNextParameters.cpp
@@ -5,7 +5,7 @@
 #include "Animation/AnimNodeBase.h"
 #include "Animation/AnimSubsystem_PropertyAccess.h"
 #include "Animation/AnimClassInterface.h"
-#include "AnimGraphParamStackScope.h"
+#include "Graph/AnimGraphParamStackScope.h"
 
 void FAnimNodeExposedValueHandler_AnimNextParameters::Initialize(const UClass* InClass)
 {
diff --git a/Engine/Plugins/Experimental/AnimNext/Source/AnimNext/Private/Graph/AnimGraph/AnimNode_AnimNextGraph.cpp b/Engine/Plugins/Experimental/AnimNext/Source/AnimNext/Private/Graph/AnimGraph/AnimNode_AnimNextGraph.cpp
deleted file mode 100644
diff --git a/Engine/Plugins/Experimental/AnimNext/Source/AnimNext/Private/Graph/AnimGraph/AnimNode_AnimNextParameters.cpp b/Engine/Plugins/Experimental/AnimNext/Source/AnimNext/Private/Graph/AnimGraph/AnimNode_AnimNextParameters.cpp
--- a/Engine/Plugins/Experimental/AnimNext/Source/AnimNext/Private/Graph/AnimGraph/AnimNode_AnimNextParameters.cpp
+++ b/Engine/Plugins/Experimental/AnimNext/Source/AnimNext/Private/Graph/AnimGraph/AnimNode_AnimNextParameters.cpp
@@ -3,7 +3,7 @@
 #include "Graph/AnimGraph/AnimNode_AnimNextParameters.h"
 #include "Param/ParamStack.h"
 #include "Graph/AnimNextGraph.h"
-#include "AnimGraphParamStackScope.h"
+#include "Graph/AnimGraphParamStackScope.h"
 #include "Param/ParametersProxy.h"
 
 #include UE_INLINE_GENERATED_CPP_BY_NAME(AnimNode_AnimNextParameters)
diff --git a/Engine/Plugins/Experimental/AnimNext/Source/AnimNext/Private/Graph/AnimNextGraphInstancePtr.cpp b/Engine/Plugins/Experimental/AnimNext/Source/AnimNext/Private/Graph/AnimNextGraphInstancePtr.cpp
--- a/Engine/Plugins/Experimental/AnimNext/Source/AnimNext/Private/Graph/AnimNextGraphInstancePtr.cpp
+++ b/Engine/Plugins/Experimental/AnimNext/Source/AnimNext/Private/Graph/AnimNextGraphInstancePtr.cpp
@@ -22,7 +22,6 @@ void FAnimNextGraphInstancePtr::Release()
 		{
 			const UAnimNextGraph* Graph = Impl->GetGraph();
 			FScopeLock Lock(&Graph->GraphInstancesLock);
-			check(Graph->GraphInstances.Contains(Impl.Get()));
 			Graph->GraphInstances.Remove(Impl.Get());
 		}
 #endif
diff --git a/Engine/Plugins/Experimental/AnimNext/Source/AnimNext/Private/Graph/SyncGroup_GraphInstanceComponent.cpp b/Engine/Plugins/Experimental/AnimNext/Source/AnimNext/Private/Graph/SyncGroup_GraphInstanceComponent.cpp
deleted file mode 100644
diff --git a/Engine/Plugins/Experimental/AnimNext/Source/AnimNext/Private/Graph/SyncGroup_GraphInstanceComponent.h b/Engine/Plugins/Experimental/AnimNext/Source/AnimNext/Private/Graph/SyncGroup_GraphInstanceComponent.h
deleted file mode 100644
diff --git a/Engine/Plugins/Experimental/AnimNext/Source/AnimNext/Private/Module.cpp b/Engine/Plugins/Experimental/AnimNext/Source/AnimNext/Private/Module.cpp
--- a/Engine/Plugins/Experimental/AnimNext/Source/AnimNext/Private/Module.cpp
+++ b/Engine/Plugins/Experimental/AnimNext/Source/AnimNext/Private/Module.cpp
@@ -1,10 +1,10 @@
 // Copyright Epic Games, Inc. All Rights Reserved.
 
 #include "CoreMinimal.h"
+#include "IAnimNextModule.h"
 #include "AnimNextConfig.h"
 #include "Animation/BlendProfile.h"
 #include "Curves/CurveFloat.h"
-#include "Modules/ModuleInterface.h"
 #include "Modules/ModuleManager.h"
 #include "Misc/CoreDelegates.h"
 #include "DataRegistry.h"
@@ -33,7 +33,7 @@
 namespace UE::AnimNext
 {
 
-class FModule : public IModuleInterface
+class FModule : public IAnimNextModule
 {
 public:
 	virtual void StartupModule() override
@@ -97,6 +97,34 @@ public:
 #endif
 	}
 
+	const IAnimNextAnimGraph* AnimGraphImpl = nullptr;
+
+	virtual void RegisterAnimNextAnimGraph(const IAnimNextAnimGraph& InAnimGraphImpl) override
+	{
+		AnimGraphImpl = &InAnimGraphImpl;
+	}
+
+	virtual void UnregisterAnimNextAnimGraph() override
+	{
+		AnimGraphImpl = nullptr;
+	}
+
+	virtual void UpdateGraph(FAnimNextGraphInstancePtr& GraphInstance, float DeltaTime) override
+	{
+		if (AnimGraphImpl != nullptr)
+		{
+			AnimGraphImpl->UpdateGraph(GraphInstance, DeltaTime);
+		}
+	}
+
+	virtual void EvaluateGraph(FAnimNextGraphInstancePtr& GraphInstance, const UE::AnimNext::FReferencePose& RefPose, int32 GraphLODLevel, FLODPoseHeap& OutputPose) const override
+	{
+		if (AnimGraphImpl != nullptr)
+		{
+			AnimGraphImpl->EvaluateGraph(GraphInstance, RefPose, GraphLODLevel, OutputPose);
+		}
+	}
+
 #if WITH_ANIMNEXT_CONSOLE_COMMANDS
 	TArray<IConsoleObject*> ConsoleCommands;
 
@@ -262,6 +290,11 @@ public:
 #endif
 };
 
+IAnimNextModule& IAnimNextModule::Get()
+{
+	return FModuleManager::LoadModuleChecked<IAnimNextModule>(TEXT("AnimNext"));
+}
+
 }
 
 IMPLEMENT_MODULE(UE::AnimNext::FModule, AnimNext)
diff --git a/Engine/Plugins/Experimental/AnimNext/Source/AnimNext/Private/Scheduler/AnimNextScheduleGraphTask.cpp b/Engine/Plugins/Experimental/AnimNext/Source/AnimNext/Private/Scheduler/AnimNextScheduleGraphTask.cpp
--- a/Engine/Plugins/Experimental/AnimNext/Source/AnimNext/Private/Scheduler/AnimNextScheduleGraphTask.cpp
+++ b/Engine/Plugins/Experimental/AnimNext/Source/AnimNext/Private/Scheduler/AnimNextScheduleGraphTask.cpp
@@ -2,12 +2,10 @@
 
 #include "Scheduler/AnimNextScheduleGraphTask.h"
 
+#include "IAnimNextModule.h"
 #include "Scheduler/ScheduleContext.h"
 #include "Graph/AnimNextGraph.h"
 #include "Context.h"
-#include "TraitInterfaces/IEvaluate.h"
-#include "TraitInterfaces/IUpdate.h"
-#include "EvaluationVM/EvaluationVM.h"
 #include "Graph/AnimNext_LODPose.h"
 #include "AnimNextStats.h"
 #include "Logging/StructuredLog.h"
@@ -164,31 +162,6 @@ void FAnimNextScheduleGraphTask::RunGraph(const UE::AnimNext::FScheduleContext&
 	// This reduces churn internally by avoiding a chunk to be repeatedly allocated and freed as we push/pop marks
 	MemStack.Alloc(size_t(FPageAllocator::SmallPageSize) + 1, 16);
 
-	UE::AnimNext::UpdateGraph(GraphCache.GraphInstanceData, InContext.GetDeltaTime());
-
-	{
-		const FEvaluationProgram EvaluationProgram = UE::AnimNext::EvaluateGraph(GraphCache.GraphInstanceData);
-
-		FEvaluationVM EvaluationVM(EEvaluationFlags::All, RefPose, *GraphLODLevel);
-		bool bHasValidOutput = false;
-
-		if (!EvaluationProgram.IsEmpty())
-		{
-			EvaluationProgram.Execute(EvaluationVM);
-
-			TUniquePtr<FKeyframeState> EvaluatedKeyframe;
-			if (EvaluationVM.PopValue(KEYFRAME_STACK_NAME, EvaluatedKeyframe))
-			{
-				OutputPose->LODPose.CopyFrom(EvaluatedKeyframe->Pose);
-				bHasValidOutput = true;
-			}
-		}
-
-		if (!bHasValidOutput)
-		{
-			// We need to output a valid pose, generate one
-			FKeyframeState ReferenceKeyframe = EvaluationVM.MakeReferenceKeyframe(false);
-			OutputPose->LODPose.CopyFrom(ReferenceKeyframe.Pose);
-		}
-	}
+	IAnimNextModule::Get().UpdateGraph(GraphCache.GraphInstanceData, InContext.GetDeltaTime());
+	IAnimNextModule::Get().EvaluateGraph(GraphCache.GraphInstanceData, RefPose, *GraphLODLevel, OutputPose->LODPose);
 }
diff --git a/Engine/Plugins/Experimental/AnimNext/Source/AnimNext/Private/TraitInterfaces/IContinuousBlend.cpp b/Engine/Plugins/Experimental/AnimNext/Source/AnimNext/Private/TraitInterfaces/IContinuousBlend.cpp
deleted file mode 100644
diff --git a/Engine/Plugins/Experimental/AnimNext/Source/AnimNext/Private/TraitInterfaces/IDiscreteBlend.cpp b/Engine/Plugins/Experimental/AnimNext/Source/AnimNext/Private/TraitInterfaces/IDiscreteBlend.cpp
deleted file mode 100644
diff --git a/Engine/Plugins/Experimental/AnimNext/Source/AnimNext/Private/TraitInterfaces/IEvaluate.cpp b/Engine/Plugins/Experimental/AnimNext/Source/AnimNext/Private/TraitInterfaces/IEvaluate.cpp
deleted file mode 100644
diff --git a/Engine/Plugins/Experimental/AnimNext/Source/AnimNext/Private/TraitInterfaces/IGroupSynchronization.cpp b/Engine/Plugins/Experimental/AnimNext/Source/AnimNext/Private/TraitInterfaces/IGroupSynchronization.cpp
deleted file mode 100644
diff --git a/Engine/Plugins/Experimental/AnimNext/Source/AnimNext/Private/TraitInterfaces/ISmoothBlend.cpp b/Engine/Plugins/Experimental/AnimNext/Source/AnimNext/Private/TraitInterfaces/ISmoothBlend.cpp
deleted file mode 100644
diff --git a/Engine/Plugins/Experimental/AnimNext/Source/AnimNext/Private/TraitInterfaces/ITimeline.cpp b/Engine/Plugins/Experimental/AnimNext/Source/AnimNext/Private/TraitInterfaces/ITimeline.cpp
deleted file mode 100644
diff --git a/Engine/Plugins/Experimental/AnimNext/Source/AnimNext/Private/TraitInterfaces/IUpdate.cpp b/Engine/Plugins/Experimental/AnimNext/Source/AnimNext/Private/TraitInterfaces/IUpdate.cpp
deleted file mode 100644
diff --git a/Engine/Plugins/Experimental/AnimNext/Source/AnimNext/Private/Traits/BlendByBool.cpp b/Engine/Plugins/Experimental/AnimNext/Source/AnimNext/Private/Traits/BlendByBool.cpp
deleted file mode 100644
diff --git a/Engine/Plugins/Experimental/AnimNext/Source/AnimNext/Private/Traits/BlendByBool.h b/Engine/Plugins/Experimental/AnimNext/Source/AnimNext/Private/Traits/BlendByBool.h
deleted file mode 100644
diff --git a/Engine/Plugins/Experimental/AnimNext/Source/AnimNext/Private/Traits/BlendInertializer.cpp b/Engine/Plugins/Experimental/AnimNext/Source/AnimNext/Private/Traits/BlendInertializer.cpp
deleted file mode 100644
diff --git a/Engine/Plugins/Experimental/AnimNext/Source/AnimNext/Private/Traits/BlendInertializer.h b/Engine/Plugins/Experimental/AnimNext/Source/AnimNext/Private/Traits/BlendInertializer.h
deleted file mode 100644
diff --git a/Engine/Plugins/Experimental/AnimNext/Source/AnimNext/Private/Traits/BlendSmoother.cpp b/Engine/Plugins/Experimental/AnimNext/Source/AnimNext/Private/Traits/BlendSmoother.cpp
deleted file mode 100644
diff --git a/Engine/Plugins/Experimental/AnimNext/Source/AnimNext/Private/Traits/BlendSmoother.h b/Engine/Plugins/Experimental/AnimNext/Source/AnimNext/Private/Traits/BlendSmoother.h
deleted file mode 100644
diff --git a/Engine/Plugins/Experimental/AnimNext/Source/AnimNext/Private/Traits/BlendSmootherPerBone.cpp b/Engine/Plugins/Experimental/AnimNext/Source/AnimNext/Private/Traits/BlendSmootherPerBone.cpp
deleted file mode 100644
diff --git a/Engine/Plugins/Experimental/AnimNext/Source/AnimNext/Private/Traits/BlendSmootherPerBone.h b/Engine/Plugins/Experimental/AnimNext/Source/AnimNext/Private/Traits/BlendSmootherPerBone.h
deleted file mode 100644
diff --git a/Engine/Plugins/Experimental/AnimNext/Source/AnimNext/Private/Traits/BlendTwoWay.cpp b/Engine/Plugins/Experimental/AnimNext/Source/AnimNext/Private/Traits/BlendTwoWay.cpp
deleted file mode 100644
diff --git a/Engine/Plugins/Experimental/AnimNext/Source/AnimNext/Private/Traits/BlendTwoWay.h b/Engine/Plugins/Experimental/AnimNext/Source/AnimNext/Private/Traits/BlendTwoWay.h
deleted file mode 100644
diff --git a/Engine/Plugins/Experimental/AnimNext/Source/AnimNext/Private/Traits/ReferencePoseTrait.cpp b/Engine/Plugins/Experimental/AnimNext/Source/AnimNext/Private/Traits/ReferencePoseTrait.cpp
deleted file mode 100644
diff --git a/Engine/Plugins/Experimental/AnimNext/Source/AnimNext/Private/Traits/ReferencePoseTrait.h b/Engine/Plugins/Experimental/AnimNext/Source/AnimNext/Private/Traits/ReferencePoseTrait.h
deleted file mode 100644
diff --git a/Engine/Plugins/Experimental/AnimNext/Source/AnimNext/Private/Traits/SequencePlayer.cpp b/Engine/Plugins/Experimental/AnimNext/Source/AnimNext/Private/Traits/SequencePlayer.cpp
deleted file mode 100644
diff --git a/Engine/Plugins/Experimental/AnimNext/Source/AnimNext/Private/Traits/SequencePlayer.h b/Engine/Plugins/Experimental/AnimNext/Source/AnimNext/Private/Traits/SequencePlayer.h
deleted file mode 100644
diff --git a/Engine/Plugins/Experimental/AnimNext/Source/AnimNext/Private/Traits/SubGraphHost.cpp b/Engine/Plugins/Experimental/AnimNext/Source/AnimNext/Private/Traits/SubGraphHost.cpp
deleted file mode 100644
diff --git a/Engine/Plugins/Experimental/AnimNext/Source/AnimNext/Private/Traits/SubGraphHost.h b/Engine/Plugins/Experimental/AnimNext/Source/AnimNext/Private/Traits/SubGraphHost.h
deleted file mode 100644
diff --git a/Engine/Plugins/Experimental/AnimNext/Source/AnimNext/Private/Traits/SynchronizeUsingGroups.cpp b/Engine/Plugins/Experimental/AnimNext/Source/AnimNext/Private/Traits/SynchronizeUsingGroups.cpp
deleted file mode 100644
diff --git a/Engine/Plugins/Experimental/AnimNext/Source/AnimNext/Private/Traits/SynchronizeUsingGroups.h b/Engine/Plugins/Experimental/AnimNext/Source/AnimNext/Private/Traits/SynchronizeUsingGroups.h
deleted file mode 100644
diff --git a/Engine/Plugins/Experimental/AnimNext/Source/AnimNext/Public/AnimNextStats.h b/Engine/Plugins/Experimental/AnimNext/Source/AnimNext/Public/AnimNextStats.h
--- a/Engine/Plugins/Experimental/AnimNext/Source/AnimNext/Public/AnimNextStats.h
+++ b/Engine/Plugins/Experimental/AnimNext/Source/AnimNext/Public/AnimNextStats.h
@@ -17,8 +17,6 @@ DECLARE_CYCLE_STAT_EXTERN(TEXT("AnimNext Task: Scope Exit"), STAT_AnimNext_Task_
 
 DECLARE_CYCLE_STAT_EXTERN(TEXT("AnimNext: Allocate Graph Instance"), STAT_AnimNext_Graph_AllocateInstance, STATGROUP_AnimNext, ANIMNEXT_API);
 DECLARE_CYCLE_STAT_EXTERN(TEXT("AnimNext: Graph RigVM"), STAT_AnimNext_Graph_RigVM, STATGROUP_AnimNext, ANIMNEXT_API);
-DECLARE_CYCLE_STAT_EXTERN(TEXT("AnimNext: Update Graph"), STAT_AnimNext_UpdateGraph, STATGROUP_AnimNext, ANIMNEXT_API);
-DECLARE_CYCLE_STAT_EXTERN(TEXT("AnimNext: Evaluate Graph"), STAT_AnimNext_EvaluateGraph, STATGROUP_AnimNext, ANIMNEXT_API);
 DECLARE_CYCLE_STAT_EXTERN(TEXT("AnimNext: Execute Evaluation Program"), STAT_AnimNext_EvaluationProgram_Execute, STATGROUP_AnimNext, ANIMNEXT_API);
 DECLARE_CYCLE_STAT_EXTERN(TEXT("AnimNext: Copy Transforms (SoA)"), STAT_AnimNext_CopyTransforms_SoA, STATGROUP_AnimNext, ANIMNEXT_API);
 DECLARE_CYCLE_STAT_EXTERN(TEXT("AnimNext: Normalize Rotations (SoA)"), STAT_AnimNext_NormalizeRotations_SoA, STATGROUP_AnimNext, ANIMNEXT_API);
diff --git a/Engine/Plugins/Experimental/AnimNext/Source/AnimNext/Public/EvaluationVM/EvaluationVM.h b/Engine/Plugins/Experimental/AnimNext/Source/AnimNext/Public/EvaluationVM/EvaluationVM.h
--- a/Engine/Plugins/Experimental/AnimNext/Source/AnimNext/Public/EvaluationVM/EvaluationVM.h
+++ b/Engine/Plugins/Experimental/AnimNext/Source/AnimNext/Public/EvaluationVM/EvaluationVM.h
@@ -120,7 +120,7 @@ namespace UE::AnimNext
 	 * 
 	 * Represents a named and typed VM stack.
 	 */
-	struct FEvaluationVMStack final
+	struct ANIMNEXT_API FEvaluationVMStack final
 	{
 		// Name of the stack
 		FName Name;
@@ -312,6 +312,6 @@ namespace UE::AnimNext
 	// Various commonly used VM stacks
 
 	// A stack of FKeyframeState instances used when sampling sequences and blending their results
-	extern const FEvaluationVMStackName KEYFRAME_STACK_NAME;
+	extern ANIMNEXT_API const FEvaluationVMStackName KEYFRAME_STACK_NAME;
 	ANIM_NEXT_ENABLE_EVALUATION_STACK_USAGE(TUniquePtr<FKeyframeState>)
 }
diff --git a/Engine/Plugins/Experimental/AnimNext/Source/AnimNext/Public/Graph/AnimGraphParamStackScope.h b/Engine/Plugins/Experimental/AnimNext/Source/AnimNext/Public/Graph/AnimGraphParamStackScope.h
new file mode 100644
--- /dev/null
+++ b/Engine/Plugins/Experimental/AnimNext/Source/AnimNext/Public/Graph/AnimGraphParamStackScope.h
@@ -0,0 +1,37 @@
+// Copyright Epic Games, Inc. All Rights Reserved.
+
+#pragma once
+
+#include "Templates/SharedPointer.h"
+#include "UObject/WeakObjectPtr.h"
+
+struct FAnimationBaseContext;
+struct FAnimNodeExposedValueHandler_AnimNextParameters;
+
+namespace UE::AnimNext
+{
+	struct FParamStack;
+}
+
+namespace UE::AnimNext
+{
+
+// Utility for managing parameter stack access in the anim graph
+// Pulls from stacks that are marked as 'pending' against specific objects (in this case the component).
+// Instantiating this will attach the param stack to the current thread within its scope
+struct ANIMNEXT_API FAnimGraphParamStackScope
+{
+	UE_NODISCARD_CTOR explicit FAnimGraphParamStackScope(const FAnimationBaseContext& InContext);
+	~FAnimGraphParamStackScope();
+
+private:
+	friend struct ::FAnimNodeExposedValueHandler_AnimNextParameters;
+
+	TWeakObjectPtr<const UObject> ComponentObject;
+	TSharedPtr<FParamStack> OwnedParamStack;
+
+	// Whether we attached to a pending object
+	bool bAttachedPending = false;
+};
+
+}
\ No newline at end of file
diff --git a/Engine/Plugins/Experimental/AnimNext/Source/AnimNext/Public/Graph/AnimNextGraphInstance.h b/Engine/Plugins/Experimental/AnimNext/Source/AnimNext/Public/Graph/AnimNextGraphInstance.h
--- a/Engine/Plugins/Experimental/AnimNext/Source/AnimNext/Public/Graph/AnimNextGraphInstance.h
+++ b/Engine/Plugins/Experimental/AnimNext/Source/AnimNext/Public/Graph/AnimNextGraphInstance.h
@@ -28,7 +28,7 @@ using GraphInstanceComponentMapType = TMap<FName, TSharedPtr<UE::AnimNext::FGrap
 // This struct uses UE reflection because we wish for the GC to keep the graph
 // alive while we own a reference to it. It is not intended to be serialized on disk with a live instance.
 USTRUCT()
-struct FAnimNextGraphInstance
+struct ANIMNEXT_API FAnimNextGraphInstance
 {
 	GENERATED_BODY()
 
diff --git a/Engine/Plugins/Experimental/AnimNext/Source/AnimNext/Public/IAnimNextModule.h b/Engine/Plugins/Experimental/AnimNext/Source/AnimNext/Public/IAnimNextModule.h
new file mode 100644
--- /dev/null
+++ b/Engine/Plugins/Experimental/AnimNext/Source/AnimNext/Public/IAnimNextModule.h
@@ -0,0 +1,39 @@
+// Copyright Epic Games, Inc. All Rights Reserved.
+
+#pragma once
+
+#include "CoreMinimal.h"
+#include "Modules/ModuleInterface.h"
+#include "LODPose.h"
+#include "ReferencePose.h"
+
+struct FAnimNextGraphInstancePtr;
+
+namespace UE::AnimNext
+{
+	class ANIMNEXT_API IAnimNextAnimGraph
+	{
+	public:
+		virtual ~IAnimNextAnimGraph() = default;
+
+		virtual void UpdateGraph(FAnimNextGraphInstancePtr& GraphInstance, float DeltaTime) const = 0;
+		virtual void EvaluateGraph(FAnimNextGraphInstancePtr& GraphInstance, const UE::AnimNext::FReferencePose& RefPose, int32 GraphLODLevel, FLODPoseHeap& OutputPose) const = 0;
+	};
+
+	class ANIMNEXT_API IAnimNextModule : public IModuleInterface
+	{
+	public:
+		// WARNING
+		// These functions are here because the schedule currently cannot have tasks defined in
+		// external plugins. And as such, we must be able to update and evaluate graphs from
+		// AnimNextCore even though they are implementation details of the AnimNextAnimGraph.
+		// We thus introduce this interface.
+		// When AnimNextAnimGraph's module will load, it will register callbacks that
+		// AnimNextCore will call when we need to update/evaluate graphs.
+		static IAnimNextModule& Get();
+		virtual void RegisterAnimNextAnimGraph(const IAnimNextAnimGraph& InAnimGraphImpl) = 0;
+		virtual void UnregisterAnimNextAnimGraph() = 0;
+		virtual void UpdateGraph(FAnimNextGraphInstancePtr& GraphInstance, float DeltaTime) = 0;
+		virtual void EvaluateGraph(FAnimNextGraphInstancePtr& GraphInstance, const UE::AnimNext::FReferencePose& RefPose, int32 GraphLODLevel, FLODPoseHeap& OutputPose) const = 0;
+	};
+}
diff --git a/Engine/Plugins/Experimental/AnimNext/Source/AnimNext/Public/TraitInterfaces/IContinuousBlend.h b/Engine/Plugins/Experimental/AnimNext/Source/AnimNext/Public/TraitInterfaces/IContinuousBlend.h
deleted file mode 100644
diff --git a/Engine/Plugins/Experimental/AnimNext/Source/AnimNext/Public/TraitInterfaces/IDiscreteBlend.h b/Engine/Plugins/Experimental/AnimNext/Source/AnimNext/Public/TraitInterfaces/IDiscreteBlend.h
deleted file mode 100644
diff --git a/Engine/Plugins/Experimental/AnimNext/Source/AnimNext/Public/TraitInterfaces/IEvaluate.h b/Engine/Plugins/Experimental/AnimNext/Source/AnimNext/Public/TraitInterfaces/IEvaluate.h
deleted file mode 100644
diff --git a/Engine/Plugins/Experimental/AnimNext/Source/AnimNext/Public/TraitInterfaces/IGroupSynchronization.h b/Engine/Plugins/Experimental/AnimNext/Source/AnimNext/Public/TraitInterfaces/IGroupSynchronization.h
deleted file mode 100644
diff --git a/Engine/Plugins/Experimental/AnimNext/Source/AnimNext/Public/TraitInterfaces/ISmoothBlend.h b/Engine/Plugins/Experimental/AnimNext/Source/AnimNext/Public/TraitInterfaces/ISmoothBlend.h
deleted file mode 100644
diff --git a/Engine/Plugins/Experimental/AnimNext/Source/AnimNext/Public/TraitInterfaces/ITimeline.h b/Engine/Plugins/Experimental/AnimNext/Source/AnimNext/Public/TraitInterfaces/ITimeline.h
deleted file mode 100644
diff --git a/Engine/Plugins/Experimental/AnimNext/Source/AnimNext/Public/TraitInterfaces/IUpdate.h b/Engine/Plugins/Experimental/AnimNext/Source/AnimNext/Public/TraitInterfaces/IUpdate.h
deleted file mode 100644
diff --git a/Engine/Plugins/Experimental/AnimNext/Source/AnimNextTestSuite/Private/AnimNextRuntimeTest.h b/Engine/Plugins/Experimental/AnimNext/Source/AnimNextTestSuite/Private/AnimNextRuntimeTest.h
deleted file mode 100644
diff --git a/Engine/Plugins/Experimental/AnimNext/Source/AnimNextTestSuite/Private/AnimNextTest.h b/Engine/Plugins/Experimental/AnimNext/Source/AnimNextTestSuite/Private/AnimNextTest.h
deleted file mode 100644
diff --git a/Engine/Plugins/Experimental/AnimNext/Source/AnimNextTestSuite/Private/AnimNextTraitGraphTest.cpp b/Engine/Plugins/Experimental/AnimNext/Source/AnimNextTestSuite/Private/AnimNextTraitGraphTest.cpp
deleted file mode 100644
diff --git a/Engine/Plugins/Experimental/AnimNext/Source/AnimNextTestSuite/Private/AnimNextTraitGraphTest.h b/Engine/Plugins/Experimental/AnimNext/Source/AnimNextTestSuite/Private/AnimNextTraitGraphTest.h
deleted file mode 100644
diff --git a/Engine/Plugins/Experimental/AnimNext/Source/AnimNextTestSuite/Private/AnimNextTraitInterfacesTest.cpp b/Engine/Plugins/Experimental/AnimNext/Source/AnimNextTestSuite/Private/AnimNextTraitInterfacesTest.cpp
deleted file mode 100644
diff --git a/Engine/Plugins/Experimental/AnimNext/Source/AnimNextTestSuite/Private/AnimNextTraitInterfacesTest.h b/Engine/Plugins/Experimental/AnimNext/Source/AnimNextTestSuite/Private/AnimNextTraitInterfacesTest.h
deleted file mode 100644
diff --git a/Engine/Plugins/Experimental/AnimNext/Source/AnimNextTestSuite/Public/AnimNextRuntimeTest.h b/Engine/Plugins/Experimental/AnimNext/Source/AnimNextTestSuite/Public/AnimNextRuntimeTest.h
new file mode 100644
--- /dev/null
+++ b/Engine/Plugins/Experimental/AnimNext/Source/AnimNextTestSuite/Public/AnimNextRuntimeTest.h
@@ -0,0 +1,91 @@
+// Copyright Epic Games, Inc. All Rights Reserved.
+
+#pragma once
+
+#include "CoreMinimal.h"
+
+#if WITH_DEV_AUTOMATION_TESTS
+
+#include "UObject/Class.h"
+#include "UObject/UnrealType.h"
+
+#include "TraitCore/TraitUID.h"
+#include "TraitCore/NodeTemplate.h"
+#include "TraitCore/NodeTemplateRegistry.h"
+
+class UAnimNextGraph;
+
+namespace UE::AnimNext
+{
+	struct FNodeHandle;
+
+	// Converts a property value into its string representation using UE reflection
+	template<class TraitSharedDataType, typename PropertyType>
+	static FString ToString(const FString& PropertyName, PropertyType PropertyValue)
+	{
+		const UScriptStruct* SharedDataStruct = TraitSharedDataType::StaticStruct();
+		if (const FProperty* Property = SharedDataStruct->FindPropertyByName(*PropertyName))
+		{
+			void* PropertyDefaults = Property->AllocateAndInitializeValue();
+
+			FString Result;
+
+			if constexpr (std::is_pointer_v<PropertyType>)
+			{
+				// C-style array properties aren't handled by ExportText, we need to handle it manually
+
+				const bool bIsCArray = Property->ArrayDim > 1;
+				if (bIsCArray)
+				{
+					Result += TEXT("(");
+				}
+
+				for (int32 Index = 0; Index < Property->ArrayDim; ++Index)
+				{
+					Property->ExportText_Direct(Result, PropertyValue + Index, (PropertyType)PropertyDefaults + Index, nullptr, PPF_None, nullptr);
+
+					if (Index + 1 < Property->ArrayDim)
+					{
+						Result += TEXT(",");
+					}
+				}
+
+				if (bIsCArray)
+				{
+					Result += TEXT(")");
+				}
+			}
+			else
+			{
+				Property->ExportText_Direct(Result, &PropertyValue, PropertyDefaults, nullptr, PPF_None);
+			}
+
+			Property->DestroyAndFreeValue(PropertyDefaults);
+
+			return Result;
+		}
+
+		return FString();
+	}
+
+	// Creates a temporary empty node template registry instance and swaps it for the current one
+	struct ANIMNEXTTESTSUITE_API FScopedClearNodeTemplateRegistry final
+	{
+		FScopedClearNodeTemplateRegistry(const FScopedClearNodeTemplateRegistry&) = delete;
+		FScopedClearNodeTemplateRegistry& operator=(const FScopedClearNodeTemplateRegistry&) = delete;
+
+		FScopedClearNodeTemplateRegistry();
+		~FScopedClearNodeTemplateRegistry();
+
+		FNodeTemplateRegistry TmpRegistry;
+	};
+
+	struct ANIMNEXTTESTSUITE_API FTestUtils final
+	{
+		// Loads the graph data from the provided archive buffer and returns true on success, false otherwise
+		// On success, we resolve every node handle provided as argument
+		static bool LoadFromArchiveBuffer(UAnimNextGraph& Graph, TArray<FNodeHandle>& NodeHandles, const TArray<uint8>& SharedDataArchiveBuffer);
+	};
+}
+
+#endif
diff --git a/Engine/Plugins/Experimental/AnimNext/Source/AnimNextTestSuite/Public/AnimNextTest.h b/Engine/Plugins/Experimental/AnimNext/Source/AnimNextTestSuite/Public/AnimNextTest.h
new file mode 100644
--- /dev/null
+++ b/Engine/Plugins/Experimental/AnimNext/Source/AnimNextTestSuite/Public/AnimNextTest.h
@@ -0,0 +1,26 @@
+// Copyright Epic Games, Inc. All Rights Reserved.
+
+#pragma once
+
+#include "AnimNextTest.generated.h"
+
+//// --- Raw type ---
+USTRUCT()
+struct FAnimNextTestData
+{
+	GENERATED_BODY()
+
+	float A = 0.f;
+	float B = 0.f;
+};
+
+namespace UE::AnimNext::Tests
+{
+
+struct ANIMNEXTTESTSUITE_API FUtils final
+{
+	// Clean up after tests. Clears transaction buffer, collects garbage
+	static void CleanupAfterTests();
+};
+
+}
diff --git a/Engine/Plugins/Experimental/AnimNext/Source/AnimNextUncookedOnly/Internal/Graph/AnimGraph/AnimBlueprintExtension_AnimNextParameters.h b/Engine/Plugins/Experimental/AnimNext/Source/AnimNextUncookedOnly/Internal/Graph/AnimGraph/AnimBlueprintExtension_AnimNextParameters.h
deleted file mode 100644
diff --git a/Engine/Plugins/Experimental/AnimNext/Source/AnimNextUncookedOnly/Internal/Graph/AnimGraph/AnimGraphNode_AnimNextGraph.h b/Engine/Plugins/Experimental/AnimNext/Source/AnimNextUncookedOnly/Internal/Graph/AnimGraph/AnimGraphNode_AnimNextGraph.h
deleted file mode 100644
diff --git a/Engine/Plugins/Experimental/AnimNext/Source/AnimNextUncookedOnly/Private/Graph/AnimGraph/AnimGraphNode_AnimNextGraph.cpp b/Engine/Plugins/Experimental/AnimNext/Source/AnimNextUncookedOnly/Private/Graph/AnimGraph/AnimGraphNode_AnimNextGraph.cpp
deleted file mode 100644
diff --git a/Engine/Plugins/Experimental/AnimNext/Source/AnimNextUncookedOnly/Public/Graph/AnimGraph/AnimBlueprintExtension_AnimNextParameters.h b/Engine/Plugins/Experimental/AnimNext/Source/AnimNextUncookedOnly/Public/Graph/AnimGraph/AnimBlueprintExtension_AnimNextParameters.h
new file mode 100644
--- /dev/null
+++ b/Engine/Plugins/Experimental/AnimNext/Source/AnimNextUncookedOnly/Public/Graph/AnimGraph/AnimBlueprintExtension_AnimNextParameters.h
@@ -0,0 +1,91 @@
+// Copyright Epic Games, Inc. All Rights Reserved.
+
+#pragma once
+
+#include "CoreMinimal.h"
+#include "AnimBlueprintExtension.h"
+#include "Graph/AnimGraph/AnimSubsystem_AnimNextParameters.h"
+#include "PropertyAccess.h"
+#include "IPropertyAccessCompiler.h"
+#include "AnimBlueprintExtension_AnimNextParameters.generated.h"
+
+struct FAnimNextAnimGraphNodeParameterBinding;
+
+UCLASS(MinimalAPI)
+class UAnimBlueprintExtension_AnimNextParameters : public UAnimBlueprintExtension
+{
+	GENERATED_BODY()
+
+private:
+	friend class UAnimGraphNodeBinding_AnimNextParameters;
+
+	void ProcessNodePins(UAnimGraphNode_Base* InNode, IAnimBlueprintCompilationContext& InCompilationContext, IAnimBlueprintGeneratedClassCompiledData& OutCompiledData);
+
+	// UAnimBlueprintExtension interface
+	virtual void HandleStartCompilingClass(const UClass* InClass, IAnimBlueprintCompilationBracketContext& InCompilationContext, IAnimBlueprintGeneratedClassCompiledData& OutCompiledData) override;
+
+	void PatchEvaluationHandlers(const UClass* InClass, IAnimBlueprintCompilationBracketContext& InCompilationContext, IAnimBlueprintGeneratedClassCompiledData& OutCompiledData);
+
+	void RedirectPropertyPaths(IAnimBlueprintCompilationContext& InCompilationContext, UAnimGraphNode_Base* InNode);
+
+	struct FSinglePropertyRecord
+	{
+		/** Parameter name to use for binding */
+		FName ParameterName = NAME_None;
+
+		/** The destination property we are bound to (on an animation node) */
+		FProperty* Property = nullptr;
+
+		/** Index into array if this is an array property element */
+		int32 ArrayIndex = INDEX_NONE;
+
+		/** The property path relative to the class */
+		TArray<FString> PropertyPath;
+
+		/** The handle of the copy in the property access compiler */
+		FPropertyAccessHandle LibraryHandle;
+
+		/** The handle of the copy in the property access library */
+		FCompiledPropertyAccessHandle LibraryCompiledHandle;
+
+		/** The type of the access in the property access library */
+		EPropertyAccessCopyType LibraryCompiledAccessType = EPropertyAccessCopyType::None;
+	};
+
+	/** All the data needed for a single anim graph node's bindings */
+	struct FBindingRecord
+	{
+		// The Node this record came from
+		UAnimGraphNode_Base* AnimGraphNode = nullptr;
+
+		// The node variable that the handler is in
+		FStructProperty* NodeVariableProperty = nullptr;
+
+		// Whether or not our serviced properties are actually on the anim node 
+		bool bServicesNodeProperties = false;
+
+		// Whether or not our serviced properties are actually on the instance instead of the node
+		bool bServicesInstanceProperties = false;
+
+		// Set of properties serviced by this handler (Map from property name to the record for that property)
+		TMap<FName, FSinglePropertyRecord> ServicedProperties;
+
+		bool IsValid() const
+		{
+			return NodeVariableProperty != nullptr;
+		}
+
+		void RegisterPropertyBinding(FProperty* InProperty, const FAnimNextAnimGraphNodeParameterBinding& InBinding);
+	};
+
+	UPROPERTY()
+	FAnimSubsystem_AnimNextParameters Subsystem;
+
+	// All the binding records we process during compilation
+	TArray<FBindingRecord> BindingRecords;
+	TMap<UAnimGraphNode_Base*, int32> BindingRecordIndexMap;
+
+	// Delegate handle for registering against library pre/post-compilation
+	FDelegateHandle PreLibraryCompiledDelegateHandle;
+	FDelegateHandle PostLibraryCompiledDelegateHandle;
+};
\ No newline at end of file
diff --git a/Engine/Plugins/Experimental/AnimNextAnimGraph/AnimNextAnimGraph.uplugin b/Engine/Plugins/Experimental/AnimNextAnimGraph/AnimNextAnimGraph.uplugin
new file mode 100644
--- /dev/null
+++ b/Engine/Plugins/Experimental/AnimNextAnimGraph/AnimNextAnimGraph.uplugin
@@ -0,0 +1,43 @@
+{
+	"FileVersion" : 3,
+	"Version" : 1,
+	"VersionName" : "0.1",
+	"FriendlyName" : "Anim Next Anim Graph",
+	"Description" : "Framework for defining animation graphs.",
+	"Category" : "Animation",
+	"CreatedBy" : "Epic Games, Inc.",
+	"CreatedByURL" : "https://epicgames.com",
+	"DocsURL" : "",
+	"MarketplaceURL" : "",
+	"SupportURL" : "",
+	"EnabledByDefault" : false,
+	"CanContainContent" : false,
+	"IsBetaVersion": true,
+	"IsExperimentalVersion": true,
+	"Installed" : false,
+	"Plugins":
+	[
+		{
+			"Name": "AnimNext",
+			"Enabled": true
+		}
+	],
+	"Modules" :
+	[
+		{
+			"Name" : "AnimNextAnimGraph",
+			"Type" : "Runtime",
+			"LoadingPhase" : "PreDefault"
+		},
+		{
+			"Name" : "AnimNextAnimGraphUncookedOnly",
+			"Type" : "UncookedOnly",
+			"LoadingPhase" : "Default"
+		},
+		{
+			"Name" : "AnimNextAnimGraphTestSuite",
+			"Type" : "UncookedOnly",
+			"LoadingPhase" : "Default"
+		}
+	]
+}
diff --git a/Engine/Plugins/Experimental/AnimNextAnimGraph/Config/BaseAnimNextAnimGraph.ini b/Engine/Plugins/Experimental/AnimNextAnimGraph/Config/BaseAnimNextAnimGraph.ini
new file mode 100644
--- /dev/null
+++ b/Engine/Plugins/Experimental/AnimNextAnimGraph/Config/BaseAnimNextAnimGraph.ini
@@ -0,0 +1,25 @@
+[CoreRedirects]
+; Traits rename
++StructRedirects=(OldName="/Script/AnimNext.AnimNextBlendByBoolDecoratorSharedData", NewName="/Script/AnimNextAnimGraph.AnimNextBlendByBoolTraitSharedData")
++StructRedirects=(OldName="/Script/AnimNext.AnimNextBlendInertializerDecoratorSharedData", NewName="/Script/AnimNextAnimGraph.AnimNextBlendInertializerTraitSharedData")
++StructRedirects=(OldName="/Script/AnimNext.AnimNextBlendSmootherDecoratorSharedData", NewName="/Script/AnimNextAnimGraph.AnimNextBlendSmootherTraitSharedData")
++StructRedirects=(OldName="/Script/AnimNext.AnimNextBlendSmootherPerBoneDecoratorSharedData", NewName="/Script/AnimNextAnimGraph.AnimNextBlendSmootherPerBoneTraitSharedData")
++StructRedirects=(OldName="/Script/AnimNext.AnimNextBlendTwoWayDecoratorSharedData", NewName="/Script/AnimNextAnimGraph.AnimNextBlendTwoWayTraitSharedData")
++StructRedirects=(OldName="/Script/AnimNext.AnimNextReferencePoseDecoratorSharedData", NewName="/Script/AnimNextAnimGraph.AnimNextReferencePoseTraitSharedData")
++StructRedirects=(OldName="/Script/AnimNext.AnimNextSequencePlayerDecoratorSharedData", NewName="/Script/AnimNextAnimGraph.AnimNextSequencePlayerTraitSharedData")
++StructRedirects=(OldName="/Script/AnimNext.AnimNextSubGraphHostDecoratorSharedData", NewName="/Script/AnimNextAnimGraph.AnimNextSubGraphHostTraitSharedData")
++StructRedirects=(OldName="/Script/AnimNext.AnimNextSynchronizeUsingGroupsDecoratorSharedData", NewName="/Script/AnimNextAnimGraph.AnimNextSynchronizeUsingGroupsTraitSharedData")
+
+; Plugin refactor
++ClassRedirects=(OldName="/Script/AnimNext.AnimGraphNode_AnimNextGraph", NewName="/Script/AnimNextAnimGraph.AnimGraphNode_AnimNextGraph")
+
++StructRedirects=(OldName="/Script/AnimNext.AnimNextBlendByBoolTraitSharedData", NewName="/Script/AnimNextAnimGraph.AnimNextBlendByBoolTraitSharedData")
++StructRedirects=(OldName="/Script/AnimNext.AnimNextBlendInertializerTraitSharedData", NewName="/Script/AnimNextAnimGraph.AnimNextBlendInertializerTraitSharedData")
++StructRedirects=(OldName="/Script/AnimNext.AnimNextBlendSmootherTraitSharedData", NewName="/Script/AnimNextAnimGraph.AnimNextBlendSmootherTraitSharedData")
++StructRedirects=(OldName="/Script/AnimNext.AnimNextBlendSmootherPerBoneTraitSharedData", NewName="/Script/AnimNextAnimGraph.AnimNextBlendSmootherPerBoneTraitSharedData")
++StructRedirects=(OldName="/Script/AnimNext.AnimNextBlendTwoWayTraitSharedData", NewName="/Script/AnimNextAnimGraph.AnimNextBlendTwoWayTraitSharedData")
++StructRedirects=(OldName="/Script/AnimNext.AnimNextReferencePoseTraitSharedData", NewName="/Script/AnimNextAnimGraph.AnimNextReferencePoseTraitSharedData")
++StructRedirects=(OldName="/Script/AnimNext.AnimNextSequencePlayerTraitSharedData", NewName="/Script/AnimNextAnimGraph.AnimNextSequencePlayerTraitSharedData")
++StructRedirects=(OldName="/Script/AnimNext.AnimNextSubGraphHostTraitSharedData", NewName="/Script/AnimNextAnimGraph.AnimNextSubGraphHostTraitSharedData")
++StructRedirects=(OldName="/Script/AnimNext.AnimNextSynchronizeUsingGroupsTraitSharedData", NewName="/Script/AnimNextAnimGraph.AnimNextSynchronizeUsingGroupsTraitSharedData")
++StructRedirects=(OldName="/Script/AnimNext.AnimNode_AnimNextGraph", NewName="/Script/AnimNextAnimGraph.AnimNode_AnimNextGraph")
diff --git a/Engine/Plugins/Experimental/AnimNextAnimGraph/Source/AnimNextAnimGraph/AnimNextAnimGraph.Build.cs b/Engine/Plugins/Experimental/AnimNextAnimGraph/Source/AnimNextAnimGraph/AnimNextAnimGraph.Build.cs
new file mode 100644
--- /dev/null
+++ b/Engine/Plugins/Experimental/AnimNextAnimGraph/Source/AnimNextAnimGraph/AnimNextAnimGraph.Build.cs
@@ -0,0 +1,21 @@
+// Copyright Epic Games, Inc. All Rights Reserved.
+
+namespace UnrealBuildTool.Rules
+{
+	public class AnimNextAnimGraph : ModuleRules
+	{
+		public AnimNextAnimGraph(ReadOnlyTargetRules Target) : base(Target)
+		{
+			PublicDependencyModuleNames.AddRange(
+				new string[]
+				{
+					"Core",
+					"CoreUObject",
+					"RigVM",
+					"Engine",
+					"AnimNext",
+				}
+			);
+		}
+	}
+}
\ No newline at end of file
diff --git a/Engine/Plugins/Experimental/AnimNextAnimGraph/Source/AnimNextAnimGraph/Internal/Graph/AnimGraph/AnimNode_AnimNextGraph.h b/Engine/Plugins/Experimental/AnimNextAnimGraph/Source/AnimNextAnimGraph/Internal/Graph/AnimGraph/AnimNode_AnimNextGraph.h
new file mode 100644
--- /dev/null
+++ b/Engine/Plugins/Experimental/AnimNextAnimGraph/Source/AnimNextAnimGraph/Internal/Graph/AnimGraph/AnimNode_AnimNextGraph.h
@@ -0,0 +1,78 @@
+// Copyright Epic Games, Inc. All Rights Reserved.
+
+#pragma once
+
+#include "Animation/AnimNode_CustomProperty.h"
+#include "Graph/AnimNextGraph.h"
+#include "Graph/AnimNextGraphInstancePtr.h"
+#include "TraitCore/TraitPtr.h"
+#include "Context.h"
+#include "AnimNode_AnimNextGraph.generated.h"
+
+/**
+ * Animation node that allows a AnimNextGraph output to be used in an animation graph
+ */
+USTRUCT()
+struct ANIMNEXTANIMGRAPH_API FAnimNode_AnimNextGraph : public FAnimNode_CustomProperty
+{
+	GENERATED_BODY()
+
+	FAnimNode_AnimNextGraph();
+
+	// FAnimNode_Base interface
+	virtual void OnInitializeAnimInstance(const FAnimInstanceProxy* InProxy, const UAnimInstance* InAnimInstance) override;
+	virtual void Initialize_AnyThread(const FAnimationInitializeContext& Context) override;
+	virtual void GatherDebugData(FNodeDebugData& DebugData) override;
+	virtual void Update_AnyThread(const FAnimationUpdateContext& Context) override;
+	virtual void CacheBones_AnyThread(const FAnimationCacheBonesContext& Context) override;
+	virtual void Evaluate_AnyThread(FPoseContext & Output) override;
+	virtual int32 GetLODThreshold() const override { return LODThreshold; }
+
+	virtual void InitializeProperties(const UObject* InSourceInstance, UClass* InTargetClass) override;
+	virtual void PropagateInputProperties(const UObject* InSourceInstance) override;
+
+private:
+#if WITH_EDITOR
+	virtual void HandleObjectsReinstanced_Impl(UObject* InSourceObject, UObject* InTargetObject, const TMap<UObject*, UObject*>& OldToNewInstanceMap) override;
+#endif
+private:
+
+	/** The input pose we will pass to the graph */
+	UPROPERTY(EditAnywhere, Category = Links, meta = (DisplayName = "Source"))
+	FPoseLink SourceLink;
+
+	UPROPERTY(EditAnywhere, Category = Settings, meta = (PinHiddenByDefault))
+	TObjectPtr<UAnimNextGraph> AnimNextGraph;
+
+	// Our graph instance, we own it
+	UPROPERTY()
+	FAnimNextGraphInstancePtr GraphInstance;
+
+	/*
+	 * Max LOD that this node is allowed to run
+	 * For example if you have LODThreshold to be 2, it will run until LOD 2 (based on 0 index)
+	 * when the component LOD becomes 3, it will stop update/evaluate
+	 * currently transition would be issue and that has to be re-visited
+	 */
+	UPROPERTY(EditAnywhere, Category = Performance, meta = (DisplayName = "LOD Threshold"))
+	int32 LODThreshold;
+
+protected:
+	virtual UClass* GetTargetClass() const override { return AnimNextGraph ? AnimNextGraph->StaticClass() : nullptr; }
+	
+public:
+
+	void PostSerialize(const FArchive& Ar);
+
+	friend class UAnimGraphNode_AnimNextGraph;
+};
+
+template<>
+struct TStructOpsTypeTraits<FAnimNode_AnimNextGraph> : public TStructOpsTypeTraitsBase2<FAnimNode_AnimNextGraph>
+{
+	enum
+	{
+		WithCopy = false,
+		WithPostSerialize = true,
+	};
+};
\ No newline at end of file
diff --git a/Engine/Plugins/Experimental/AnimNextAnimGraph/Source/AnimNextAnimGraph/Private/Graph/AnimGraph/AnimNode_AnimNextGraph.cpp b/Engine/Plugins/Experimental/AnimNextAnimGraph/Source/AnimNextAnimGraph/Private/Graph/AnimGraph/AnimNode_AnimNextGraph.cpp
new file mode 100644
--- /dev/null
+++ b/Engine/Plugins/Experimental/AnimNextAnimGraph/Source/AnimNextAnimGraph/Private/Graph/AnimGraph/AnimNode_AnimNextGraph.cpp
@@ -0,0 +1,242 @@
+// Copyright Epic Games, Inc. All Rights Reserved.
+
+#include "Graph/AnimGraph/AnimNode_AnimNextGraph.h"
+#include "DataRegistry.h"
+#include "ReferencePose.h"
+#include "Components/SkeletalMeshComponent.h"
+#include "Animation/AnimInstanceProxy.h"
+#include "GameFramework/Actor.h"
+#include "AnimationRuntime.h"
+#include "GenerationTools.h"
+#include "Graph/AnimNext_LODPose.h"
+#include "Engine/SkeletalMesh.h"
+#include "Param/ParamStack.h"
+#include "Graph/AnimGraphParamStackScope.h"
+#include "TraitInterfaces/IEvaluate.h"
+#include "TraitInterfaces/IUpdate.h"
+#include "EvaluationVM/EvaluationVM.h"
+
+#include UE_INLINE_GENERATED_CPP_BY_NAME(AnimNode_AnimNextGraph)
+
+#if WITH_EDITOR
+#include "Editor.h"
+#endif
+
+TAutoConsoleVariable<int32> CVarAnimNextForceAnimBP(TEXT("a.AnimNextForceAnimBP"), 0, TEXT("If != 0, then we use the input pose of the AnimNext AnimBP node instead of the AnimNext graph."));
+
+FAnimNode_AnimNextGraph::FAnimNode_AnimNextGraph()
+	: FAnimNode_CustomProperty()
+	, AnimNextGraph(nullptr)
+	, LODThreshold(INDEX_NONE)
+{
+}
+
+void FAnimNode_AnimNextGraph::OnInitializeAnimInstance(const FAnimInstanceProxy* InProxy, const UAnimInstance* InAnimInstance)
+{
+	DECLARE_SCOPE_HIERARCHICAL_COUNTER_FUNC()
+
+	FAnimNode_CustomProperty::OnInitializeAnimInstance(InProxy, InAnimInstance);
+
+	InitializeProperties(InAnimInstance, GetTargetClass());
+}
+
+void FAnimNode_AnimNextGraph::GatherDebugData(FNodeDebugData& DebugData)
+{
+	DECLARE_SCOPE_HIERARCHICAL_COUNTER_FUNC()
+
+	SourceLink.GatherDebugData(DebugData.BranchFlow(1.f));
+}
+
+void FAnimNode_AnimNextGraph::Update_AnyThread(const FAnimationUpdateContext& Context)
+{
+	DECLARE_SCOPE_HIERARCHICAL_COUNTER_FUNC()
+
+	using namespace UE::AnimNext;
+
+	SourceLink.Update(Context);
+
+	if (IsLODEnabled(Context.AnimInstanceProxy) && !CVarAnimNextForceAnimBP.GetValueOnAnyThread() && GraphInstance.IsValid())
+	{
+		GetEvaluateGraphExposedInputs().Execute(Context);
+
+		PropagateInputProperties(Context.AnimInstanceProxy->GetAnimInstanceObject());
+
+		const int32 LODLevel = Context.AnimInstanceProxy->GetLODLevel();
+
+		UE::AnimNext::FAnimGraphParamStackScope Scope(Context);
+
+		FParamStack& ParamStack = FParamStack::Get();
+		FParamStack::FPushedLayerHandle LayerHandle = ParamStack.PushValues(
+			AnimNextGraph->GetCurrentLODParam(), LODLevel
+		);
+
+		UE::AnimNext::UpdateGraph(GraphInstance, Context.GetDeltaTime());
+
+		ParamStack.PopLayer(LayerHandle);
+	}
+
+	FAnimNode_CustomProperty::Update_AnyThread(Context);
+
+	//TRACE_ANIM_NODE_VALUE(Context, TEXT("Class"), *GetNameSafe(ControlRigClass.Get()));
+}
+
+void FAnimNode_AnimNextGraph::Initialize_AnyThread(const FAnimationInitializeContext& Context)
+{
+	DECLARE_SCOPE_HIERARCHICAL_COUNTER_FUNC()
+
+	using namespace UE::AnimNext;
+
+	SourceLink.Initialize(Context);
+
+	if (!GraphInstance.IsValid() && AnimNextGraph)
+	{
+		// If we don't have an instance yet, create one
+		UE::AnimNext::FAnimGraphParamStackScope Scope(Context);
+
+		// Populate our param stack since our instance data might need it during construction
+		const int32 LODLevel = Context.AnimInstanceProxy->GetLODLevel();
+
+		FParamStack& ParamStack = FParamStack::Get();
+		FParamStack::FPushedLayerHandle LayerHandle = ParamStack.PushValues(
+			AnimNextGraph->GetCurrentLODParam(), LODLevel
+		);
+
+		AnimNextGraph->AllocateInstance(GraphInstance);
+
+		ParamStack.PopLayer(LayerHandle);
+	}
+
+	FAnimNode_CustomProperty::Initialize_AnyThread(Context);
+}
+
+void FAnimNode_AnimNextGraph::CacheBones_AnyThread(const FAnimationCacheBonesContext& Context)
+{
+	DECLARE_SCOPE_HIERARCHICAL_COUNTER_FUNC()
+
+	FAnimNode_CustomProperty::CacheBones_AnyThread(Context);
+
+	SourceLink.CacheBones(Context);
+}
+
+void FAnimNode_AnimNextGraph::Evaluate_AnyThread(FPoseContext& Output)
+{
+	DECLARE_SCOPE_HIERARCHICAL_COUNTER_FUNC()
+
+	using namespace UE::AnimNext;
+
+	if (!CVarAnimNextForceAnimBP.GetValueOnAnyThread() && GraphInstance.IsValid())
+	{
+		USkeletalMeshComponent* SkeletalMeshComponent = Output.AnimInstanceProxy->GetSkelMeshComponent();
+		check(SkeletalMeshComponent != nullptr);
+
+		FDataHandle RefPoseHandle = FDataRegistry::Get()->GetOrGenerateReferencePose(SkeletalMeshComponent);
+		FAnimNextGraphReferencePose GraphReferencePose(RefPoseHandle);
+
+		const int32 LODLevel = Output.AnimInstanceProxy->GetLODLevel();
+
+		const UE::AnimNext::FReferencePose& RefPose = RefPoseHandle.GetRef<UE::AnimNext::FReferencePose>();
+		FAnimNextGraphLODPose ResultPose(FLODPoseHeap(RefPose, LODLevel, true, Output.ExpectsAdditivePose()));
+
+		FAnimGraphParamStackScope Scope(Output);
+		FParamStack& ParamStack = FParamStack::Get();
+		FParamStack::FPushedLayerHandle LayerHandle = ParamStack.PushValues(
+			AnimNextGraph->GetReferencePoseParam(), GraphReferencePose,
+			AnimNextGraph->GetCurrentLODParam(), LODLevel
+		);
+
+		{
+			const FEvaluationProgram EvaluationProgram = UE::AnimNext::EvaluateGraph(GraphInstance);
+
+			FEvaluationVM EvaluationVM(EEvaluationFlags::All, RefPose, LODLevel);
+			bool bHasValidOutput = false;
+
+			if (!EvaluationProgram.IsEmpty())
+			{
+				EvaluationProgram.Execute(EvaluationVM);
+
+				TUniquePtr<FKeyframeState> EvaluatedKeyframe;
+				if (EvaluationVM.PopValue(KEYFRAME_STACK_NAME, EvaluatedKeyframe))
+				{
+					ResultPose.LODPose.CopyFrom(EvaluatedKeyframe->Pose);
+					bHasValidOutput = true;
+				}
+			}
+
+			if (!bHasValidOutput)
+			{
+				// We need to output a valid pose, generate one
+				FKeyframeState ReferenceKeyframe = EvaluationVM.MakeReferenceKeyframe(Output.ExpectsAdditivePose());
+				ResultPose.LODPose.CopyFrom(ReferenceKeyframe.Pose);
+			}
+		}
+
+		FGenerationTools::RemapPose(ResultPose.LODPose, Output);
+
+		ParamStack.PopLayer(LayerHandle);
+	}
+	else
+	{
+		if (SourceLink.GetLinkNode())
+		{
+			SourceLink.Evaluate(Output);
+		}
+	}
+
+	FAnimNode_CustomProperty::Evaluate_AnyThread(Output);
+}
+
+void FAnimNode_AnimNextGraph::PostSerialize(const FArchive& Ar)
+{
+	DECLARE_SCOPE_HIERARCHICAL_COUNTER_FUNC()
+
+	// after compile, we have to reinitialize
+	// because it needs new execution code
+	// since memory has changed
+	if (Ar.IsObjectReferenceCollector())
+	{
+		if (AnimNextGraph)
+		{
+			//AnimNextGraph->Initialize();
+		}
+	}
+}
+
+void FAnimNode_AnimNextGraph::InitializeProperties(const UObject* InSourceInstance, UClass* InTargetClass)
+{
+	// Build property lists
+	SourceProperties.Reset(SourcePropertyNames.Num());
+	DestProperties.Reset(SourcePropertyNames.Num());
+
+	check(SourcePropertyNames.Num() == DestPropertyNames.Num());
+
+	for (int32 Idx = 0; Idx < SourcePropertyNames.Num(); ++Idx)
+	{
+		FName& SourceName = SourcePropertyNames[Idx];
+		UClass* SourceClass = InSourceInstance->GetClass();
+
+		FProperty* SourceProperty = FindFProperty<FProperty>(SourceClass, SourceName);
+		SourceProperties.Add(SourceProperty);
+		DestProperties.Add(nullptr);
+	}
+}
+
+void FAnimNode_AnimNextGraph::PropagateInputProperties(const UObject* InSourceInstance)
+{
+	if (InSourceInstance)
+	{
+		// Assign value to the properties exposed as pins
+		for (int32 PropIdx = 0; PropIdx < SourceProperties.Num(); ++PropIdx)
+		{
+		}
+	}
+}
+
+
+#if WITH_EDITOR
+
+void FAnimNode_AnimNextGraph::HandleObjectsReinstanced_Impl(UObject* InSourceObject, UObject* InTargetObject, const TMap<UObject*, UObject*>& OldToNewInstanceMap)
+{
+	FAnimNode_CustomProperty::HandleObjectsReinstanced_Impl(InSourceObject, InTargetObject, OldToNewInstanceMap);
+}
+
+#endif
diff --git a/Engine/Plugins/Experimental/AnimNextAnimGraph/Source/AnimNextAnimGraph/Private/Graph/SyncGroup_GraphInstanceComponent.cpp b/Engine/Plugins/Experimental/AnimNextAnimGraph/Source/AnimNextAnimGraph/Private/Graph/SyncGroup_GraphInstanceComponent.cpp
new file mode 100644
--- /dev/null
+++ b/Engine/Plugins/Experimental/AnimNextAnimGraph/Source/AnimNextAnimGraph/Private/Graph/SyncGroup_GraphInstanceComponent.cpp
@@ -0,0 +1,102 @@
+// Copyright Epic Games, Inc. All Rights Reserved.
+
+#include "Graph/SyncGroup_GraphInstanceComponent.h"
+#include "TraitInterfaces/IGroupSynchronization.h"
+
+namespace UE::AnimNext
+{
+	void FSyncGroupGraphInstanceComponent::RegisterWithGroup(FName GroupName, EAnimGroupRole::Type GroupRole, const FWeakTraitPtr& TraitPtr, const FTraitUpdateState& TraitState)
+	{
+		FSyncGroupState& GroupState = SyncGroupMap.FindOrAdd(GroupName);
+		GroupState.Members.Add(FSyncGroupMember{ TraitState, TraitPtr, GroupRole });
+	}
+
+	void FSyncGroupGraphInstanceComponent::PreUpdate(FExecutionContext& Context)
+	{
+		// Reset our group state, we want to start fresh every update
+		SyncGroupMap.Reset();
+	}
+
+	void FSyncGroupGraphInstanceComponent::PostUpdate(FExecutionContext& Context)
+	{
+		FTraitStackBinding TraitStack;
+		TTraitBinding<IGroupSynchronization> GroupSyncTrait;
+
+		// Now that we have discovered all groups and their memberships, we can perform synchronization
+		for (const TTuple<FName, FSyncGroupState>& It : SyncGroupMap)
+		{
+			const FSyncGroupState& GroupState = It.Value;
+
+			// Find our leader by looking at the total weight and the group role
+			const int32 NumMembers = GroupState.Members.Num();
+			check(NumMembers != 0);	// Groups should never be empty
+
+			int32 LeaderIndex = INDEX_NONE;
+			float LeaderTotalWeight = -1.0f;
+
+			for (int32 MemberIndex = 0; MemberIndex < NumMembers; ++MemberIndex)
+			{
+				const FSyncGroupMember& GroupMember = GroupState.Members[MemberIndex];
+
+				switch (GroupMember.GroupRole)
+				{
+				case EAnimGroupRole::CanBeLeader:
+				case EAnimGroupRole::TransitionLeader:
+					// Highest weight is the leader
+					if (GroupMember.TraitState.GetTotalWeight() > LeaderTotalWeight)
+					{
+						LeaderIndex = MemberIndex;
+						LeaderTotalWeight = GroupMember.TraitState.GetTotalWeight();
+					}
+					break;
+				case EAnimGroupRole::AlwaysLeader:
+					// Always set the leader index
+					LeaderIndex = MemberIndex;
+					LeaderTotalWeight = 2.0f;		// Some high value
+					break;
+				default:
+				case EAnimGroupRole::AlwaysFollower:
+				case EAnimGroupRole::TransitionFollower:
+					// Never set the leader index
+					// If we find no leader, we'll use the first index as set below
+					break;
+				}
+			}
+
+			if (LeaderIndex == INDEX_NONE)
+			{
+				// If none of the entries wish to be a leader, grab the first and force it
+				LeaderIndex = 0;
+			}
+
+			// Now that we found our leader, advance it by the delta time
+			float LeaderProgressRatio;
+			{
+				const FSyncGroupMember& GroupLeader = GroupState.Members[LeaderIndex];
+
+				Context.BindTo(GroupLeader.TraitPtr);
+				ensure(Context.GetStack(GroupLeader.TraitPtr, TraitStack));
+				ensure(TraitStack.GetInterface(GroupSyncTrait));
+
+				LeaderProgressRatio = GroupSyncTrait.AdvanceBy(Context, GroupLeader.TraitState.GetDeltaTime());
+			}
+
+			// Advance every follower to the same progress ratio as the leader
+			for (int32 MemberIndex = 0; MemberIndex < NumMembers; ++MemberIndex)
+			{
+				if (MemberIndex == LeaderIndex)
+				{
+					continue;	// Ignore the leader, it already advanced
+				}
+
+				const FSyncGroupMember& GroupMember = GroupState.Members[MemberIndex];
+
+				Context.BindTo(GroupMember.TraitPtr);
+				ensure(Context.GetStack(GroupMember.TraitPtr, TraitStack));
+				ensure(TraitStack.GetInterface(GroupSyncTrait));
+
+				GroupSyncTrait.AdvanceToRatio(Context, LeaderProgressRatio);
+			}
+		}
+	}
+}
diff --git a/Engine/Plugins/Experimental/AnimNextAnimGraph/Source/AnimNextAnimGraph/Private/Graph/SyncGroup_GraphInstanceComponent.h b/Engine/Plugins/Experimental/AnimNextAnimGraph/Source/AnimNextAnimGraph/Private/Graph/SyncGroup_GraphInstanceComponent.h
new file mode 100644
--- /dev/null
+++ b/Engine/Plugins/Experimental/AnimNextAnimGraph/Source/AnimNextAnimGraph/Private/Graph/SyncGroup_GraphInstanceComponent.h
@@ -0,0 +1,43 @@
+// Copyright Epic Games, Inc. All Rights Reserved.
+
+#pragma once
+
+#include "CoreMinimal.h"
+#include "Animation/AnimationAsset.h"
+#include "TraitCore/TraitPtr.h"
+#include "TraitInterfaces/IUpdate.h"
+#include "Graph/GraphInstanceComponent.h"
+
+namespace UE::AnimNext
+{
+	/**
+	 * FSyncGroupGraphInstanceComponent
+	 *
+	 * This component maintains the necessary state to support group based synchronization.
+	 */
+	struct FSyncGroupGraphInstanceComponent : public FGraphInstanceComponent
+	{
+		DECLARE_ANIM_GRAPH_INSTANCE_COMPONENT(FSyncGroupGraphInstanceComponent)
+
+		void RegisterWithGroup(FName GroupName, EAnimGroupRole::Type GroupRole, const FWeakTraitPtr& TraitPtr, const FTraitUpdateState& TraitState);
+
+		// FGraphInstanceComponent impl
+		virtual void PreUpdate(FExecutionContext& Context) override;
+		virtual void PostUpdate(FExecutionContext& Context) override;
+
+	private:
+		struct FSyncGroupMember
+		{
+			FTraitUpdateState					TraitState;
+			FWeakTraitPtr						TraitPtr;
+			TEnumAsByte<EAnimGroupRole::Type>	GroupRole;
+		};
+
+		struct FSyncGroupState
+		{
+			TArray<FSyncGroupMember> Members;
+		};
+
+		TMap<FName, FSyncGroupState> SyncGroupMap;
+	};
+}
diff --git a/Engine/Plugins/Experimental/AnimNextAnimGraph/Source/AnimNextAnimGraph/Private/Module.cpp b/Engine/Plugins/Experimental/AnimNextAnimGraph/Source/AnimNextAnimGraph/Private/Module.cpp
new file mode 100644
--- /dev/null
+++ b/Engine/Plugins/Experimental/AnimNextAnimGraph/Source/AnimNextAnimGraph/Private/Module.cpp
@@ -0,0 +1,79 @@
+// Copyright Epic Games, Inc. All Rights Reserved.
+
+#include "CoreMinimal.h"
+#include "Animation/BlendProfile.h"
+#include "Curves/CurveFloat.h"
+#include "Graph/AnimNextGraph.h"
+#include "Modules/ModuleInterface.h"
+#include "Modules/ModuleManager.h"
+#include "RigVMCore/RigVMRegistry.h"
+#include "Animation/AnimSequence.h"
+
+#include "IAnimNextModule.h"
+#include "TraitInterfaces/IEvaluate.h"
+#include "TraitInterfaces/IUpdate.h"
+#include "EvaluationVM/EvaluationVM.h"
+
+namespace UE::AnimNext
+{
+
+class FModule : public IModuleInterface, IAnimNextAnimGraph
+{
+public:
+	virtual void StartupModule() override
+	{
+		static TPair<UClass*, FRigVMRegistry::ERegisterObjectOperation> const AllowedObjectTypes[] =
+		{
+			{ UAnimSequence::StaticClass(), FRigVMRegistry::ERegisterObjectOperation::Class },
+			{ UScriptStruct::StaticClass(), FRigVMRegistry::ERegisterObjectOperation::Class },
+			{ UBlendProfile::StaticClass(), FRigVMRegistry::ERegisterObjectOperation::Class },
+			{ UCurveFloat::StaticClass(), FRigVMRegistry::ERegisterObjectOperation::Class },
+			{ UAnimNextGraph::StaticClass(), FRigVMRegistry::ERegisterObjectOperation::Class },
+		};
+
+		FRigVMRegistry::Get().RegisterObjectTypes(AllowedObjectTypes);
+
+		IAnimNextModule::Get().RegisterAnimNextAnimGraph(*this);
+	}
+
+	virtual void ShutdownModule() override
+	{
+		IAnimNextModule::Get().UnregisterAnimNextAnimGraph();
+	}
+
+	virtual void UpdateGraph(FAnimNextGraphInstancePtr& GraphInstance, float DeltaTime) const override
+	{
+		UE::AnimNext::UpdateGraph(GraphInstance, DeltaTime);
+	}
+
+	virtual void EvaluateGraph(FAnimNextGraphInstancePtr& GraphInstance, const UE::AnimNext::FReferencePose& RefPose, int32 GraphLODLevel, FLODPoseHeap& OutputPose) const override
+	{
+		const FEvaluationProgram EvaluationProgram = UE::AnimNext::EvaluateGraph(GraphInstance);
+
+		FEvaluationVM EvaluationVM(EEvaluationFlags::All, RefPose, GraphLODLevel);
+		bool bHasValidOutput = false;
+
+		if (!EvaluationProgram.IsEmpty())
+		{
+			EvaluationProgram.Execute(EvaluationVM);
+
+			TUniquePtr<FKeyframeState> EvaluatedKeyframe;
+			if (EvaluationVM.PopValue(KEYFRAME_STACK_NAME, EvaluatedKeyframe))
+			{
+				OutputPose.CopyFrom(EvaluatedKeyframe->Pose);
+				bHasValidOutput = true;
+			}
+		}
+
+		if (!bHasValidOutput)
+		{
+			// We need to output a valid pose, generate one
+			FKeyframeState ReferenceKeyframe = EvaluationVM.MakeReferenceKeyframe(false);
+			OutputPose.CopyFrom(ReferenceKeyframe.Pose);
+		}
+	}
+};
+
+}
+
+IMPLEMENT_MODULE(UE::AnimNext::FModule, AnimNextAnimGraph)
diff --git a/Engine/Plugins/Experimental/AnimNextAnimGraph/Source/AnimNextAnimGraph/Private/TraitInterfaces/IContinuousBlend.cpp b/Engine/Plugins/Experimental/AnimNextAnimGraph/Source/AnimNextAnimGraph/Private/TraitInterfaces/IContinuousBlend.cpp
new file mode 100644
--- /dev/null
+++ b/Engine/Plugins/Experimental/AnimNextAnimGraph/Source/AnimNextAnimGraph/Private/TraitInterfaces/IContinuousBlend.cpp
@@ -0,0 +1,19 @@
+// Copyright Epic Games, Inc. All Rights Reserved.
+
+#include "TraitInterfaces/IContinuousBlend.h"
+
+#include "TraitCore/ExecutionContext.h"
+
+namespace UE::AnimNext
+{
+	float IContinuousBlend::GetBlendWeight(const FExecutionContext& Context, const TTraitBinding<IContinuousBlend>& Binding, int32 ChildIndex) const
+	{
+		TTraitBinding<IContinuousBlend> SuperBinding;
+		if (Binding.GetStackInterfaceSuper(SuperBinding))
+		{
+			return SuperBinding.GetBlendWeight(Context, ChildIndex);
+		}
+
+		return -1.0f;
+	}
+}
diff --git a/Engine/Plugins/Experimental/AnimNextAnimGraph/Source/AnimNextAnimGraph/Private/TraitInterfaces/IDiscreteBlend.cpp b/Engine/Plugins/Experimental/AnimNextAnimGraph/Source/AnimNextAnimGraph/Private/TraitInterfaces/IDiscreteBlend.cpp
new file mode 100644
--- /dev/null
+++ b/Engine/Plugins/Experimental/AnimNextAnimGraph/Source/AnimNextAnimGraph/Private/TraitInterfaces/IDiscreteBlend.cpp
@@ -0,0 +1,68 @@
+// Copyright Epic Games, Inc. All Rights Reserved.
+
+#include "TraitInterfaces/IDiscreteBlend.h"
+
+#include "TraitCore/ExecutionContext.h"
+
+namespace UE::AnimNext
+{
+	float IDiscreteBlend::GetBlendWeight(const FExecutionContext& Context, const TTraitBinding<IDiscreteBlend>& Binding, int32 ChildIndex) const
+	{
+		TTraitBinding<IDiscreteBlend> SuperBinding;
+		if (Binding.GetStackInterfaceSuper(SuperBinding))
+		{
+			return SuperBinding.GetBlendWeight(Context, ChildIndex);
+		}
+
+		return -1.0f;
+	}
+
+	const FAlphaBlend* IDiscreteBlend::GetBlendState(const FExecutionContext& Context, const TTraitBinding<IDiscreteBlend>& Binding, int32 ChildIndex) const
+	{
+		TTraitBinding<IDiscreteBlend> SuperBinding;
+		if (Binding.GetStackInterfaceSuper(SuperBinding))
+		{
+			return SuperBinding.GetBlendState(Context, ChildIndex);
+		}
+
+		return nullptr;
+	}
+
+	int32 IDiscreteBlend::GetBlendDestinationChildIndex(const FExecutionContext& Context, const TTraitBinding<IDiscreteBlend>& Binding) const
+	{
+		TTraitBinding<IDiscreteBlend> SuperBinding;
+		if (Binding.GetStackInterfaceSuper(SuperBinding))
+		{
+			return SuperBinding.GetBlendDestinationChildIndex(Context);
+		}
+
+		return INDEX_NONE;
+	}
+
+	void IDiscreteBlend::OnBlendTransition(const FExecutionContext& Context, const TTraitBinding<IDiscreteBlend>& Binding, int32 OldChildIndex, int32 NewChildIndex) const
+	{
+		TTraitBinding<IDiscreteBlend> SuperBinding;
+		if (Binding.GetStackInterfaceSuper(SuperBinding))
+		{
+			return SuperBinding.OnBlendTransition(Context, OldChildIndex, NewChildIndex);
+		}
+	}
+
+	void IDiscreteBlend::OnBlendInitiated(const FExecutionContext& Context, const TTraitBinding<IDiscreteBlend>& Binding, int32 ChildIndex) const
+	{
+		TTraitBinding<IDiscreteBlend> SuperBinding;
+		if (Binding.GetStackInterfaceSuper(SuperBinding))
+		{
+			return SuperBinding.OnBlendInitiated(Context, ChildIndex);
+		}
+	}
+
+	void IDiscreteBlend::OnBlendTerminated(const FExecutionContext& Context, const TTraitBinding<IDiscreteBlend>& Binding, int32 ChildIndex) const
+	{
+		TTraitBinding<IDiscreteBlend> SuperBinding;
+		if (Binding.GetStackInterfaceSuper(SuperBinding))
+		{
+			return SuperBinding.OnBlendTerminated(Context, ChildIndex);
+		}
+	}
+}
diff --git a/Engine/Plugins/Experimental/AnimNextAnimGraph/Source/AnimNextAnimGraph/Private/TraitInterfaces/IEvaluate.cpp b/Engine/Plugins/Experimental/AnimNextAnimGraph/Source/AnimNextAnimGraph/Private/TraitInterfaces/IEvaluate.cpp
new file mode 100644
--- /dev/null
+++ b/Engine/Plugins/Experimental/AnimNextAnimGraph/Source/AnimNextAnimGraph/Private/TraitInterfaces/IEvaluate.cpp
@@ -0,0 +1,206 @@
+// Copyright Epic Games, Inc. All Rights Reserved.
+
+#include "TraitInterfaces/IEvaluate.h"
+
+#include "TraitCore/ExecutionContext.h"
+#include "TraitInterfaces/IHierarchy.h"
+#include "AnimNextAnimGraphStats.h"
+
+DEFINE_STAT(STAT_AnimNext_EvaluateGraph);
+
+namespace UE::AnimNext
+{
+	FEvaluateTraversalContext::FEvaluateTraversalContext(FEvaluationProgram& InEvaluationProgram)
+		: EvaluationProgram(InEvaluationProgram)
+	{
+	}
+
+	void IEvaluate::PreEvaluate(FEvaluateTraversalContext& Context, const TTraitBinding<IEvaluate>& Binding) const
+	{
+		TTraitBinding<IEvaluate> SuperBinding;
+		if (Binding.GetStackInterfaceSuper(SuperBinding))
+		{
+			SuperBinding.PreEvaluate(Context);
+		}
+	}
+
+	void IEvaluate::PostEvaluate(FEvaluateTraversalContext& Context, const TTraitBinding<IEvaluate>& Binding) const
+	{
+		TTraitBinding<IEvaluate> SuperBinding;
+		if (Binding.GetStackInterfaceSuper(SuperBinding))
+		{
+			SuperBinding.PostEvaluate(Context);
+		}
+	}
+
+	enum class EEvaluateStep
+	{
+		PreEvaluate,
+		PostEvaluate,
+	};
+
+	// This structure is transient and lives either on the stack or the memstack and its destructor may not be called
+	struct FEvaluateEntry
+	{
+		// The trait handle that points to our node to update
+		FWeakTraitPtr		TraitPtr;
+
+		// Which step we wish to perform when we next see this entry
+		EEvaluateStep		DesiredStep = EEvaluateStep::PreEvaluate;
+
+		// The trait stack binding for this evaluate entry
+		FTraitStackBinding	TraitStack;
+
+		// Once we've called PreUpdate, we cache the trait binding to avoid a redundant query to call PostUpdate
+		TTraitBinding<IEvaluate> EvaluateTrait;
+
+		// These pointers are mutually exclusive
+		// An entry is either part of the pending stack, the free list, or neither
+		union
+		{
+			// Next entry in the linked list of free entries
+			FEvaluateEntry* NextFreeEntry = nullptr;
+
+			// Previous entry below us in the nodes pending stack
+			FEvaluateEntry* PrevStackEntry;
+		};
+
+		FEvaluateEntry(const FWeakTraitPtr& InTraitPtr, EEvaluateStep InDesiredStep, FEvaluateEntry* InPrevStackEntry = nullptr)
+			: TraitPtr(InTraitPtr)
+			, DesiredStep(InDesiredStep)
+			, PrevStackEntry(InPrevStackEntry)
+		{
+		}
+	};
+
+	FEvaluationProgram EvaluateGraph(FAnimNextGraphInstancePtr& GraphInstance)
+	{
+		return EvaluateGraph(GraphInstance.GetGraphRootPtr());
+	}
+
+	FEvaluationProgram EvaluateGraph(const FWeakTraitPtr& GraphRootPtr)
+	{
+		SCOPE_CYCLE_COUNTER(STAT_AnimNext_EvaluateGraph);
+		
+		FEvaluationProgram EvaluationProgram;
+
+		if (!GraphRootPtr.IsValid())
+		{
+			return EvaluationProgram;	// Nothing to update
+		}
+
+		FMemStack& MemStack = FMemStack::Get();
+		FMemMark Mark(MemStack);
+
+		FChildrenArray Children;
+		Children.Reserve(16);
+
+		FEvaluateTraversalContext TraversalContext(EvaluationProgram);
+		TraversalContext.BindTo(GraphRootPtr);
+
+		// Add the graph root to kick start the evaluation process
+		FEvaluateEntry GraphRootEntry(GraphRootPtr, EEvaluateStep::PreEvaluate);
+		FEvaluateEntry* NodesPendingUpdateStackTop = &GraphRootEntry;
+
+		// List of free entries we can recycle
+		FEvaluateEntry* FreeEntryList = nullptr;
+
+		TTraitBinding<IHierarchy> HierarchyTrait;
+
+		// Process every node twice: pre-evaluate and post-evaluate
+		while (NodesPendingUpdateStackTop != nullptr)
+		{
+			// Grab the top most entry
+			FEvaluateEntry* Entry = NodesPendingUpdateStackTop;
+			bool bIsEntryUsed = true;
+
+			const FWeakTraitPtr& EntryTraitPtr = Entry->TraitPtr;
+
+			if (Entry->DesiredStep == EEvaluateStep::PreEvaluate)
+			{
+				// Bind and cache our trait stack
+				ensure(TraversalContext.GetStack(EntryTraitPtr, Entry->TraitStack));
+
+				if (Entry->TraitStack.GetInterface(Entry->EvaluateTrait))
+				{
+					// This is the first time we visit this node, time to pre-evaluate
+					Entry->EvaluateTrait.PreEvaluate(TraversalContext);
+
+					// Leave our entry on top of the stack, we'll need to call PostEvaluate once the children
+					// we'll push on top finish
+					Entry->DesiredStep = EEvaluateStep::PostEvaluate;
+				}
+				else
+				{
+					// This node doesn't implement IUpdate, we can pop it from the stack
+					NodesPendingUpdateStackTop = Entry->PrevStackEntry;
+					bIsEntryUsed = false;
+				}
+
+				if (Entry->TraitStack.GetInterface(HierarchyTrait))
+				{
+					IHierarchy::GetStackChildren(TraversalContext, Entry->TraitStack, Children);
+
+					// Append our children in reserve order so that they are visited in the same order they were added
+					for (int32 ChildIndex = Children.Num() - 1; ChildIndex >= 0; --ChildIndex)
+					{
+						const FWeakTraitPtr& ChildPtr = Children[ChildIndex];
+						if (!ChildPtr.IsValid())
+						{
+							continue;
+						}
+
+						// Insert our new child on top of the stack
+
+						FEvaluateEntry* ChildEntry;
+						if (FreeEntryList != nullptr)
+						{
+							// Grab an entry from the free list
+							ChildEntry = FreeEntryList;
+							FreeEntryList = ChildEntry->NextFreeEntry;
+
+							ChildEntry->TraitPtr = ChildPtr;
+							ChildEntry->DesiredStep = EEvaluateStep::PreEvaluate;
+							ChildEntry->PrevStackEntry = NodesPendingUpdateStackTop;
+						}
+						else
+						{
+							// Allocate a new entry
+							ChildEntry = new(MemStack) FEvaluateEntry(ChildPtr, EEvaluateStep::PreEvaluate, NodesPendingUpdateStackTop);
+						}
+
+						NodesPendingUpdateStackTop = ChildEntry;
+					}
+
+					// Reset our container for the next time we need it
+					Children.Reset();
+				}
+
+				// Break and continue to the next top-most entry
+				// It is either a child ready for its pre-evaluate or our current entry ready for its post-evaluate (leaf)
+			}
+			else
+			{
+				// We've already visited this node once, time to post-update
+				check(Entry->EvaluateTrait.IsValid());
+				Entry->EvaluateTrait.PostEvaluate(TraversalContext);
+
+				// Now that we are done processing this entry, we can pop it
+				NodesPendingUpdateStackTop = Entry->PrevStackEntry;
+				bIsEntryUsed = false;
+
+				// Break and continue to the next top-most entry
+				// It is either a sibling read for its post-evaluate or our parent entry ready for its post-evaluate
+			}
+
+			if (!bIsEntryUsed)
+			{
+				// This entry is no longer used, add it to the free list
+				Entry->NextFreeEntry = FreeEntryList;
+				FreeEntryList = Entry;
+			}
+		}
+
+		return EvaluationProgram;
+	}
+}
diff --git a/Engine/Plugins/Experimental/AnimNextAnimGraph/Source/AnimNextAnimGraph/Private/TraitInterfaces/IGroupSynchronization.cpp b/Engine/Plugins/Experimental/AnimNextAnimGraph/Source/AnimNextAnimGraph/Private/TraitInterfaces/IGroupSynchronization.cpp
new file mode 100644
--- /dev/null
+++ b/Engine/Plugins/Experimental/AnimNextAnimGraph/Source/AnimNextAnimGraph/Private/TraitInterfaces/IGroupSynchronization.cpp
@@ -0,0 +1,50 @@
+// Copyright Epic Games, Inc. All Rights Reserved.
+
+#include "TraitInterfaces/IGroupSynchronization.h"
+
+#include "TraitCore/ExecutionContext.h"
+
+namespace UE::AnimNext
+{
+	FName IGroupSynchronization::GetGroupName(const FExecutionContext& Context, const TTraitBinding<IGroupSynchronization>& Binding) const
+	{
+		TTraitBinding<IGroupSynchronization> SuperBinding;
+		if (Binding.GetStackInterfaceSuper(SuperBinding))
+		{
+			return SuperBinding.GetGroupName(Context);
+		}
+
+		return NAME_None;
+	}
+
+	EAnimGroupRole::Type IGroupSynchronization::GetGroupRole(const FExecutionContext& Context, const TTraitBinding<IGroupSynchronization>& Binding) const
+	{
+		TTraitBinding<IGroupSynchronization> SuperBinding;
+		if (Binding.GetStackInterfaceSuper(SuperBinding))
+		{
+			return SuperBinding.GetGroupRole(Context);
+		}
+
+		return EAnimGroupRole::CanBeLeader;
+	}
+
+	float IGroupSynchronization::AdvanceBy(const FExecutionContext& Context, const TTraitBinding<IGroupSynchronization>& Binding, float DeltaTime) const
+	{
+		TTraitBinding<IGroupSynchronization> SuperBinding;
+		if (Binding.GetStackInterfaceSuper(SuperBinding))
+		{
+			return SuperBinding.AdvanceBy(Context, DeltaTime);
+		}
+
+		return 0.0f;
+	}
+
+	void IGroupSynchronization::AdvanceToRatio(const FExecutionContext& Context, const TTraitBinding<IGroupSynchronization>& Binding, float ProgressRatio) const
+	{
+		TTraitBinding<IGroupSynchronization> SuperBinding;
+		if (Binding.GetStackInterfaceSuper(SuperBinding))
+		{
+			SuperBinding.AdvanceToRatio(Context, ProgressRatio);
+		}
+	}
+}
diff --git a/Engine/Plugins/Experimental/AnimNextAnimGraph/Source/AnimNextAnimGraph/Private/TraitInterfaces/ISmoothBlend.cpp b/Engine/Plugins/Experimental/AnimNextAnimGraph/Source/AnimNextAnimGraph/Private/TraitInterfaces/ISmoothBlend.cpp
new file mode 100644
--- /dev/null
+++ b/Engine/Plugins/Experimental/AnimNextAnimGraph/Source/AnimNextAnimGraph/Private/TraitInterfaces/ISmoothBlend.cpp
@@ -0,0 +1,19 @@
+// Copyright Epic Games, Inc. All Rights Reserved.
+
+#include "TraitInterfaces/ISmoothBlend.h"
+
+#include "TraitCore/ExecutionContext.h"
+
+namespace UE::AnimNext
+{
+	float ISmoothBlend::GetBlendTime(const FExecutionContext& Context, const TTraitBinding<ISmoothBlend>& Binding, int32 ChildIndex) const
+	{
+		TTraitBinding<ISmoothBlend> SuperBinding;
+		if (Binding.GetStackInterfaceSuper(SuperBinding))
+		{
+			return SuperBinding.GetBlendTime(Context, ChildIndex);
+		}
+
+		return 0.0f;
+	}
+}
diff --git a/Engine/Plugins/Experimental/AnimNextAnimGraph/Source/AnimNextAnimGraph/Private/TraitInterfaces/ITimeline.cpp b/Engine/Plugins/Experimental/AnimNextAnimGraph/Source/AnimNextAnimGraph/Private/TraitInterfaces/ITimeline.cpp
new file mode 100644
--- /dev/null
+++ b/Engine/Plugins/Experimental/AnimNextAnimGraph/Source/AnimNextAnimGraph/Private/TraitInterfaces/ITimeline.cpp
@@ -0,0 +1,39 @@
+// Copyright Epic Games, Inc. All Rights Reserved.
+
+#include "TraitInterfaces/ITimeline.h"
+
+#include "TraitCore/ExecutionContext.h"
+
+namespace UE::AnimNext
+{
+	float ITimeline::GetPlayRate(const FExecutionContext& Context, const TTraitBinding<ITimeline>& Binding) const
+	{
+		TTraitBinding<ITimeline> SuperBinding;
+		if (Binding.GetStackInterfaceSuper(SuperBinding))
+		{
+			return SuperBinding.GetPlayRate(Context);
+		}
+
+		return 1.0f;
+	}
+
+	float ITimeline::AdvanceBy(const FExecutionContext& Context, const TTraitBinding<ITimeline>& Binding, float DeltaTime) const
+	{
+		TTraitBinding<ITimeline> SuperBinding;
+		if (Binding.GetStackInterfaceSuper(SuperBinding))
+		{
+			return SuperBinding.AdvanceBy(Context, DeltaTime);
+		}
+
+		return 0.0f;
+	}
+
+	void ITimeline::AdvanceToRatio(const FExecutionContext& Context, const TTraitBinding<ITimeline>& Binding, float ProgressRatio) const
+	{
+		TTraitBinding<ITimeline> SuperBinding;
+		if (Binding.GetStackInterfaceSuper(SuperBinding))
+		{
+			SuperBinding.AdvanceToRatio(Context, ProgressRatio);
+		}
+	}
+}
diff --git a/Engine/Plugins/Experimental/AnimNextAnimGraph/Source/AnimNextAnimGraph/Private/TraitInterfaces/IUpdate.cpp b/Engine/Plugins/Experimental/AnimNextAnimGraph/Source/AnimNextAnimGraph/Private/TraitInterfaces/IUpdate.cpp
new file mode 100644
--- /dev/null
+++ b/Engine/Plugins/Experimental/AnimNextAnimGraph/Source/AnimNextAnimGraph/Private/TraitInterfaces/IUpdate.cpp
@@ -0,0 +1,350 @@
+// Copyright Epic Games, Inc. All Rights Reserved.
+
+#include "TraitInterfaces/IUpdate.h"
+
+#include "TraitInterfaces/IHierarchy.h"
+#include "Graph/GraphInstanceComponent.h"
+#include "AnimNextAnimGraphStats.h"
+
+DEFINE_STAT(STAT_AnimNext_UpdateGraph);
+
+namespace UE::AnimNext
+{
+	namespace Private
+	{
+		// This structure is transient and lives either on the stack or the memstack and its destructor may not be called
+		struct FUpdateEntry
+		{
+			// The trait state for this entry
+			FTraitUpdateState			TraitState;
+
+			// The trait handle that points to our node to update
+			FWeakTraitPtr				TraitPtr;
+
+			// Whether or not PreUpdate had been called already
+			// TODO: Store bHasPreUpdated in the LSB of the entry pointer to save padding?
+			bool						bHasPreUpdated = false;
+
+			// The trait stack binding for this update entry
+			FTraitStackBinding			TraitStack;
+
+			// Once we've called PreUpdate, we cache the trait binding to avoid a redundant query to call PostUpdate
+			TTraitBinding<IUpdate>		UpdateTrait;
+
+			// These pointers are mutually exclusive
+			// An entry is either part of the queued update stack, the update stack, the free list, or none of the above
+			union
+			{
+				// Next entry in the stack of free entries
+				FUpdateEntry* NextFreeEntry = nullptr;
+
+				// Previous entry on the update stack
+				FUpdateEntry* PrevUpdateStackEntry;
+
+				// Previous entry on the queued update stack
+				FUpdateEntry* PrevQueuedUpdateStackEntry;
+			};
+
+			FUpdateEntry(const FWeakTraitPtr& InTraitPtr, const FTraitUpdateState InTraitState)
+				: TraitState(InTraitState)
+				, TraitPtr(InTraitPtr)
+			{
+			}
+		};
+	}
+
+	void IUpdate::PreUpdate(FUpdateTraversalContext& Context, const TTraitBinding<IUpdate>& Binding, const FTraitUpdateState& TraitState) const
+	{
+		TTraitBinding<IUpdate> SuperBinding;
+		if (Binding.GetStackInterfaceSuper(SuperBinding))
+		{
+			SuperBinding.PreUpdate(Context, TraitState);
+		}
+	}
+
+	void IUpdate::PostUpdate(FUpdateTraversalContext& Context, const TTraitBinding<IUpdate>& Binding, const FTraitUpdateState& TraitState) const
+	{
+		TTraitBinding<IUpdate> SuperBinding;
+		if (Binding.GetStackInterfaceSuper(SuperBinding))
+		{
+			SuperBinding.PostUpdate(Context, TraitState);
+		}
+	}
+
+	void IUpdateTraversal::QueueChildrenForTraversal(FUpdateTraversalContext& Context, const TTraitBinding<IUpdateTraversal>& Binding, const FTraitUpdateState& TraitState, FUpdateTraversalQueue& TraversalQueue) const
+	{
+		// Nothing to do
+		// This function is called for each trait on the stack, one by one
+		// No need to forward to our super
+	}
+
+	//////////////////////////////////////////////////////////////////////////
+	// Traversal implementation
+
+	void FUpdateTraversalContext::PushQueuedUpdateEntries(FUpdateTraversalQueue& TraversalQueue)
+	{
+		// Pop every entry from the queued update stack and push them onto the update stack
+		// reversing their order
+		while (Private::FUpdateEntry* Entry = TraversalQueue.QueuedUpdateStackHead)
+		{
+			// Update our queued stack head
+			TraversalQueue.QueuedUpdateStackHead = Entry->PrevQueuedUpdateStackEntry;
+
+			// Push our new entry onto the update stack
+			PushUpdateEntry(Entry);
+		}
+	}
+
+	void FUpdateTraversalContext::PushUpdateEntry(Private::FUpdateEntry* Entry)
+	{
+		Entry->PrevUpdateStackEntry = UpdateStackHead;
+		UpdateStackHead = Entry;
+	}
+
+	Private::FUpdateEntry* FUpdateTraversalContext::PopUpdateEntry()
+	{
+		Private::FUpdateEntry* ChildEntry = UpdateStackHead;
+		if (ChildEntry != nullptr)
+		{
+			// We have a child, set our new head
+			UpdateStackHead = ChildEntry->PrevUpdateStackEntry;
+		}
+
+		return ChildEntry;
+	}
+
+	void FUpdateTraversalContext::PushFreeEntry(Private::FUpdateEntry* Entry)
+	{
+		Entry->NextFreeEntry = FreeEntryStackHead;
+		FreeEntryStackHead = Entry;
+	}
+
+	Private::FUpdateEntry* FUpdateTraversalContext::GetNewEntry(const FWeakTraitPtr& TraitPtr, const FTraitUpdateState& TraitState)
+	{
+		Private::FUpdateEntry* FreeEntry = FreeEntryStackHead;
+		if (FreeEntry != nullptr)
+		{
+			// We have a free entry, set our new head
+			FreeEntryStackHead = FreeEntry->NextFreeEntry;
+
+			// Update our entry
+			FreeEntry->TraitState = TraitState;
+			FreeEntry->TraitPtr = TraitPtr;
+			FreeEntry->bHasPreUpdated = false;
+			FreeEntry->NextFreeEntry = nullptr;		// Mark it as not being a member of any list
+		}
+		else
+		{
+			// Allocate a new entry
+			FreeEntry = new(MemStack) Private::FUpdateEntry(TraitPtr, TraitState);
+		}
+
+		return FreeEntry;
+	}
+
+	FUpdateTraversalQueue::FUpdateTraversalQueue(FUpdateTraversalContext& InTraversalContext)
+		: TraversalContext(InTraversalContext)
+	{
+	}
+
+	void FUpdateTraversalQueue::Push(const FWeakTraitPtr& ChildPtr, const FTraitUpdateState& ChildTraitState)
+	{
+		if (!ChildPtr.IsValid())
+		{
+			return;	// Don't queue invalid pointers
+		}
+
+		Private::FUpdateEntry* ChildEntry = TraversalContext.GetNewEntry(ChildPtr, ChildTraitState);
+
+		// We push children that are queued onto a stack
+		// Once pre-update is done, we'll pop queued entries one by one and push them
+		// onto the update stack
+		// This has the effect of reversing the entries so that they are traversed in
+		// the same order they are queued in:
+		//    - First queued will be at the bottom of the queued stack and it ends up at the
+		//      at the top of the update stack (last entry pushed)
+		ChildEntry->PrevQueuedUpdateStackEntry = QueuedUpdateStackHead;
+		QueuedUpdateStackHead = ChildEntry;
+	}
+
+	// Performance note
+	// When we process an animation graph for a frame, typically we'll update first before we evaluate
+	// As a result of this, when we query for the update interface here, we will likely hit cold memory
+	// which will cache miss (by touching the graph instance for the first time).
+	// 
+	// The processor will cache miss and continue to process as many instructions as it can before the
+	// out-of-order execution window fills up. This is problematic here because a lot of the subsequent
+	// instructions depend on the node instance and the interface it returns. The processor will be unable
+	// to execute any of the instructions that follow in the current loop iteration. However, it might
+	// be able to get started on the next node entry which is likely to cache miss as well. Should the
+	// processor make it that far and it turns out that we have to push a child onto the stack, all
+	// of the work it tried to do ahead of time will have to be thrown away.
+	// 
+	// There are two things that we can do here to try and help performance: prefetch ahead and bulk
+	// query.
+	// 
+	// If we prefetch, we have to be careful because we do not know what the node will do in its PreUpdate.
+	// If it turns out that it does a lot of work, our prefetch might end up getting thrown out. This is
+	// because prefetched cache lines typically end up being the first evicted unless they are touched first.
+	// It is thus dangerous to use manual prefetching when the memory access pattern isn't fully known.
+	// In practice, it is likely viable as most nodes won't do too much work.
+	// 
+	// A better approach could be to instead bulk query for our interfaces. We could cache in the FUpdateEntry
+	// the trait bindings for IUpdate and IHierarchy (and re-use the binding for IUpdate for PostUpdate).
+	// Every iteration we could check how many children are queued up on the stack. We could then grab N
+	// entries (2 to 4) and query their interfaces in bulk. The idea is to clump the cache miss instructions
+	// together and to interleave the interface queries. This will queue up as much work as possible in the
+	// out-of-order execution window that will not be thrown away because of a branch. Eventually the first
+	// interface query will complete and execution will resume here to call PreUpdate etc. This will be able
+	// to happen while the processor still waits on the cache misses and finishes the interface query of the
+	// other bulked children. The same effect could be achieved by querying the interfaces after the call
+	// to GetChildren by bulk querying all of them right then. This way, as soon as the execution window
+	// can clear the end of the loop, it can start working on the next entry which will be warm in the L1
+	// cache allowing the CPU to carry ahead before all child interfaces are fully resolved.
+	// 
+	// The above may seem like a stretch and an insignificant over optimization but it could very well be the
+	// key to unlocking large performance gains during traversal. The above optimization would allow us to
+	// perform as much useful work as possible while waiting for memory, hiding its slow latency by fully
+	// leveraging out-of-order CPU execution.
+
+	void UpdateGraph(FAnimNextGraphInstancePtr& GraphInstance, float DeltaTime)
+	{
+		SCOPE_CYCLE_COUNTER(STAT_AnimNext_UpdateGraph);
+		
+		if (!GraphInstance.IsValid())
+		{
+			return;	// Nothing to update
+		}
+
+		if (!ensure(GraphInstance.IsRoot()))
+		{
+			return;	// We can only update starting at the root
+		}
+
+		FUpdateTraversalContext TraversalContext;
+
+		FMemStack& MemStack = TraversalContext.GetMemStack();
+		FMemMark Mark(MemStack);
+
+		FChildrenArray Children;
+		TTraitBinding<IHierarchy> HierarchyTrait;
+		TTraitBinding<IUpdateTraversal> UpdateTraversalTrait;
+		FTraitBinding TraitBinding;
+
+		FUpdateTraversalQueue TraversalQueue(TraversalContext);
+
+		// Before we start the traversal, we give the graph instance components the chance to do some work
+		TraversalContext.BindTo(GraphInstance);
+		for (auto It = TraversalContext.GetComponentIterator(); It; ++It)
+		{
+			It.Value()->PreUpdate(TraversalContext);
+		}
+
+		// Add the graph root to start the update process
+		Private::FUpdateEntry RootEntry(GraphInstance.GetGraphRootPtr(), FTraitUpdateState(DeltaTime));
+		TraversalContext.PushUpdateEntry(&RootEntry);
+
+		while (Private::FUpdateEntry* Entry = TraversalContext.PopUpdateEntry())
+		{
+			const FWeakTraitPtr& EntryTraitPtr = Entry->TraitPtr;
+
+			if (!Entry->bHasPreUpdated)
+			{
+				// This is the first time we visit this node, time to pre-update
+
+				// Bind and cache our trait stack
+				ensure(TraversalContext.GetStack(EntryTraitPtr, Entry->TraitStack));
+
+				// First, if it has latent pins, we must execute and cache their results
+				// This will ensure that other calls into this node will have a consistent view of
+				// what the node saw when it started to update. We thus take a snapshot.
+				const bool bIsFrozen = false;	// Not yet supported
+				Entry->TraitStack.SnapshotLatentProperties(bIsFrozen);
+
+				// If this trait stack implements IUpdate, call into it
+				if (Entry->TraitStack.GetInterface(Entry->UpdateTrait))
+				{
+					Entry->UpdateTrait.PreUpdate(TraversalContext, Entry->TraitState);
+
+					// Make sure that next time we visit this entry, we'll post-update
+					Entry->bHasPreUpdated = true;
+
+					// Push this entry onto the update stack, we'll call it once all our children have finished executing
+					TraversalContext.PushUpdateEntry(Entry);
+				}
+				else
+				{
+					// We don't need this entry anymore
+					TraversalContext.PushFreeEntry(Entry);
+				}
+
+				// Now visit the trait stack and queue our children
+				ensure(Entry->TraitStack.GetTopTrait(TraitBinding));
+				do
+				{
+					if (TraitBinding.AsInterface(UpdateTraversalTrait))
+					{
+						// Request that the trait queues the children it wants to visit
+						// This is a separate function from PreUpdate to simplify traversal management. It is often the case that
+						// the base trait is the one best placed to figure out how to optimally queue children since it
+						// owns the handles to them. However, if an additive trait wishes to override PreUpdate, it might want
+						// to perform logic after the base PreUpdate but before children are queued. Without a separate function,
+						// we would have to rewrite the base PreUpdate entirely and use IHierarchy to query the handles of our children.
+						UpdateTraversalTrait.QueueChildrenForTraversal(TraversalContext, Entry->TraitState, TraversalQueue);
+
+
+						// Iterate over our queued children and push them onto the update stack
+						// We do this to allow children to be queued in traversal order which is intuitive
+						// but to traverse them in that order, they must be pushed in reverse order onto the update stack
+						TraversalContext.PushQueuedUpdateEntries(TraversalQueue);
+					}
+					else if (TraitBinding.AsInterface(HierarchyTrait))
+					{
+						HierarchyTrait.GetChildren(TraversalContext, Children);
+
+						// Append our children in reserve order so that they are visited in the same order they were added
+						for (int32 ChildIndex = Children.Num() - 1; ChildIndex >= 0; --ChildIndex)
+						{
+							const FWeakTraitPtr& ChildPtr = Children[ChildIndex];
+							if (ChildPtr.IsValid())
+							{
+								Private::FUpdateEntry* ChildEntry = TraversalContext.GetNewEntry(ChildPtr, Entry->TraitState);
+								TraversalContext.PushUpdateEntry(ChildEntry);
+							}
+						}
+
+						// Reset our container for the next time we need it
+						Children.Reset();
+					}
+					else
+					{
+						// We don't have any children since we don't implement any of the relevant interfaces
+					}
+				} while (Entry->TraitStack.GetParentTrait(TraitBinding, TraitBinding));
+			}
+			else
+			{
+				// We've already visited this node once, time to PostUpdate
+				check(Entry->UpdateTrait.IsValid());
+				Entry->UpdateTrait.PostUpdate(TraversalContext, Entry->TraitState);
+
+				// Now that it finished updating, we can pop any scoped interfaces this node might have pushed
+				TraversalContext.PopStackScopedInterfaces(Entry->TraitStack);
+
+				// We don't need this entry anymore
+				TraversalContext.PushFreeEntry(Entry);
+			}
+		}
+
+		// After we finish the traversal, we give the graph instance components the chance to do some work
+		TraversalContext.BindTo(GraphInstance);
+		for (auto It = TraversalContext.GetComponentIterator(); It; ++It)
+		{
+			It.Value()->PostUpdate(TraversalContext);
+		}
+
+		// At this point, we shouldn't have any remaining scoped interfaces
+		// If this fails, it means we failed to pop them due to a push/pop mismatch
+		ensure(!TraversalContext.HasScopedInterfaces());
+	}
+}
diff --git a/Engine/Plugins/Experimental/AnimNextAnimGraph/Source/AnimNextAnimGraph/Private/Traits/BlendByBool.cpp b/Engine/Plugins/Experimental/AnimNextAnimGraph/Source/AnimNextAnimGraph/Private/Traits/BlendByBool.cpp
new file mode 100644
--- /dev/null
+++ b/Engine/Plugins/Experimental/AnimNextAnimGraph/Source/AnimNextAnimGraph/Private/Traits/BlendByBool.cpp
@@ -0,0 +1,170 @@
+// Copyright Epic Games, Inc. All Rights Reserved.
+
+#include "Traits/BlendByBool.h"
+
+#include "Animation/AnimTypes.h"
+#include "TraitCore/ExecutionContext.h"
+#include "EvaluationVM/Tasks/BlendKeyframes.h"
+
+namespace UE::AnimNext
+{
+	AUTO_REGISTER_ANIM_TRAIT(FBlendByBoolTrait)
+
+	// Trait implementation boilerplate
+	#define TRAIT_INTERFACE_ENUMERATOR(GeneratorMacro) \
+		GeneratorMacro(IDiscreteBlend) \
+		GeneratorMacro(IHierarchy) \
+		GeneratorMacro(IUpdate) \
+		GeneratorMacro(IUpdateTraversal) \
+
+	GENERATE_ANIM_TRAIT_IMPLEMENTATION(FBlendByBoolTrait, TRAIT_INTERFACE_ENUMERATOR)
+	#undef TRAIT_INTERFACE_ENUMERATOR
+
+	static constexpr int32 TRUE_CHILD_INDEX = 0;
+	static constexpr int32 FALSE_CHILD_INDEX = 1;
+
+	void FBlendByBoolTrait::PreUpdate(FUpdateTraversalContext& Context, const TTraitBinding<IUpdate>& Binding, const FTraitUpdateState& TraitState) const
+	{
+		const FSharedData* SharedData = Binding.GetSharedData<FSharedData>();
+		FInstanceData* InstanceData = Binding.GetInstanceData<FInstanceData>();
+
+		TTraitBinding<IDiscreteBlend> DiscreteBlendTrait;
+		Binding.GetStackInterface(DiscreteBlendTrait);
+
+		const int32 DestinationChildIndex = DiscreteBlendTrait.GetBlendDestinationChildIndex(Context);
+		if (InstanceData->PreviousChildIndex != DestinationChildIndex)
+		{
+			DiscreteBlendTrait.OnBlendTransition(Context, InstanceData->PreviousChildIndex, DestinationChildIndex);
+
+			InstanceData->PreviousChildIndex = DestinationChildIndex;
+		}
+	}
+
+	void FBlendByBoolTrait::QueueChildrenForTraversal(FUpdateTraversalContext& Context, const TTraitBinding<IUpdateTraversal>& Binding, const FTraitUpdateState& TraitState, FUpdateTraversalQueue& TraversalQueue) const
+	{
+		const FInstanceData* InstanceData = Binding.GetInstanceData<FInstanceData>();
+
+		// The destination child index has been updated in PreUpdate, we can use the cached version
+		const int32 DestinationChildIndex = InstanceData->PreviousChildIndex;
+
+		TTraitBinding<IDiscreteBlend> DiscreteBlendTrait;
+		Binding.GetStackInterface(DiscreteBlendTrait);
+
+		const float BlendWeightTrue = DiscreteBlendTrait.GetBlendWeight(Context, TRUE_CHILD_INDEX);
+		if (InstanceData->TrueChild.IsValid() && FAnimWeight::IsRelevant(BlendWeightTrue))
+		{
+			FTraitUpdateState TraitStateTrue = TraitState.WithWeight(BlendWeightTrue);
+			if (DestinationChildIndex != TRUE_CHILD_INDEX)
+			{
+				TraitStateTrue = TraitStateTrue.AsBlendingOut();
+			}
+
+			TraversalQueue.Push(InstanceData->TrueChild, TraitStateTrue);
+		}
+
+		const float BlendWeightFalse = 1.0f - BlendWeightTrue;
+		if (InstanceData->FalseChild.IsValid() && FAnimWeight::IsRelevant(BlendWeightFalse))
+		{
+			FTraitUpdateState TraitStateFalse = TraitState.WithWeight(BlendWeightFalse);
+			if (DestinationChildIndex != FALSE_CHILD_INDEX)
+			{
+				TraitStateFalse = TraitStateFalse.AsBlendingOut();
+			}
+
+			TraversalQueue.Push(InstanceData->FalseChild, TraitStateFalse);
+		}
+	}
+
+	uint32 FBlendByBoolTrait::GetNumChildren(const FExecutionContext& Context, const TTraitBinding<IHierarchy>& Binding) const
+	{
+		return 2;
+	}
+
+	void FBlendByBoolTrait::GetChildren(const FExecutionContext& Context, const TTraitBinding<IHierarchy>& Binding, FChildrenArray& Children) const
+	{
+		const FInstanceData* InstanceData = Binding.GetInstanceData<FInstanceData>();
+
+		// Add the two children, even if the handles are empty
+		Children.Add(InstanceData->TrueChild);
+		Children.Add(InstanceData->FalseChild);
+	}
+
+	float FBlendByBoolTrait::GetBlendWeight(const FExecutionContext& Context, const TTraitBinding<IDiscreteBlend>& Binding, int32 ChildIndex) const
+	{
+		TTraitBinding<IDiscreteBlend> DiscreteBlendTrait;
+		Binding.GetStackInterface(DiscreteBlendTrait);
+
+		const float DestinationChildIndex = DiscreteBlendTrait.GetBlendDestinationChildIndex(Context);
+
+		if (ChildIndex == TRUE_CHILD_INDEX)
+		{
+			return (DestinationChildIndex == TRUE_CHILD_INDEX) ? 1.0f : 0.0f;
+		}
+		else if (ChildIndex == FALSE_CHILD_INDEX)
+		{
+			return (DestinationChildIndex == FALSE_CHILD_INDEX) ? 1.0f : 0.0f;
+		}
+		else
+		{
+			// Invalid child index
+			return -1.0f;
+		}
+	}
+
+	int32 FBlendByBoolTrait::GetBlendDestinationChildIndex(const FExecutionContext& Context, const TTraitBinding<IDiscreteBlend>& Binding) const
+	{
+		const FSharedData* SharedData = Binding.GetSharedData<FSharedData>();
+
+		const bool bCondition = SharedData->GetbCondition(Binding);
+		return bCondition ? TRUE_CHILD_INDEX : FALSE_CHILD_INDEX;
+	}
+
+	void FBlendByBoolTrait::OnBlendTransition(const FExecutionContext& Context, const TTraitBinding<IDiscreteBlend>& Binding, int32 OldChildIndex, int32 NewChildIndex) const
+	{
+		TTraitBinding<IDiscreteBlend> DiscreteBlendTrait;
+		Binding.GetStackInterface(DiscreteBlendTrait);
+
+		// We initiate immediately when we transition
+		DiscreteBlendTrait.OnBlendInitiated(Context, NewChildIndex);
+
+		// We terminate immediately when we transition
+		DiscreteBlendTrait.OnBlendTerminated(Context, OldChildIndex);
+	}
+
+	void FBlendByBoolTrait::OnBlendInitiated(const FExecutionContext& Context, const TTraitBinding<IDiscreteBlend>& Binding, int32 ChildIndex) const
+	{
+		const FSharedData* SharedData = Binding.GetSharedData<FSharedData>();
+		FInstanceData* InstanceData = Binding.GetInstanceData<FInstanceData>();
+
+		// Allocate our new child instance
+		if (ChildIndex == TRUE_CHILD_INDEX)
+		{
+			if (!InstanceData->TrueChild.IsValid())
+			{
+				InstanceData->TrueChild = Context.AllocateNodeInstance(Binding, SharedData->TrueChild);
+			}
+		}
+		else if (ChildIndex == FALSE_CHILD_INDEX)
+		{
+			if (!InstanceData->FalseChild.IsValid())
+			{
+				InstanceData->FalseChild = Context.AllocateNodeInstance(Binding, SharedData->FalseChild);
+			}
+		}
+	}
+
+	void FBlendByBoolTrait::OnBlendTerminated(const FExecutionContext& Context, const TTraitBinding<IDiscreteBlend>& Binding, int32 ChildIndex) const
+	{
+		FInstanceData* InstanceData = Binding.GetInstanceData<FInstanceData>();
+
+		// Deallocate our child instance
+		if (ChildIndex == TRUE_CHILD_INDEX)
+		{
+			InstanceData->TrueChild.Reset();
+		}
+		else if (ChildIndex == FALSE_CHILD_INDEX)
+		{
+			InstanceData->FalseChild.Reset();
+		}
+	}
+}
diff --git a/Engine/Plugins/Experimental/AnimNextAnimGraph/Source/AnimNextAnimGraph/Private/Traits/BlendByBool.h b/Engine/Plugins/Experimental/AnimNextAnimGraph/Source/AnimNextAnimGraph/Private/Traits/BlendByBool.h
new file mode 100644
--- /dev/null
+++ b/Engine/Plugins/Experimental/AnimNextAnimGraph/Source/AnimNextAnimGraph/Private/Traits/BlendByBool.h
@@ -0,0 +1,77 @@
+// Copyright Epic Games, Inc. All Rights Reserved.
+
+#pragma once
+
+#include "CoreMinimal.h"
+
+#include "TraitCore/Trait.h"
+#include "TraitInterfaces/IDiscreteBlend.h"
+#include "TraitInterfaces/IHierarchy.h"
+#include "TraitInterfaces/IUpdate.h"
+
+#include "BlendByBool.generated.h"
+
+USTRUCT(meta = (DisplayName = "Blend By Bool"))
+struct FAnimNextBlendByBoolTraitSharedData : public FAnimNextTraitSharedData
+{
+	GENERATED_BODY()
+
+	/** First output to be blended. */
+	UPROPERTY()
+	FAnimNextTraitHandle TrueChild;
+
+	/** Second output to be blended. */
+	UPROPERTY()
+	FAnimNextTraitHandle FalseChild;
+
+	/** The boolean condition that decides which child is active. */
+	UPROPERTY(EditAnywhere, Category = "Default")
+	bool bCondition = false;
+
+	// Latent pin support boilerplate
+	#define TRAIT_LATENT_PROPERTIES_ENUMERATOR(GeneratorMacro) \
+		GeneratorMacro(bCondition) \
+
+	GENERATE_TRAIT_LATENT_PROPERTIES(FAnimNextBlendByBoolTraitSharedData, TRAIT_LATENT_PROPERTIES_ENUMERATOR)
+	#undef TRAIT_LATENT_PROPERTIES_ENUMERATOR
+};
+
+namespace UE::AnimNext
+{
+	/**
+	 * FBlendByBoolTrait
+	 * 
+	 * A trait that can blend two discrete inputs through a boolean.
+	 */
+	struct FBlendByBoolTrait : FBaseTrait, IUpdate, IUpdateTraversal, IHierarchy, IDiscreteBlend
+	{
+		DECLARE_ANIM_TRAIT(FBlendByBoolTrait, 0x452572ff, FBaseTrait)
+
+		using FSharedData = FAnimNextBlendByBoolTraitSharedData;
+
+		struct FInstanceData : FTrait::FInstanceData
+		{
+			FTraitPtr TrueChild;
+			FTraitPtr FalseChild;
+
+			int32 PreviousChildIndex = INDEX_NONE;
+		};
+
+		// IUpdate impl
+		virtual void PreUpdate(FUpdateTraversalContext& Context, const TTraitBinding<IUpdate>& Binding, const FTraitUpdateState& TraitState) const override;
+
+		// IUpdateTraversal impl
+		virtual void QueueChildrenForTraversal(FUpdateTraversalContext& Context, const TTraitBinding<IUpdateTraversal>& Binding, const FTraitUpdateState& TraitState, FUpdateTraversalQueue& TraversalQueue) const override;
+
+		// IHierarchy impl
+		virtual uint32 GetNumChildren(const FExecutionContext& Context, const TTraitBinding<IHierarchy>& Binding) const override;
+		virtual void GetChildren(const FExecutionContext& Context, const TTraitBinding<IHierarchy>& Binding, FChildrenArray& Children) const override;
+
+		// IDiscreteBlend impl
+		virtual float GetBlendWeight(const FExecutionContext& Context, const TTraitBinding<IDiscreteBlend>& Binding, int32 ChildIndex) const override;
+		virtual int32 GetBlendDestinationChildIndex(const FExecutionContext& Context, const TTraitBinding<IDiscreteBlend>& Binding) const override;
+		virtual void OnBlendTransition(const FExecutionContext& Context, const TTraitBinding<IDiscreteBlend>& Binding, int32 OldChildIndex, int32 NewChildIndex) const override;
+		virtual void OnBlendInitiated(const FExecutionContext& Context, const TTraitBinding<IDiscreteBlend>& Binding, int32 ChildIndex) const override;
+		virtual void OnBlendTerminated(const FExecutionContext& Context, const TTraitBinding<IDiscreteBlend>& Binding, int32 ChildIndex) const override;
+	};
+}
diff --git a/Engine/Plugins/Experimental/AnimNextAnimGraph/Source/AnimNextAnimGraph/Private/Traits/BlendInertializer.cpp b/Engine/Plugins/Experimental/AnimNextAnimGraph/Source/AnimNextAnimGraph/Private/Traits/BlendInertializer.cpp
new file mode 100644
--- /dev/null
+++ b/Engine/Plugins/Experimental/AnimNextAnimGraph/Source/AnimNextAnimGraph/Private/Traits/BlendInertializer.cpp
@@ -0,0 +1,51 @@
+// Copyright Epic Games, Inc. All Rights Reserved.
+
+#include "Traits/BlendInertializer.h"
+
+#include "TraitCore/ExecutionContext.h"
+
+namespace UE::AnimNext
+{
+	AUTO_REGISTER_ANIM_TRAIT(FBlendInertializerTrait)
+
+	// Trait implementation boilerplate
+	#define TRAIT_INTERFACE_ENUMERATOR(GeneratorMacro) \
+		GeneratorMacro(IDiscreteBlend) \
+		GeneratorMacro(ISmoothBlend) \
+
+	GENERATE_ANIM_TRAIT_IMPLEMENTATION(FBlendInertializerTrait, TRAIT_INTERFACE_ENUMERATOR)
+	#undef TRAIT_INTERFACE_ENUMERATOR
+
+	void FBlendInertializerTrait::OnBlendTransition(const FExecutionContext& Context, const TTraitBinding<IDiscreteBlend>& Binding, int32 OldChildIndex, int32 NewChildIndex) const
+	{
+		// Trigger the new transition
+		IDiscreteBlend::OnBlendTransition(Context, Binding, OldChildIndex, NewChildIndex);
+
+		// TODO: Implement the inertialization request API
+#if 0
+		UE::Anim::IInertializationRequester* InertializationRequester = Context.GetMessage<UE::Anim::IInertializationRequester>();
+		if (InertializationRequester)
+		{
+			FInertializationRequest Request;
+			Request.Duration = CurrentBlendTimes[ChildIndex];		// TODO: Get from ISmoothBlend interface using super
+			Request.BlendProfile = CurrentBlendProfile;				// TODO: Get from shared data
+			Request.bUseBlendMode = true;
+			Request.BlendMode = GetBlendType();						// TODO: Add to ISmoothBlend interface
+			Request.CustomBlendCurve = GetCustomBlendCurve();		// TODO: Add to ISmoothBlend interface
+
+			InertializationRequester->RequestInertialization(Request);
+			bRequestedInertializationOnActiveChildIndexChange = true;
+		}
+		else
+		{
+			FAnimNode_Inertialization::LogRequestError(Context, BlendPose[ChildIndex]);
+		}
+#endif
+	}
+
+	float FBlendInertializerTrait::GetBlendTime(const FExecutionContext& Context, const TTraitBinding<ISmoothBlend>& Binding, int32 ChildIndex) const
+	{
+		// We hijack the blend time and always transition instantaneously
+		return 0.0f;
+	}
+}
diff --git a/Engine/Plugins/Experimental/AnimNextAnimGraph/Source/AnimNextAnimGraph/Private/Traits/BlendInertializer.h b/Engine/Plugins/Experimental/AnimNextAnimGraph/Source/AnimNextAnimGraph/Private/Traits/BlendInertializer.h
new file mode 100644
--- /dev/null
+++ b/Engine/Plugins/Experimental/AnimNextAnimGraph/Source/AnimNextAnimGraph/Private/Traits/BlendInertializer.h
@@ -0,0 +1,46 @@
+// Copyright Epic Games, Inc. All Rights Reserved.
+
+#pragma once
+
+#include "CoreMinimal.h"
+
+#include "TraitCore/Trait.h"
+#include "TraitInterfaces/IDiscreteBlend.h"
+#include "TraitInterfaces/ISmoothBlend.h"
+
+#include "BlendInertializer.generated.h"
+
+class UBlendProfile;
+
+USTRUCT(meta = (DisplayName = "Blend Inertializer"))
+struct FAnimNextBlendInertializerTraitSharedData : public FAnimNextTraitSharedData
+{
+	GENERATED_BODY()
+
+	/** Blend profile that configures how fast to blend each bone. */
+	// TODO: Can't show list of blend profiles, we need to find a skeleton to perform the lookup with
+	UPROPERTY(EditAnywhere, Category = "Default", meta = (Inline, UseAsBlendProfile = true))
+	TObjectPtr<UBlendProfile> BlendProfile = nullptr;
+};
+
+namespace UE::AnimNext
+{
+	/**
+	 * FBlendInertializerTrait
+	 * 
+	 * A trait that converts a normal smooth blend into an inertializing blend.
+	 */
+	struct FBlendInertializerTrait : FAdditiveTrait, IDiscreteBlend, ISmoothBlend
+	{
+		DECLARE_ANIM_TRAIT(FBlendInertializerTrait, 0x7ea0bdee, FAdditiveTrait)
+
+		using FSharedData = FAnimNextBlendInertializerTraitSharedData;
+		using FInstanceData = FTrait::FInstanceData;
+
+		// IDiscreteBlend impl
+		virtual void OnBlendTransition(const FExecutionContext& Context, const TTraitBinding<IDiscreteBlend>& Binding, int32 OldChildIndex, int32 NewChildIndex) const override;
+
+		// ISmoothBlend impl
+		virtual float GetBlendTime(const FExecutionContext& Context, const TTraitBinding<ISmoothBlend>& Binding, int32 ChildIndex) const override;
+	};
+}
diff --git a/Engine/Plugins/Experimental/AnimNextAnimGraph/Source/AnimNextAnimGraph/Private/Traits/BlendSmoother.cpp b/Engine/Plugins/Experimental/AnimNextAnimGraph/Source/AnimNextAnimGraph/Private/Traits/BlendSmoother.cpp
new file mode 100644
--- /dev/null
+++ b/Engine/Plugins/Experimental/AnimNextAnimGraph/Source/AnimNextAnimGraph/Private/Traits/BlendSmoother.cpp
@@ -0,0 +1,269 @@
+// Copyright Epic Games, Inc. All Rights Reserved.
+
+#include "Traits/BlendSmoother.h"
+
+#include "Animation/AnimTypes.h"
+#include "TraitCore/ExecutionContext.h"
+#include "TraitInterfaces/IHierarchy.h"
+#include "EvaluationVM/Tasks/BlendKeyframes.h"
+#include "EvaluationVM/Tasks/NormalizeRotations.h"
+
+namespace UE::AnimNext
+{
+	AUTO_REGISTER_ANIM_TRAIT(FBlendSmootherTrait)
+
+	// Trait implementation boilerplate
+	#define TRAIT_INTERFACE_ENUMERATOR(GeneratorMacro) \
+		GeneratorMacro(IDiscreteBlend) \
+		GeneratorMacro(IEvaluate) \
+		GeneratorMacro(ISmoothBlend) \
+		GeneratorMacro(IUpdate) \
+
+	GENERATE_ANIM_TRAIT_IMPLEMENTATION(FBlendSmootherTrait, TRAIT_INTERFACE_ENUMERATOR)
+	#undef TRAIT_INTERFACE_ENUMERATOR
+
+	void FBlendSmootherTrait::PostEvaluate(FEvaluateTraversalContext& Context, const TTraitBinding<IEvaluate>& Binding) const
+	{
+		const FInstanceData* InstanceData = Binding.GetInstanceData<FInstanceData>();
+
+		// We override the default behavior since we need to blend over time
+
+		int32 NumBlending = 0;
+		for (const FBlendData& ChildBlendData : InstanceData->PerChildBlendData)
+		{
+			NumBlending += ChildBlendData.bIsBlending ? 1 : 0;
+		}
+
+		if (NumBlending < 2)
+		{
+			return;	// If we don't have at least 2 children blending, there is nothing to do
+		}
+
+		// Children are visited depth first, in the order returned
+		// As such, when we evaluate the task program, the keyframe of the last child will be
+		// on top of the keyframe stack
+		// We thus process children in reverse order
+
+		// The last child override the top keyframe and scales it
+		int32 ChildIndex = InstanceData->PerChildBlendData.Num() - 1;
+		for (; ChildIndex >= 0; --ChildIndex)
+		{
+			const FBlendData& ChildBlendData = InstanceData->PerChildBlendData[ChildIndex];
+			if (!ChildBlendData.bIsBlending)
+			{
+				continue;	// Skip this inactive child
+			}
+
+			// This trait controls the blend weight and owns it
+			Context.AppendTask(FAnimNextBlendOverwriteKeyframeWithScaleTask::Make(ChildBlendData.Weight));
+
+			// We found the last child to blend
+			break;
+		}
+
+		// Other children accumulate with scale
+		ChildIndex--;
+		for (; ChildIndex >= 0; --ChildIndex)
+		{
+			const FBlendData& ChildBlendData = InstanceData->PerChildBlendData[ChildIndex];
+			if (!ChildBlendData.bIsBlending)
+			{
+				continue;	// Skip this inactive child
+			}
+
+			// This trait controls the blend weight and owns it
+			Context.AppendTask(FAnimNextBlendAddKeyframeWithScaleTask::Make(ChildBlendData.Weight));
+		}
+
+		// Once we are done, we normalize rotations
+		Context.AppendTask(FAnimNextNormalizeKeyframeRotationsTask());
+	}
+
+	void FBlendSmootherTrait::PreUpdate(FUpdateTraversalContext& Context, const TTraitBinding<IUpdate>& Binding, const FTraitUpdateState& TraitState) const
+	{
+		const FSharedData* SharedData = Binding.GetSharedData<FSharedData>();
+		FInstanceData* InstanceData = Binding.GetInstanceData<FInstanceData>();
+
+		// If this is our first update, allocate our blend data
+		if (InstanceData->PerChildBlendData.IsEmpty())
+		{
+			InitializeInstanceData(Context, Binding, SharedData, InstanceData);
+		}
+
+		// Update the traits below us, they might trigger a transition
+		IUpdate::PreUpdate(Context, Binding, TraitState);
+
+		const float DeltaTime = TraitState.GetDeltaTime();
+
+		// Advance the weights
+		float SumWeight = 0.0f;
+		uint32 NumBlending = 0;
+
+		for (FBlendData& ChildBlendData : InstanceData->PerChildBlendData)
+		{
+			if (!ChildBlendData.bIsBlending)
+			{
+				continue;	// Skip children that aren't blending
+			}
+
+			ChildBlendData.Blend.Update(DeltaTime);
+
+			float NewBlendWeight = ChildBlendData.Blend.GetBlendedValue();
+
+			if (!FAnimWeight::IsRelevant(NewBlendWeight))
+			{
+				// Our new weight is no longer relevant, snap it to zero and normalization below will fix-up the other weights
+				// We'll then terminate the blend below
+				NewBlendWeight = 0.0f;
+			}
+
+			ChildBlendData.Weight = NewBlendWeight;
+			SumWeight += NewBlendWeight;
+			NumBlending++;
+		}
+
+		if (NumBlending <= 1)
+		{
+			return;	// Nothing to do if we don't blend at least 2 children together
+		}
+
+		// Renormalize the weights if the sum isn't near 0.0 or near 1.0
+		if (SumWeight > ZERO_ANIMWEIGHT_THRESH &&
+			FMath::Abs(SumWeight - 1.0f) > ZERO_ANIMWEIGHT_THRESH)
+		{
+			const float ReciprocalSum = 1.0f / SumWeight;
+
+			for (FBlendData& ChildBlendData : InstanceData->PerChildBlendData)
+			{
+				ChildBlendData.Weight *= ReciprocalSum;
+			}
+		}
+
+		TTraitBinding<IDiscreteBlend> DiscreteBlendTrait;
+		Binding.GetStackInterface(DiscreteBlendTrait);
+
+		// Free any newly inactive children
+		const int32 NumChildren = InstanceData->PerChildBlendData.Num();
+		for (int32 ChildIndex = 0; ChildIndex < NumChildren; ++ChildIndex)
+		{
+			FBlendData& ChildBlendData = InstanceData->PerChildBlendData[ChildIndex];
+
+			if (ChildBlendData.bIsBlending && ChildBlendData.Weight <= 0.0f)
+			{
+				// This child has finished blending out, terminate it
+				DiscreteBlendTrait.OnBlendTerminated(Context, ChildIndex);
+
+				ChildBlendData.bIsBlending = false;
+			}
+		}
+	}
+
+	float FBlendSmootherTrait::GetBlendWeight(const FExecutionContext& Context, const TTraitBinding<IDiscreteBlend>& Binding, int32 ChildIndex) const
+	{
+		const FInstanceData* InstanceData = Binding.GetInstanceData<FInstanceData>();
+		return InstanceData->PerChildBlendData.IsValidIndex(ChildIndex) ? InstanceData->PerChildBlendData[ChildIndex].Weight : -1.0f;
+	}
+
+	const FAlphaBlend* FBlendSmootherTrait::GetBlendState(const FExecutionContext& Context, const TTraitBinding<IDiscreteBlend>& Binding, int32 ChildIndex) const
+	{
+		const FInstanceData* InstanceData = Binding.GetInstanceData<FInstanceData>();
+		return InstanceData->PerChildBlendData.IsValidIndex(ChildIndex) ? &InstanceData->PerChildBlendData[ChildIndex].Blend : nullptr;
+	}
+
+	void FBlendSmootherTrait::OnBlendTransition(const FExecutionContext& Context, const TTraitBinding<IDiscreteBlend>& Binding, int32 OldChildIndex, int32 NewChildIndex) const
+	{
+		const FSharedData* SharedData = Binding.GetSharedData<FSharedData>();
+		FInstanceData* InstanceData = Binding.GetInstanceData<FInstanceData>();
+
+		const int32 NumChildren = InstanceData->PerChildBlendData.Num();
+		if (NewChildIndex >= NumChildren)
+		{
+			// We have a new child
+			check(NewChildIndex == NumChildren);
+
+			FBlendData& ChildBlendData = InstanceData->PerChildBlendData.AddDefaulted_GetRef();
+
+			ChildBlendData.Blend.SetBlendOption(SharedData->BlendType);
+			ChildBlendData.Blend.SetCustomCurve(SharedData->CustomBlendCurve);
+		}
+
+		// scale by the weight difference since we want consistency:
+		// - if you're moving from 0 to full weight 1, it will use the normal blend time
+		// - if you're moving from 0.5 to full weight 1, it will get there in half the time
+		const float NewChildCurrentWeight = InstanceData->PerChildBlendData[NewChildIndex].Weight;
+		const float NewChildDesiredWeight = 1.0f;
+		const float WeightDifference = FMath::Clamp(FMath::Abs(NewChildDesiredWeight - NewChildCurrentWeight), 0.0f, 1.0f);
+
+		TTraitBinding<ISmoothBlend> SmoothBlendTrait;
+		Binding.GetStackInterface(SmoothBlendTrait);
+
+		const float BlendTime = SmoothBlendTrait.GetBlendTime(Context, NewChildIndex);
+		const float RemainingBlendTime = OldChildIndex != INDEX_NONE ? (BlendTime * WeightDifference) : 0.0f;
+
+		if (OldChildIndex != INDEX_NONE)
+		{
+			// Make sure the old child starts blending out
+			FBlendData& OldChildBlendData = InstanceData->PerChildBlendData[OldChildIndex];
+			OldChildBlendData.Blend.SetValueRange(OldChildBlendData.Weight, 0.0f);
+			check(OldChildBlendData.bIsBlending);
+		}
+
+		{
+			// Setup the new child to blend in
+			FBlendData& NewChildBlendData = InstanceData->PerChildBlendData[NewChildIndex];
+			NewChildBlendData.Blend.SetValueRange(NewChildBlendData.Weight, 1.0f);
+			NewChildBlendData.Blend.ResetAlpha();	// Reset the alpha right away in case another trait needs it
+			NewChildBlendData.bIsBlending = true;
+		}
+
+		// We set the new blend time on all children
+		for (FBlendData& ChildBlendData : InstanceData->PerChildBlendData)
+		{
+			ChildBlendData.Blend.SetBlendTime(RemainingBlendTime);
+		}
+
+		// Don't call the super since we hijack the transition to smooth it out over time
+		// We just initiate the new blend manually
+
+		TTraitBinding<IDiscreteBlend> DiscreteBlendTrait;
+		Binding.GetStackInterface(DiscreteBlendTrait);
+
+		DiscreteBlendTrait.OnBlendInitiated(Context, NewChildIndex);
+	}
+
+	float FBlendSmootherTrait::GetBlendTime(const FExecutionContext& Context, const TTraitBinding<ISmoothBlend>& Binding, int32 ChildIndex) const
+	{
+		const FSharedData* SharedData = Binding.GetSharedData<FSharedData>();
+
+		if (SharedData->BlendTimes.IsValidIndex(ChildIndex))
+		{
+			return SharedData->BlendTimes[ChildIndex];
+		}
+		else if (!SharedData->BlendTimes.IsEmpty())
+		{
+			// If we index outside the array of values we have, use the last value
+			// Allows a user to specify a single blend time to be used with all children
+			return SharedData->BlendTimes.Last();
+		}
+		else
+		{
+			// No blend time has been specified, we snap
+			return 0.0f;
+		}
+	}
+
+	void FBlendSmootherTrait::InitializeInstanceData(const FExecutionContext& Context, const FTraitBinding& Binding, const FSharedData* SharedData, FInstanceData* InstanceData)
+	{
+		check(InstanceData->PerChildBlendData.IsEmpty());
+
+		const uint32 NumChildren = IHierarchy::GetNumStackChildren(Context, Binding);
+
+		InstanceData->PerChildBlendData.SetNum(NumChildren);
+
+		for (FBlendData& ChildBlendData : InstanceData->PerChildBlendData)
+		{
+			ChildBlendData.Blend.SetBlendOption(SharedData->BlendType);
+			ChildBlendData.Blend.SetCustomCurve(SharedData->CustomBlendCurve);
+		}
+	}
+}
diff --git a/Engine/Plugins/Experimental/AnimNextAnimGraph/Source/AnimNextAnimGraph/Private/Traits/BlendSmoother.h b/Engine/Plugins/Experimental/AnimNextAnimGraph/Source/AnimNextAnimGraph/Private/Traits/BlendSmoother.h
new file mode 100644
--- /dev/null
+++ b/Engine/Plugins/Experimental/AnimNextAnimGraph/Source/AnimNextAnimGraph/Private/Traits/BlendSmoother.h
@@ -0,0 +1,85 @@
+// Copyright Epic Games, Inc. All Rights Reserved.
+
+#pragma once
+
+#include "CoreMinimal.h"
+
+#include "AlphaBlend.h"
+#include "TraitCore/Trait.h"
+#include "TraitInterfaces/IDiscreteBlend.h"
+#include "TraitInterfaces/IEvaluate.h"
+#include "TraitInterfaces/ISmoothBlend.h"
+#include "TraitInterfaces/IUpdate.h"
+
+#include "BlendSmoother.generated.h"
+
+class UCurveFloat;
+
+USTRUCT(meta = (DisplayName = "Blend Smoother"))
+struct FAnimNextBlendSmootherTraitSharedData : public FAnimNextTraitSharedData
+{
+	GENERATED_BODY()
+
+	/** How long to take when blending into each child. */
+	UPROPERTY(EditAnywhere, Category = "Default", meta = (Inline))
+	TArray<float> BlendTimes;
+
+	/** What type of blend equation to use when converting the time elapsed into a blend weight. */
+	UPROPERTY(EditAnywhere, Category = "Default", meta = (Inline))
+	EAlphaBlendOption BlendType = EAlphaBlendOption::Linear;
+
+	/** Custom curve to use when the Custom blend type is used. */
+	UPROPERTY(EditAnywhere, Category = "Default", meta = (Inline))
+	TObjectPtr<UCurveFloat> CustomBlendCurve;
+};
+
+namespace UE::AnimNext
+{
+	/**
+	 * FBlendSmootherTrait
+	 * 
+	 * A trait that smoothly blends between discrete states over time.
+	 */
+	struct FBlendSmootherTrait : FAdditiveTrait, IEvaluate, IUpdate, IDiscreteBlend, ISmoothBlend
+	{
+		DECLARE_ANIM_TRAIT(FBlendSmootherTrait, 0xbaeb537b, FAdditiveTrait)
+
+		using FSharedData = FAnimNextBlendSmootherTraitSharedData;
+
+		// Struct for tracking blends for each pose
+		struct FBlendData
+		{
+			// Helper struct to update a time based weight
+			FAlphaBlend Blend;
+
+			// Current child weight (normalized with all children)
+			float Weight = 0.0f;
+
+			// Whether or not this child is actively blending
+			bool bIsBlending = false;
+		};
+
+		struct FInstanceData : FTrait::FInstanceData
+		{
+			// Blend state per child
+			TArray<FBlendData> PerChildBlendData;
+		};
+
+		// IEvaluate impl
+		virtual void PostEvaluate(FEvaluateTraversalContext& Context, const TTraitBinding<IEvaluate>& Binding) const override;
+
+		// IUpdate impl
+		virtual void PreUpdate(FUpdateTraversalContext& Context, const TTraitBinding<IUpdate>& Binding, const FTraitUpdateState& TraitState) const override;
+
+		// IDiscreteBlend impl
+		virtual float GetBlendWeight(const FExecutionContext& Context, const TTraitBinding<IDiscreteBlend>& Binding, int32 ChildIndex) const override;
+		virtual const FAlphaBlend* GetBlendState(const FExecutionContext& Context, const TTraitBinding<IDiscreteBlend>& Binding, int32 ChildIndex) const override;
+		virtual void OnBlendTransition(const FExecutionContext& Context, const TTraitBinding<IDiscreteBlend>& Binding, int32 OldChildIndex, int32 NewChildIndex) const override;
+
+		// ISmoothBlend impl
+		virtual float GetBlendTime(const FExecutionContext& Context, const TTraitBinding<ISmoothBlend>& Binding, int32 ChildIndex) const override;
+
+		// Internal impl
+		static void InitializeInstanceData(const FExecutionContext& Context, const FTraitBinding& Binding, const FSharedData* SharedData, FInstanceData* InstanceData);
+	};
+}
diff --git a/Engine/Plugins/Experimental/AnimNextAnimGraph/Source/AnimNextAnimGraph/Private/Traits/BlendSmootherPerBone.cpp b/Engine/Plugins/Experimental/AnimNextAnimGraph/Source/AnimNextAnimGraph/Private/Traits/BlendSmootherPerBone.cpp
new file mode 100644
--- /dev/null
+++ b/Engine/Plugins/Experimental/AnimNextAnimGraph/Source/AnimNextAnimGraph/Private/Traits/BlendSmootherPerBone.cpp
@@ -0,0 +1,174 @@
+// Copyright Epic Games, Inc. All Rights Reserved.
+
+#include "Traits/BlendSmootherPerBone.h"
+
+#include "AlphaBlend.h"
+#include "Animation/BlendProfile.h"
+#include "TraitCore/ExecutionContext.h"
+#include "TraitInterfaces/IHierarchy.h"
+#include "EvaluationVM/Tasks/BlendKeyframesPerBone.h"
+#include "EvaluationVM/Tasks/NormalizeRotations.h"
+
+namespace UE::AnimNext
+{
+	AUTO_REGISTER_ANIM_TRAIT(FBlendSmootherPerBoneTrait)
+
+	// Trait implementation boilerplate
+	#define TRAIT_INTERFACE_ENUMERATOR(GeneratorMacro) \
+		GeneratorMacro(IDiscreteBlend) \
+		GeneratorMacro(IEvaluate) \
+		GeneratorMacro(IUpdate) \
+
+	GENERATE_ANIM_TRAIT_IMPLEMENTATION(FBlendSmootherPerBoneTrait, TRAIT_INTERFACE_ENUMERATOR)
+	#undef TRAIT_INTERFACE_ENUMERATOR
+
+	void FBlendSmootherPerBoneTrait::PostEvaluate(FEvaluateTraversalContext& Context, const TTraitBinding<IEvaluate>& Binding) const
+	{
+		const FSharedData* SharedData = Binding.GetSharedData<FSharedData>();
+		const FInstanceData* InstanceData = Binding.GetInstanceData<FInstanceData>();
+
+		if (!SharedData->BlendProfile)
+		{
+			// No blend profile set, default smooth blend behavior
+			IEvaluate::PostEvaluate(Context, Binding);
+			return;
+		}
+
+		// We override the default behavior since we need to blend per bone
+
+		const int32 NumChildren = InstanceData->PerChildBlendData.Num();
+		if (NumChildren < 2)
+		{
+			return;	// If we don't have at least 2 children, there is nothing to do
+		}
+
+		// Children are visited depth first, in the order returned
+		// As such, when we evaluate the task program, the keyframe of the last child will be
+		// on top of the keyframe stack
+		// We thus process children in reverse order
+
+		// The last child override the top keyframe and scales it
+		{
+			const int32 ChildIndex = NumChildren - 1;
+			const FBlendSampleData& PoseSampleData = InstanceData->PerBoneSampleData[ChildIndex];
+
+			Context.AppendTask(FAnimNextBlendOverwriteKeyframePerBoneWithScaleTask::Make(SharedData->BlendProfile, PoseSampleData, PoseSampleData.TotalWeight));
+		}
+
+		// Other children accumulate with scale
+		for (int32 ChildIndex = NumChildren - 2; ChildIndex >= 0; --ChildIndex)
+		{
+			const FBlendSampleData& PoseSampleDataA = InstanceData->PerBoneSampleData[ChildIndex];
+			const FBlendSampleData& PoseSampleDataB = InstanceData->PerBoneSampleData[ChildIndex + 1];	// Above on the keyframe stack
+
+			Context.AppendTask(FAnimNextBlendAddKeyframePerBoneWithScaleTask::Make(SharedData->BlendProfile, PoseSampleDataA, PoseSampleDataB, PoseSampleDataA.TotalWeight));
+		}
+
+		// Once we are done, we normalize rotations
+		Context.AppendTask(FAnimNextNormalizeKeyframeRotationsTask());
+	}
+
+	void FBlendSmootherPerBoneTrait::PreUpdate(FUpdateTraversalContext& Context, const TTraitBinding<IUpdate>& Binding, const FTraitUpdateState& TraitState) const
+	{
+		const FSharedData* SharedData = Binding.GetSharedData<FSharedData>();
+		FInstanceData* InstanceData = Binding.GetInstanceData<FInstanceData>();
+
+		// If this is our first update, allocate our blend data
+		if (InstanceData->PerChildBlendData.IsEmpty())
+		{
+			InitializeInstanceData(Context, Binding, SharedData, InstanceData);
+		}
+
+		// Update the traits below us, they might trigger a transition
+		IUpdate::PreUpdate(Context, Binding, TraitState);
+
+		if (!SharedData->BlendProfile)
+		{
+			return;	// No blend profile set, nothing to do
+		}
+
+		TTraitBinding<IDiscreteBlend> DiscreteBlendTrait;
+		Binding.GetStackInterface(DiscreteBlendTrait);
+
+		const int32 DestinationChildIndex = DiscreteBlendTrait.GetBlendDestinationChildIndex(Context);
+
+		// If we're using a blend profile, extract the scales and build blend sample data
+		const int32 NumChildren = InstanceData->PerChildBlendData.Num();
+		for (int32 ChildIndex = 0; ChildIndex < NumChildren; ++ChildIndex)
+		{
+			const float BlendWeight = DiscreteBlendTrait.GetBlendWeight(Context, ChildIndex);
+			const FAlphaBlend* BlendState = DiscreteBlendTrait.GetBlendState(Context, ChildIndex);
+
+			FBlendSampleData& PoseSampleData = InstanceData->PerBoneSampleData[ChildIndex];
+			PoseSampleData.TotalWeight = BlendWeight;
+
+			const FBlendData& BlendData = InstanceData->PerChildBlendData[ChildIndex];
+			const bool bInverse = SharedData->BlendProfile->Mode == EBlendProfileMode::WeightFactor ? (DestinationChildIndex != ChildIndex) : false;
+			SharedData->BlendProfile->UpdateBoneWeights(PoseSampleData, *BlendState, BlendData.StartAlpha, BlendWeight, bInverse);
+		}
+
+		FBlendSampleData::NormalizeDataWeight(InstanceData->PerBoneSampleData);
+	}
+
+	void FBlendSmootherPerBoneTrait::OnBlendTransition(const FExecutionContext& Context, const TTraitBinding<IDiscreteBlend>& Binding, int32 OldChildIndex, int32 NewChildIndex) const
+	{
+		const FSharedData* SharedData = Binding.GetSharedData<FSharedData>();
+		FInstanceData* InstanceData = Binding.GetInstanceData<FInstanceData>();
+
+		// Trigger the new transition
+		IDiscreteBlend::OnBlendTransition(Context, Binding, OldChildIndex, NewChildIndex);
+
+		if (!SharedData->BlendProfile)
+		{
+			return;	// No blend profile set, nothing to do
+		}
+
+		const int32 NumChildren = InstanceData->PerChildBlendData.Num();
+		if (NewChildIndex >= NumChildren)
+		{
+			// We have a new child
+			check(NewChildIndex == NumChildren + 1);
+
+			InstanceData->PerChildBlendData.AddDefaulted();
+			FBlendSampleData& SampleData = InstanceData->PerBoneSampleData.AddDefaulted_GetRef();
+
+			const uint32 NumBlendEntries = SharedData->BlendProfile->GetNumBlendEntries();
+			SampleData.SampleDataIndex = NewChildIndex;
+			SampleData.PerBoneBlendData.AddZeroed(NumBlendEntries);
+		}
+
+		TTraitBinding<IDiscreteBlend> DiscreteBlendTrait;
+		Binding.GetStackInterface(DiscreteBlendTrait);
+
+		for (int32 ChildIndex = 0; ChildIndex < NumChildren; ++ChildIndex)
+		{
+			FBlendData& ChildBlendData = InstanceData->PerChildBlendData[ChildIndex];
+
+			const FAlphaBlend* BlendState = DiscreteBlendTrait.GetBlendState(Context, ChildIndex);
+			ChildBlendData.StartAlpha = BlendState->GetAlpha();
+		}
+	}
+
+	void FBlendSmootherPerBoneTrait::InitializeInstanceData(const FExecutionContext& Context, const FTraitBinding& Binding, const FSharedData* SharedData, FInstanceData* InstanceData)
+	{
+		check(InstanceData->PerChildBlendData.IsEmpty());
+
+		if (!SharedData->BlendProfile)
+		{
+			return;	// No blend profile set, nothing to do
+		}
+
+		const uint32 NumChildren = IHierarchy::GetNumStackChildren(Context, Binding);
+
+		InstanceData->PerChildBlendData.SetNum(NumChildren);
+		InstanceData->PerBoneSampleData.SetNum(NumChildren);
+
+		const uint32 NumBlendEntries = SharedData->BlendProfile->GetNumBlendEntries();
+		for (uint32 ChildIndex = 0; ChildIndex < NumChildren; ++ChildIndex)
+		{
+			FBlendSampleData& SampleData = InstanceData->PerBoneSampleData[ChildIndex];
+			SampleData.SampleDataIndex = ChildIndex;
+			SampleData.PerBoneBlendData.AddZeroed(NumBlendEntries);
+		}
+	}
+}
diff --git a/Engine/Plugins/Experimental/AnimNextAnimGraph/Source/AnimNextAnimGraph/Private/Traits/BlendSmootherPerBone.h b/Engine/Plugins/Experimental/AnimNextAnimGraph/Source/AnimNextAnimGraph/Private/Traits/BlendSmootherPerBone.h
new file mode 100644
--- /dev/null
+++ b/Engine/Plugins/Experimental/AnimNextAnimGraph/Source/AnimNextAnimGraph/Private/Traits/BlendSmootherPerBone.h
@@ -0,0 +1,70 @@
+// Copyright Epic Games, Inc. All Rights Reserved.
+
+#pragma once
+
+#include "CoreMinimal.h"
+
+#include "Animation/AnimationAsset.h"
+#include "Curves/CurveFloat.h"
+#include "TraitCore/Trait.h"
+#include "TraitInterfaces/IDiscreteBlend.h"
+#include "TraitInterfaces/IEvaluate.h"
+#include "TraitInterfaces/IUpdate.h"
+
+#include "BlendSmootherPerBone.generated.h"
+
+class UBlendProfile;
+
+USTRUCT(meta = (DisplayName = "Blend Smoother Per Bone"))
+struct FAnimNextBlendSmootherPerBoneTraitSharedData : public FAnimNextTraitSharedData
+{
+	GENERATED_BODY()
+
+	/** Blend profile that configures how fast to blend each bone. */
+	// TODO: Can't show list of blend profiles, we need to find a skeleton to perform the lookup with
+	UPROPERTY(EditAnywhere, Category = "Default", meta = (Inline, UseAsBlendProfile=true))
+	TObjectPtr<UBlendProfile> BlendProfile = nullptr;
+};
+
+namespace UE::AnimNext
+{
+	/**
+	 * FBlendSmootherTrait
+	 * 
+	 * A trait that smoothly blends between discrete states over time.
+	 */
+	struct FBlendSmootherPerBoneTrait : FAdditiveTrait, IEvaluate, IUpdate, IDiscreteBlend
+	{
+		DECLARE_ANIM_TRAIT(FBlendSmootherPerBoneTrait, 0xb97ffc16, FAdditiveTrait)
+
+		using FSharedData = FAnimNextBlendSmootherPerBoneTraitSharedData;
+
+		// Struct for tracking blends for each pose
+		struct FBlendData
+		{
+			// Which blend alpha we started the blend with
+			float StartAlpha = 0.0f;
+		};
+
+		struct FInstanceData : FTrait::FInstanceData
+		{
+			// Blend state per child
+			TArray<FBlendData> PerChildBlendData;
+
+			// Per-bone blending data for each child
+			TArray<FBlendSampleData> PerBoneSampleData;
+		};
+
+		// IEvaluate impl
+		virtual void PostEvaluate(FEvaluateTraversalContext& Context, const TTraitBinding<IEvaluate>& Binding) const override;
+
+		// IUpdate impl
+		virtual void PreUpdate(FUpdateTraversalContext& Context, const TTraitBinding<IUpdate>& Binding, const FTraitUpdateState& TraitState) const override;
+
+		// IDiscreteBlend impl
+		virtual void OnBlendTransition(const FExecutionContext& Context, const TTraitBinding<IDiscreteBlend>& Binding, int32 OldChildIndex, int32 NewChildIndex) const override;
+
+		// Internal impl
+		static void InitializeInstanceData(const FExecutionContext& Context, const FTraitBinding& Binding, const FSharedData* SharedData, FInstanceData* InstanceData);
+	};
+}
diff --git a/Engine/Plugins/Experimental/AnimNextAnimGraph/Source/AnimNextAnimGraph/Private/Traits/BlendTwoWay.cpp b/Engine/Plugins/Experimental/AnimNextAnimGraph/Source/AnimNextAnimGraph/Private/Traits/BlendTwoWay.cpp
new file mode 100644
--- /dev/null
+++ b/Engine/Plugins/Experimental/AnimNextAnimGraph/Source/AnimNextAnimGraph/Private/Traits/BlendTwoWay.cpp
@@ -0,0 +1,139 @@
+// Copyright Epic Games, Inc. All Rights Reserved.
+
+#include "Traits/BlendTwoWay.h"
+
+#include "Animation/AnimTypes.h"
+#include "TraitCore/ExecutionContext.h"
+#include "EvaluationVM/Tasks/BlendKeyframes.h"
+
+namespace UE::AnimNext
+{
+	AUTO_REGISTER_ANIM_TRAIT(FBlendTwoWayTrait)
+
+	// Trait implementation boilerplate
+	#define TRAIT_INTERFACE_ENUMERATOR(GeneratorMacro) \
+		GeneratorMacro(IContinuousBlend) \
+		GeneratorMacro(IEvaluate) \
+		GeneratorMacro(IHierarchy) \
+		GeneratorMacro(IUpdate) \
+		GeneratorMacro(IUpdateTraversal) \
+
+	GENERATE_ANIM_TRAIT_IMPLEMENTATION(FBlendTwoWayTrait, TRAIT_INTERFACE_ENUMERATOR)
+	#undef TRAIT_INTERFACE_ENUMERATOR
+
+	void FBlendTwoWayTrait::PostEvaluate(FEvaluateTraversalContext& Context, const TTraitBinding<IEvaluate>& Binding) const
+	{
+		const FInstanceData* InstanceData = Binding.GetInstanceData<FInstanceData>();
+
+		if (InstanceData->ChildA.IsValid() && InstanceData->ChildB.IsValid())
+		{
+			// We have two children, interpolate them
+
+			TTraitBinding<IContinuousBlend> ContinuousBlendTrait;
+			Binding.GetStackInterface(ContinuousBlendTrait);
+
+			const float BlendWeight = ContinuousBlendTrait.GetBlendWeight(Context, 1);
+			Context.AppendTask(FAnimNextBlendTwoKeyframesTask::Make(BlendWeight));
+		}
+		else
+		{
+			// We have only one child that is active, do nothing
+		}
+	}
+
+	void FBlendTwoWayTrait::PreUpdate(FUpdateTraversalContext& Context, const TTraitBinding<IUpdate>& Binding, const FTraitUpdateState& TraitState) const
+	{
+		const FSharedData* SharedData = Binding.GetSharedData<FSharedData>();
+		FInstanceData* InstanceData = Binding.GetInstanceData<FInstanceData>();
+
+		TTraitBinding<IContinuousBlend> ContinuousBlendTrait;
+		Binding.GetStackInterface(ContinuousBlendTrait);
+
+		const float BlendWeightB = ContinuousBlendTrait.GetBlendWeight(Context, 1);
+		if (!FAnimWeight::IsFullWeight(BlendWeightB))
+		{
+			if (!InstanceData->ChildA.IsValid())
+			{
+				// We need to blend a child that isn't instanced yet, allocate it
+				InstanceData->ChildA = Context.AllocateNodeInstance(Binding, SharedData->ChildA);
+			}
+
+			if (!FAnimWeight::IsRelevant(BlendWeightB))
+			{
+				// We no longer need this child, release it
+				InstanceData->ChildB.Reset();
+			}
+		}
+
+		if (FAnimWeight::IsRelevant(BlendWeightB))
+		{
+			if (!InstanceData->ChildB.IsValid())
+			{
+				// We need to blend a child that isn't instanced yet, allocate it
+				InstanceData->ChildB = Context.AllocateNodeInstance(Binding, SharedData->ChildB);
+			}
+
+			if (FAnimWeight::IsFullWeight(BlendWeightB))
+			{
+				// We no longer need this child, release it
+				InstanceData->ChildA.Reset();
+			}
+		}
+	}
+
+	void FBlendTwoWayTrait::QueueChildrenForTraversal(FUpdateTraversalContext& Context, const TTraitBinding<IUpdateTraversal>& Binding, const FTraitUpdateState& TraitState, FUpdateTraversalQueue& TraversalQueue) const
+	{
+		const FInstanceData* InstanceData = Binding.GetInstanceData<FInstanceData>();
+
+		TTraitBinding<IContinuousBlend> ContinuousBlendTrait;
+		Binding.GetStackInterface(ContinuousBlendTrait);
+
+		const float BlendWeightB = ContinuousBlendTrait.GetBlendWeight(Context, 1);
+		if (InstanceData->ChildA.IsValid())
+		{
+			const float BlendWeightA = 1.0f - BlendWeightB;
+			TraversalQueue.Push(InstanceData->ChildA, TraitState.WithWeight(BlendWeightA));
+		}
+
+		if (InstanceData->ChildB.IsValid())
+		{
+			TraversalQueue.Push(InstanceData->ChildB, TraitState.WithWeight(BlendWeightB));
+		}
+	}
+
+	uint32 FBlendTwoWayTrait::GetNumChildren(const FExecutionContext& Context, const TTraitBinding<IHierarchy>& Binding) const
+	{
+		return 2;
+	}
+
+	void FBlendTwoWayTrait::GetChildren(const FExecutionContext& Context, const TTraitBinding<IHierarchy>& Binding, FChildrenArray& Children) const
+	{
+		const FInstanceData* InstanceData = Binding.GetInstanceData<FInstanceData>();
+
+		// Add the two children, even if the handles are empty
+		Children.Add(InstanceData->ChildA);
+		Children.Add(InstanceData->ChildB);
+	}
+
+	float FBlendTwoWayTrait::GetBlendWeight(const FExecutionContext& Context, const TTraitBinding<IContinuousBlend>& Binding, int32 ChildIndex) const
+	{
+		const FSharedData* SharedData = Binding.GetSharedData<FSharedData>();
+
+		const float BlendWeight = SharedData->GetBlendWeight(Binding);
+		const float ClampedWeight = FMath::Clamp(BlendWeight, 0.0f, 1.0f);
+
+		if (ChildIndex == 0)
+		{
+			return 1.0f - ClampedWeight;
+		}
+		else if (ChildIndex == 1)
+		{
+			return ClampedWeight;
+		}
+		else
+		{
+			// Invalid child index
+			return -1.0f;
+		}
+	}
+}
diff --git a/Engine/Plugins/Experimental/AnimNextAnimGraph/Source/AnimNextAnimGraph/Private/Traits/BlendTwoWay.h b/Engine/Plugins/Experimental/AnimNextAnimGraph/Source/AnimNextAnimGraph/Private/Traits/BlendTwoWay.h
new file mode 100644
--- /dev/null
+++ b/Engine/Plugins/Experimental/AnimNextAnimGraph/Source/AnimNextAnimGraph/Private/Traits/BlendTwoWay.h
@@ -0,0 +1,75 @@
+// Copyright Epic Games, Inc. All Rights Reserved.
+
+#pragma once
+
+#include "CoreMinimal.h"
+
+#include "TraitCore/Trait.h"
+#include "TraitInterfaces/IContinuousBlend.h"
+#include "TraitInterfaces/IEvaluate.h"
+#include "TraitInterfaces/IHierarchy.h"
+#include "TraitInterfaces/IUpdate.h"
+
+#include "BlendTwoWay.generated.h"
+
+USTRUCT(meta = (DisplayName = "Blend Two Way"))
+struct FAnimNextBlendTwoWayTraitSharedData : public FAnimNextTraitSharedData
+{
+	GENERATED_BODY()
+
+	/** First output to be blended (full weight is 0.0). */
+	UPROPERTY()
+	FAnimNextTraitHandle ChildA;
+
+	/** Second output to be blended (full weight is 1.0). */
+	UPROPERTY()
+	FAnimNextTraitHandle ChildB;
+
+	/** How much to blend our two children: 0.0 is fully child A while 1.0 is fully child B. */
+	UPROPERTY(EditAnywhere, Category = "Default")
+	float BlendWeight = 0.0f;
+
+	// Latent pin support boilerplate
+	#define TRAIT_LATENT_PROPERTIES_ENUMERATOR(GeneratorMacro) \
+		GeneratorMacro(BlendWeight) \
+
+	GENERATE_TRAIT_LATENT_PROPERTIES(FAnimNextBlendTwoWayTraitSharedData, TRAIT_LATENT_PROPERTIES_ENUMERATOR)
+	#undef TRAIT_LATENT_PROPERTIES_ENUMERATOR
+};
+
+namespace UE::AnimNext
+{
+	/**
+	 * FBlendTwoWayTrait
+	 * 
+	 * A trait that can blend two inputs.
+	 */
+	struct FBlendTwoWayTrait : FBaseTrait, IEvaluate, IUpdate, IUpdateTraversal, IHierarchy, IContinuousBlend
+	{
+		DECLARE_ANIM_TRAIT(FBlendTwoWayTrait, 0x62f40e2b, FBaseTrait)
+
+		using FSharedData = FAnimNextBlendTwoWayTraitSharedData;
+
+		struct FInstanceData : FTrait::FInstanceData
+		{
+			FTraitPtr ChildA;
+			FTraitPtr ChildB;
+		};
+
+		// IEvaluate impl
+		virtual void PostEvaluate(FEvaluateTraversalContext& Context, const TTraitBinding<IEvaluate>& Binding) const override;
+
+		// IUpdate impl
+		virtual void PreUpdate(FUpdateTraversalContext& Context, const TTraitBinding<IUpdate>& Binding, const FTraitUpdateState& TraitState) const override;
+
+		// IUpdateTraversal impl
+		virtual void QueueChildrenForTraversal(FUpdateTraversalContext& Context, const TTraitBinding<IUpdateTraversal>& Binding, const FTraitUpdateState& TraitState, FUpdateTraversalQueue& TraversalQueue) const override;
+
+		// IHierarchy impl
+		virtual uint32 GetNumChildren(const FExecutionContext& Context, const TTraitBinding<IHierarchy>& Binding) const override;
+		virtual void GetChildren(const FExecutionContext& Context, const TTraitBinding<IHierarchy>& Binding, FChildrenArray& Children) const override;
+
+		// IContinuousBlend impl
+		virtual float GetBlendWeight(const FExecutionContext& Context, const TTraitBinding<IContinuousBlend>& Binding, int32 ChildIndex) const override;
+	};
+}
diff --git a/Engine/Plugins/Experimental/AnimNextAnimGraph/Source/AnimNextAnimGraph/Private/Traits/ReferencePoseTrait.cpp b/Engine/Plugins/Experimental/AnimNextAnimGraph/Source/AnimNextAnimGraph/Private/Traits/ReferencePoseTrait.cpp
new file mode 100644
--- /dev/null
+++ b/Engine/Plugins/Experimental/AnimNextAnimGraph/Source/AnimNextAnimGraph/Private/Traits/ReferencePoseTrait.cpp
@@ -0,0 +1,28 @@
+// Copyright Epic Games, Inc. All Rights Reserved.
+
+#include "Traits/ReferencePoseTrait.h"
+
+#include "TraitCore/ExecutionContext.h"
+#include "EvaluationVM/Tasks/PushReferenceKeyframe.h"
+
+namespace UE::AnimNext
+{
+	AUTO_REGISTER_ANIM_TRAIT(FReferencePoseTrait)
+
+	// Trait implementation boilerplate
+	#define TRAIT_INTERFACE_ENUMERATOR(GeneratorMacro) \
+		GeneratorMacro(IEvaluate) \
+
+	GENERATE_ANIM_TRAIT_IMPLEMENTATION(FReferencePoseTrait, TRAIT_INTERFACE_ENUMERATOR)
+	#undef TRAIT_INTERFACE_ENUMERATOR
+
+	void FReferencePoseTrait::PreEvaluate(FEvaluateTraversalContext& Context, const TTraitBinding<IEvaluate>& Binding) const
+	{
+		const FSharedData* SharedData = Binding.GetSharedData<FSharedData>();
+
+		FAnimNextPushReferenceKeyframeTask Task;
+		Task.bIsAdditive = SharedData->ReferencePoseType == EAnimNextReferencePoseType::AdditiveIdentity;
+
+		Context.AppendTask(Task);
+	}
+}
diff --git a/Engine/Plugins/Experimental/AnimNextAnimGraph/Source/AnimNextAnimGraph/Private/Traits/ReferencePoseTrait.h b/Engine/Plugins/Experimental/AnimNextAnimGraph/Source/AnimNextAnimGraph/Private/Traits/ReferencePoseTrait.h
new file mode 100644
--- /dev/null
+++ b/Engine/Plugins/Experimental/AnimNextAnimGraph/Source/AnimNextAnimGraph/Private/Traits/ReferencePoseTrait.h
@@ -0,0 +1,50 @@
+// Copyright Epic Games, Inc. All Rights Reserved.
+
+#pragma once
+
+#include "CoreMinimal.h"
+#include "TraitCore/Trait.h"
+#include "TraitInterfaces/IEvaluate.h"
+
+#include "ReferencePoseTrait.generated.h"
+
+// TODO: Ideally, the reference pose we output should be a tag and the task that consumes the reference pose should
+// be the one to determine whether it should be in local space or the additive identity (or something else).
+// This way, we have a single option in the graph, and the system can figure out what to use removing the room the user mistakes.
+// It will also be more efficient since we don't need to manipulate the reference pose until the point of use where we might
+// be able to avoid the copy.
+
+UENUM()
+enum class EAnimNextReferencePoseType : int32
+{
+	MeshLocalSpace,
+	AdditiveIdentity,
+};
+
+USTRUCT(meta = (DisplayName = "Reference Pose"))
+struct FAnimNextReferencePoseTraitSharedData : public FAnimNextTraitSharedData
+{
+	GENERATED_BODY()
+
+	/** The type of the reference pose. */
+	UPROPERTY(EditAnywhere, Category = "Default", meta = (Inline))
+	EAnimNextReferencePoseType ReferencePoseType = EAnimNextReferencePoseType::MeshLocalSpace;
+};
+
+namespace UE::AnimNext
+{
+	/**
+	 * FReferencePoseTrait
+	 * 
+	 * A trait that outputs a reference pose.
+	 */
+	struct FReferencePoseTrait : FBaseTrait, IEvaluate
+	{
+		DECLARE_ANIM_TRAIT(FReferencePoseTrait, 0x7508ab89, FBaseTrait)
+
+		using FSharedData = FAnimNextReferencePoseTraitSharedData;
+
+		// IEvaluate impl
+		virtual void PreEvaluate(FEvaluateTraversalContext& Context, const TTraitBinding<IEvaluate>& Binding) const override;
+	};
+}
diff --git a/Engine/Plugins/Experimental/AnimNextAnimGraph/Source/AnimNextAnimGraph/Private/Traits/SequencePlayer.cpp b/Engine/Plugins/Experimental/AnimNextAnimGraph/Source/AnimNextAnimGraph/Private/Traits/SequencePlayer.cpp
new file mode 100644
--- /dev/null
+++ b/Engine/Plugins/Experimental/AnimNextAnimGraph/Source/AnimNextAnimGraph/Private/Traits/SequencePlayer.cpp
@@ -0,0 +1,94 @@
+// Copyright Epic Games, Inc. All Rights Reserved.
+
+#include "Traits/SequencePlayer.h"
+
+#include "AnimationRuntime.h"
+#include "TraitCore/ExecutionContext.h"
+#include "EvaluationVM/Tasks/PushAnimSequenceKeyframe.h"
+
+namespace UE::AnimNext
+{
+	AUTO_REGISTER_ANIM_TRAIT(FSequencePlayerTrait)
+
+	// Trait implementation boilerplate
+	#define TRAIT_INTERFACE_ENUMERATOR(GeneratorMacro) \
+		GeneratorMacro(IEvaluate) \
+		GeneratorMacro(ITimeline) \
+		GeneratorMacro(IUpdate) \
+
+	GENERATE_ANIM_TRAIT_IMPLEMENTATION(FSequencePlayerTrait, TRAIT_INTERFACE_ENUMERATOR)
+	#undef TRAIT_INTERFACE_ENUMERATOR
+
+	void FSequencePlayerTrait::FInstanceData::Construct(const FExecutionContext& Context, const FTraitBinding& Binding)
+	{
+		const FSharedData* SharedData = Binding.GetSharedData<FSharedData>();
+		if (SharedData->AnimSequence != nullptr)
+		{
+			const float SequenceLength = SharedData->AnimSequence->GetPlayLength();
+			InternalTimeAccumulator = FMath::Clamp(SharedData->GetStartPosition(Binding), 0.0f, SequenceLength);
+		}
+	}
+
+	void FSequencePlayerTrait::PreEvaluate(FEvaluateTraversalContext& Context, const TTraitBinding<IEvaluate>& Binding) const
+	{
+		const FSharedData* SharedData = Binding.GetSharedData<FSharedData>();
+
+		const FInstanceData* InstanceData = Binding.GetInstanceData<FInstanceData>();
+		const bool bInterpolate = true;
+
+		FAnimNextAnimSequenceKeyframeTask Task = FAnimNextAnimSequenceKeyframeTask::MakeFromSampleTime(SharedData->AnimSequence, InstanceData->InternalTimeAccumulator, bInterpolate);
+		Task.bExtractTrajectory = true;	/*Output.AnimInstanceProxy->ShouldExtractRootMotion()*/
+
+		Context.AppendTask(Task);
+	}
+
+	float FSequencePlayerTrait::GetPlayRate(const FExecutionContext& Context, const TTraitBinding<ITimeline>& Binding) const
+	{
+		const FSharedData* SharedData = Binding.GetSharedData<FSharedData>();
+		return SharedData->GetPlayRate(Binding);
+	}
+
+	float FSequencePlayerTrait::AdvanceBy(const FExecutionContext& Context, const TTraitBinding<ITimeline>& Binding, float DeltaTime) const
+	{
+		const FSharedData* SharedData = Binding.GetSharedData<FSharedData>();
+		if (UAnimSequence* AnimSeq = SharedData->AnimSequence.Get())
+		{
+			FInstanceData* InstanceData = Binding.GetInstanceData<FInstanceData>();
+
+			TTraitBinding<ITimeline> TimelineTrait;
+			Binding.GetStackInterface(TimelineTrait);
+
+			const float PlayRate = TimelineTrait.GetPlayRate(Context);
+			const bool bIsLooping = SharedData->GetbLoop(Binding);
+			const float SequenceLength = AnimSeq->GetPlayLength();
+
+			FAnimationRuntime::AdvanceTime(bIsLooping, DeltaTime * PlayRate, InstanceData->InternalTimeAccumulator, SequenceLength);
+
+			return FMath::Clamp(InstanceData->InternalTimeAccumulator / SequenceLength, 0.0f, 1.0f);
+		}
+
+		return 0.0f;
+	}
+
+	void FSequencePlayerTrait::AdvanceToRatio(const FExecutionContext& Context, const TTraitBinding<ITimeline>& Binding, float ProgressRatio) const
+	{
+		const FSharedData* SharedData = Binding.GetSharedData<FSharedData>();
+		if (UAnimSequence* AnimSeq = SharedData->AnimSequence.Get())
+		{
+			FInstanceData* InstanceData = Binding.GetInstanceData<FInstanceData>();
+
+			const float SequenceLength = AnimSeq->GetPlayLength();
+
+			InstanceData->InternalTimeAccumulator = FMath::Clamp(ProgressRatio, 0.0f, 1.0f) * SequenceLength;
+		}
+	}
+
+	void FSequencePlayerTrait::PreUpdate(FUpdateTraversalContext& Context, const TTraitBinding<IUpdate>& Binding, const FTraitUpdateState& TraitState) const
+	{
+		// We just advance the timeline
+		TTraitBinding<ITimeline> TimelineTrait;
+		Binding.GetStackInterface(TimelineTrait);
+
+		TimelineTrait.AdvanceBy(Context, TraitState.GetDeltaTime());
+	}
+}
diff --git a/Engine/Plugins/Experimental/AnimNextAnimGraph/Source/AnimNextAnimGraph/Private/Traits/SequencePlayer.h b/Engine/Plugins/Experimental/AnimNextAnimGraph/Source/AnimNextAnimGraph/Private/Traits/SequencePlayer.h
new file mode 100644
--- /dev/null
+++ b/Engine/Plugins/Experimental/AnimNextAnimGraph/Source/AnimNextAnimGraph/Private/Traits/SequencePlayer.h
@@ -0,0 +1,76 @@
+// Copyright Epic Games, Inc. All Rights Reserved.
+
+#pragma once
+
+#include "CoreMinimal.h"
+#include "TraitCore/Trait.h"
+#include "TraitInterfaces/IEvaluate.h"
+#include "TraitInterfaces/ITimeline.h"
+#include "TraitInterfaces/IUpdate.h"
+#include "Animation/AnimSequence.h"
+
+#include "SequencePlayer.generated.h"
+
+USTRUCT(meta = (DisplayName = "Sequence Player"))
+struct FAnimNextSequencePlayerTraitSharedData : public FAnimNextTraitSharedData
+{
+	GENERATED_BODY()
+
+	/** The sequence to play. */
+	UPROPERTY(EditAnywhere, Category = "Default", meta = (Inline))
+	TObjectPtr<UAnimSequence> AnimSequence;
+
+	/** The play rate multiplier at which this sequence plays. */
+	UPROPERTY(EditAnywhere, Category = "Default")
+	float PlayRate = 1.0f;
+
+	/** The time at which we should start playing this sequence. */
+	UPROPERTY(EditAnywhere, Category = "Default")
+	float StartPosition = 0.0f;
+
+	/** Whether or not this sequence playback will loop. */
+	UPROPERTY(EditAnywhere, Category = "Default")
+	bool bLoop = 0.0f;
+
+	// Latent pin support boilerplate
+	#define TRAIT_LATENT_PROPERTIES_ENUMERATOR(GeneratorMacro) \
+		GeneratorMacro(PlayRate) \
+		GeneratorMacro(StartPosition) \
+		GeneratorMacro(bLoop) \
+
+	GENERATE_TRAIT_LATENT_PROPERTIES(FAnimNextSequencePlayerTraitSharedData, TRAIT_LATENT_PROPERTIES_ENUMERATOR)
+	#undef TRAIT_LATENT_PROPERTIES_ENUMERATOR
+};
+
+namespace UE::AnimNext
+{
+	/**
+	 * FSequencePlayerTrait
+	 * 
+	 * A trait that can play an animation sequence.
+	 */
+	struct FSequencePlayerTrait : FBaseTrait, IEvaluate, ITimeline, IUpdate
+	{
+		DECLARE_ANIM_TRAIT(FSequencePlayerTrait, 0x7a0dc157, FBaseTrait)
+
+		using FSharedData = FAnimNextSequencePlayerTraitSharedData;
+
+		struct FInstanceData : FTrait::FInstanceData
+		{
+			float InternalTimeAccumulator = 0.0f;
+
+			void Construct(const FExecutionContext& Context, const FTraitBinding& Binding);
+		};
+
+		// IEvaluate impl
+		virtual void PreEvaluate(FEvaluateTraversalContext& Context, const TTraitBinding<IEvaluate>& Binding) const override;
+
+		// ITimeline impl
+		virtual float GetPlayRate(const FExecutionContext& Context, const TTraitBinding<ITimeline>& Binding) const override;
+		virtual float AdvanceBy(const FExecutionContext& Context, const TTraitBinding<ITimeline>& Binding, float DeltaTime) const override;
+		virtual void AdvanceToRatio(const FExecutionContext& Context, const TTraitBinding<ITimeline>& Binding, float ProgressRatio) const override;
+
+		// IUpdate impl
+		virtual void PreUpdate(FUpdateTraversalContext& Context, const TTraitBinding<IUpdate>& Binding, const FTraitUpdateState& TraitState) const override;
+	};
+}
diff --git a/Engine/Plugins/Experimental/AnimNextAnimGraph/Source/AnimNextAnimGraph/Private/Traits/SubGraphHost.cpp b/Engine/Plugins/Experimental/AnimNextAnimGraph/Source/AnimNextAnimGraph/Private/Traits/SubGraphHost.cpp
new file mode 100644
--- /dev/null
+++ b/Engine/Plugins/Experimental/AnimNextAnimGraph/Source/AnimNextAnimGraph/Private/Traits/SubGraphHost.cpp
@@ -0,0 +1,243 @@
+// Copyright Epic Games, Inc. All Rights Reserved.
+
+#include "Traits/SubGraphHost.h"
+
+#include "TraitCore/ExecutionContext.h"
+#include "TraitCore/NodeInstance.h"
+#include "Graph/AnimNextGraphInstance.h"
+
+namespace UE::AnimNext
+{
+	AUTO_REGISTER_ANIM_TRAIT(FSubGraphHostTrait)
+
+	// Trait implementation boilerplate
+	#define TRAIT_INTERFACE_ENUMERATOR(GeneratorMacro) \
+		GeneratorMacro(IDiscreteBlend) \
+		GeneratorMacro(IGarbageCollection) \
+		GeneratorMacro(IHierarchy) \
+		GeneratorMacro(IUpdate) \
+		GeneratorMacro(IUpdateTraversal) \
+
+	GENERATE_ANIM_TRAIT_IMPLEMENTATION(FSubGraphHostTrait, TRAIT_INTERFACE_ENUMERATOR)
+	#undef TRAIT_INTERFACE_ENUMERATOR
+
+	void FSubGraphHostTrait::FInstanceData::Construct(const FExecutionContext& Context, const FTraitBinding& Binding)
+	{
+		FTrait::FInstanceData::Construct(Context, Binding);
+
+		IGarbageCollection::RegisterWithGC(Context, Binding);
+
+		const FSharedData* SharedData = Binding.GetSharedData<FSharedData>();
+		ReferencePoseChildPtr = Context.AllocateNodeInstance(Binding, SharedData->ReferencePoseChild);
+	}
+
+	void FSubGraphHostTrait::FInstanceData::Destruct(const FExecutionContext& Context, const FTraitBinding& Binding)
+	{
+		FTrait::FInstanceData::Destruct(Context, Binding);
+
+		IGarbageCollection::UnregisterWithGC(Context, Binding);
+	}
+
+	uint32 FSubGraphHostTrait::GetNumChildren(const FExecutionContext& Context, const TTraitBinding<IHierarchy>& Binding) const
+	{
+		const FInstanceData* InstanceData = Binding.GetInstanceData<FInstanceData>();
+		return InstanceData->SubGraphSlots.Num();
+	}
+
+	void FSubGraphHostTrait::GetChildren(const FExecutionContext& Context, const TTraitBinding<IHierarchy>& Binding, FChildrenArray& Children) const
+	{
+		const FInstanceData* InstanceData = Binding.GetInstanceData<FInstanceData>();
+
+		for (const FSubGraphSlot& SubGraphEntry : InstanceData->SubGraphSlots)
+		{
+			if (SubGraphEntry.State == ESlotState::ActiveWithReferencePose)
+			{
+				Children.Add(InstanceData->ReferencePoseChildPtr);
+			}
+			else
+			{
+				// Even if the slot is inactive, we queue an empty handle
+				Children.Add(SubGraphEntry.GraphInstance.GetGraphRootPtr());
+			}
+		}
+	}
+
+	void FSubGraphHostTrait::PreUpdate(FUpdateTraversalContext& Context, const TTraitBinding<IUpdate>& Binding, const FTraitUpdateState& TraitState) const
+	{
+		const FSharedData* SharedData = Binding.GetSharedData<FSharedData>();
+		FInstanceData* InstanceData = Binding.GetInstanceData<FInstanceData>();
+
+		const bool bHasActiveSubGraph = InstanceData->CurrentlyActiveSubGraphIndex != INDEX_NONE;
+
+		TObjectPtr<const UAnimNextGraph> CurrentActiveSubGraph;
+		FName CurrentActiveEntryPoint = NAME_None;
+		if (bHasActiveSubGraph)
+		{
+			const FSubGraphSlot& SubGraphSlot = InstanceData->SubGraphSlots[InstanceData->CurrentlyActiveSubGraphIndex];
+			CurrentActiveSubGraph = SubGraphSlot.SubGraph;
+			CurrentActiveEntryPoint = SubGraphSlot.EntryPoint;
+		}
+
+		const TObjectPtr<const UAnimNextGraph> DesiredSubGraph = SharedData->GetSubGraph(Binding);
+		const FName EntryPoint = SharedData->GetEntryPoint(Binding);
+
+		// Check for re-entrancy and early-out if we are linking back to the current instance or one of its parents
+		const FAnimNextGraphInstance* OwnerGraphInstance = &Binding.GetTraitPtr().GetNodeInstance()->GetOwner();
+		while (OwnerGraphInstance != nullptr)
+		{
+			if (OwnerGraphInstance->UsesGraph(DesiredSubGraph) && OwnerGraphInstance->UsesEntryPoint(EntryPoint))
+			{
+				return;
+			}
+
+			OwnerGraphInstance = OwnerGraphInstance->GetParentGraphInstance();
+		}
+
+		if (!bHasActiveSubGraph || CurrentActiveSubGraph != DesiredSubGraph || CurrentActiveEntryPoint != EntryPoint)
+		{
+			// Find an empty slot we can use
+			int32 FreeSlotIndex = INDEX_NONE;
+
+			const int32 NumSubGraphSlots = InstanceData->SubGraphSlots.Num();
+			for (int32 SlotIndex = 0; SlotIndex < NumSubGraphSlots; ++SlotIndex)
+			{
+				if (InstanceData->SubGraphSlots[SlotIndex].State == ESlotState::Inactive)
+				{
+					// This slot is inactive, we can re-use it
+					FreeSlotIndex = SlotIndex;
+					break;
+				}
+			}
+
+			if (FreeSlotIndex == INDEX_NONE)
+			{
+				// All slots are in use, add a new one
+				FreeSlotIndex = InstanceData->SubGraphSlots.AddDefaulted();
+			}
+
+			FSubGraphSlot& SubGraphSlot = InstanceData->SubGraphSlots[FreeSlotIndex];
+			SubGraphSlot.SubGraph = DesiredSubGraph;
+			SubGraphSlot.State = DesiredSubGraph ? ESlotState::ActiveWithGraph : ESlotState::ActiveWithReferencePose;
+			SubGraphSlot.EntryPoint = EntryPoint;
+
+			const int32 OldChildIndex = InstanceData->CurrentlyActiveSubGraphIndex;
+			const int32 NewChildIndex = FreeSlotIndex;
+
+			InstanceData->CurrentlyActiveSubGraphIndex = FreeSlotIndex;
+
+			TTraitBinding<IDiscreteBlend> DiscreteBlendTrait;
+			Binding.GetStackInterface(DiscreteBlendTrait);
+
+			DiscreteBlendTrait.OnBlendTransition(Context, OldChildIndex, NewChildIndex);
+		}
+	}
+
+	void FSubGraphHostTrait::QueueChildrenForTraversal(FUpdateTraversalContext& Context, const TTraitBinding<IUpdateTraversal>& Binding, const FTraitUpdateState& TraitState, FUpdateTraversalQueue& TraversalQueue) const
+	{
+		const FInstanceData* InstanceData = Binding.GetInstanceData<FInstanceData>();
+
+		const int32 NumSubGraphs = InstanceData->SubGraphSlots.Num();
+		if (NumSubGraphs == 0)
+		{
+			return;
+		}
+
+		TTraitBinding<IDiscreteBlend> DiscreteBlendTrait;
+		Binding.GetStackInterface(DiscreteBlendTrait);
+
+		for (int32 SubGraphIndex = 0; SubGraphIndex < NumSubGraphs; ++SubGraphIndex)
+		{
+			const float BlendWeight = DiscreteBlendTrait.GetBlendWeight(Context, SubGraphIndex);
+
+			FTraitUpdateState SubGraphTraitState = TraitState.WithWeight(BlendWeight);
+			if (SubGraphIndex != InstanceData->CurrentlyActiveSubGraphIndex)
+			{
+				SubGraphTraitState = SubGraphTraitState.AsBlendingOut();
+			}
+
+			TraversalQueue.Push(InstanceData->SubGraphSlots[SubGraphIndex].GraphInstance.GetGraphRootPtr(), SubGraphTraitState);
+		}
+	}
+
+	float FSubGraphHostTrait::GetBlendWeight(const FExecutionContext& Context, const TTraitBinding<IDiscreteBlend>& Binding, int32 ChildIndex) const
+	{
+		const FInstanceData* InstanceData = Binding.GetInstanceData<FInstanceData>();
+
+		if (ChildIndex == InstanceData->CurrentlyActiveSubGraphIndex)
+		{
+			return 1.0f;	// Active child has full weight
+		}
+		else if (InstanceData->SubGraphSlots.IsValidIndex(ChildIndex))
+		{
+			return 0.0f;	// Other children have no weight
+		}
+		else
+		{
+			// Invalid child index
+			return -1.0f;
+		}
+	}
+
+	int32 FSubGraphHostTrait::GetBlendDestinationChildIndex(const FExecutionContext& Context, const TTraitBinding<IDiscreteBlend>& Binding) const
+	{
+		const FInstanceData* InstanceData = Binding.GetInstanceData<FInstanceData>();
+
+		return InstanceData->CurrentlyActiveSubGraphIndex;
+	}
+
+	void FSubGraphHostTrait::OnBlendTransition(const FExecutionContext& Context, const TTraitBinding<IDiscreteBlend>& Binding, int32 OldChildIndex, int32 NewChildIndex) const
+	{
+		TTraitBinding<IDiscreteBlend> DiscreteBlendTrait;
+		Binding.GetStackInterface(DiscreteBlendTrait);
+
+		// We initiate immediately when we transition
+		DiscreteBlendTrait.OnBlendInitiated(Context, NewChildIndex);
+
+		// We terminate immediately when we transition
+		DiscreteBlendTrait.OnBlendTerminated(Context, OldChildIndex);
+	}
+
+	void FSubGraphHostTrait::OnBlendInitiated(const FExecutionContext& Context, const TTraitBinding<IDiscreteBlend>& Binding, int32 ChildIndex) const
+	{
+		FInstanceData* InstanceData = Binding.GetInstanceData<FInstanceData>();
+
+		if (InstanceData->SubGraphSlots.IsValidIndex(ChildIndex))
+		{
+			// Allocate our new sub-graph instance
+			FSubGraphSlot& SubGraphEntry = InstanceData->SubGraphSlots[ChildIndex];
+
+			if (SubGraphEntry.State == ESlotState::ActiveWithGraph)
+			{
+				SubGraphEntry.SubGraph->AllocateInstance(Binding.GetTraitPtr().GetNodeInstance()->GetOwner(), SubGraphEntry.GraphInstance, SubGraphEntry.EntryPoint);
+			}
+		}
+	}
+
+	void FSubGraphHostTrait::OnBlendTerminated(const FExecutionContext& Context, const TTraitBinding<IDiscreteBlend>& Binding, int32 ChildIndex) const
+	{
+		FInstanceData* InstanceData = Binding.GetInstanceData<FInstanceData>();
+
+		if (InstanceData->SubGraphSlots.IsValidIndex(ChildIndex))
+		{
+			// Deallocate our sub-graph instance
+			FSubGraphSlot& SubGraphEntry = InstanceData->SubGraphSlots[ChildIndex];
+
+			if (SubGraphEntry.State == ESlotState::ActiveWithGraph)
+			{
+				InstanceData->SubGraphSlots[ChildIndex].GraphInstance.Release();
+			}
+
+			SubGraphEntry.State = ESlotState::Inactive;
+		}
+	}
+
+	void FSubGraphHostTrait::AddReferencedObjects(const FExecutionContext& Context, const TTraitBinding<IGarbageCollection>& Binding, FReferenceCollector& Collector) const
+	{
+		FInstanceData* InstanceData = Binding.GetInstanceData<FInstanceData>();
+
+		for (FSubGraphSlot& SubGraphEntry : InstanceData->SubGraphSlots)
+		{
+			Collector.AddPropertyReferencesWithStructARO(FAnimNextGraphInstancePtr::StaticStruct(), &SubGraphEntry.GraphInstance);
+		}
+	}
+}
diff --git a/Engine/Plugins/Experimental/AnimNextAnimGraph/Source/AnimNextAnimGraph/Private/Traits/SubGraphHost.h b/Engine/Plugins/Experimental/AnimNextAnimGraph/Source/AnimNextAnimGraph/Private/Traits/SubGraphHost.h
new file mode 100644
--- /dev/null
+++ b/Engine/Plugins/Experimental/AnimNextAnimGraph/Source/AnimNextAnimGraph/Private/Traits/SubGraphHost.h
@@ -0,0 +1,114 @@
+// Copyright Epic Games, Inc. All Rights Reserved.
+
+#pragma once
+
+#include "CoreMinimal.h"
+
+#include "TraitCore/Trait.h"
+#include "TraitInterfaces/IDiscreteBlend.h"
+#include "TraitInterfaces/IGarbageCollection.h"
+#include "TraitInterfaces/IHierarchy.h"
+#include "TraitInterfaces/IUpdate.h"
+#include "Graph/AnimNextGraph.h"
+#include "Graph/AnimNextGraphInstancePtr.h"
+
+#include "SubGraphHost.generated.h"
+
+USTRUCT(meta = (DisplayName = "SubGraph"))
+struct FAnimNextSubGraphHostTraitSharedData : public FAnimNextTraitSharedData
+{
+	GENERATED_BODY()
+
+	/** Subgraph that we host and manage. */
+	UPROPERTY(EditAnywhere, Category = "Default")
+	TObjectPtr<const UAnimNextGraph> SubGraph;
+
+	/** A dummy child that we can use to output the bind pose. This property is hidden and automatically populated during compilation. */
+	UPROPERTY(meta = (Hidden))
+	FAnimNextTraitHandle ReferencePoseChild;
+
+	/** Entry point in the Subgraph that we will use */
+	UPROPERTY(EditAnywhere, Category = "Default", meta=(CustomWidget = "ParamName", AllowedParamType = "FAnimNextEntryPoint"))
+	FName EntryPoint;
+
+	// Latent pin support boilerplate
+	#define TRAIT_LATENT_PROPERTIES_ENUMERATOR(GeneratorMacro) \
+		GeneratorMacro(SubGraph) \
+		GeneratorMacro(EntryPoint) \
+
+	GENERATE_TRAIT_LATENT_PROPERTIES(FAnimNextSubGraphHostTraitSharedData, TRAIT_LATENT_PROPERTIES_ENUMERATOR)
+	#undef TRAIT_LATENT_PROPERTIES_ENUMERATOR
+};
+
+namespace UE::AnimNext
+{
+	/**
+	 * FSubGraphHostTrait
+	 * 
+	 * A trait that hosts and manages a sub-graph instance.
+	 */
+	struct FSubGraphHostTrait : FBaseTrait, IUpdate, IUpdateTraversal, IHierarchy, IDiscreteBlend, IGarbageCollection
+	{
+		DECLARE_ANIM_TRAIT(FSubGraphHostTrait, 0xb1085d7b, FBaseTrait)
+
+		enum class ESlotState
+		{
+			ActiveWithGraph,
+			ActiveWithReferencePose,
+			Inactive,
+		};
+
+		struct FSubGraphSlot
+		{
+			// The sub-graph to use
+			TObjectPtr<const UAnimNextGraph> SubGraph;
+
+			// The graph instance
+			FAnimNextGraphInstancePtr GraphInstance;
+
+			// The entry point to use
+			FName EntryPoint;
+
+			// The current slot state
+			ESlotState State = ESlotState::Inactive;
+		};
+
+		using FSharedData = FAnimNextSubGraphHostTraitSharedData;
+
+		struct FInstanceData : FTrait::FInstanceData
+		{
+			void Construct(const FExecutionContext& Context, const FTraitBinding& Binding);
+			void Destruct(const FExecutionContext& Context, const FTraitBinding& Binding);
+
+			// List of sub-graph slots
+			TArray<FSubGraphSlot> SubGraphSlots;
+
+			// The index of the currently active sub-graph slot
+			// All other sub-graphs are blending out
+			int32 CurrentlyActiveSubGraphIndex = INDEX_NONE;
+
+			// Our child node pointer. This child is shared between all slots that have no graph provided.
+			FTraitPtr ReferencePoseChildPtr;
+		};
+
+		// IUpdate impl
+		virtual void PreUpdate(FUpdateTraversalContext& Context, const TTraitBinding<IUpdate>& Binding, const FTraitUpdateState& TraitState) const override;
+
+		// IUpdateTraversal impl
+		virtual void QueueChildrenForTraversal(FUpdateTraversalContext& Context, const TTraitBinding<IUpdateTraversal>& Binding, const FTraitUpdateState& TraitState, FUpdateTraversalQueue& TraversalQueue) const override;
+
+		// IHierarchy impl
+		virtual uint32 GetNumChildren(const FExecutionContext& Context, const TTraitBinding<IHierarchy>& Binding) const override;
+		virtual void GetChildren(const FExecutionContext& Context, const TTraitBinding<IHierarchy>& Binding, FChildrenArray& Children) const override;
+
+		// IDiscreteBlend impl
+		virtual float GetBlendWeight(const FExecutionContext& Context, const TTraitBinding<IDiscreteBlend>& Binding, int32 ChildIndex) const override;
+		virtual int32 GetBlendDestinationChildIndex(const FExecutionContext& Context, const TTraitBinding<IDiscreteBlend>& Binding) const override;
+		virtual void OnBlendTransition(const FExecutionContext& Context, const TTraitBinding<IDiscreteBlend>& Binding, int32 OldChildIndex, int32 NewChildIndex) const override;
+		virtual void OnBlendInitiated(const FExecutionContext& Context, const TTraitBinding<IDiscreteBlend>& Binding, int32 ChildIndex) const override;
+		virtual void OnBlendTerminated(const FExecutionContext& Context, const TTraitBinding<IDiscreteBlend>& Binding, int32 ChildIndex) const override;
+
+		// IGarbageCollection impl
+		virtual void AddReferencedObjects(const FExecutionContext& Context, const TTraitBinding<IGarbageCollection>& Binding, FReferenceCollector& Collector) const override;
+	};
+}
diff --git a/Engine/Plugins/Experimental/AnimNextAnimGraph/Source/AnimNextAnimGraph/Private/Traits/SynchronizeUsingGroups.cpp b/Engine/Plugins/Experimental/AnimNextAnimGraph/Source/AnimNextAnimGraph/Private/Traits/SynchronizeUsingGroups.cpp
new file mode 100644
--- /dev/null
+++ b/Engine/Plugins/Experimental/AnimNextAnimGraph/Source/AnimNextAnimGraph/Private/Traits/SynchronizeUsingGroups.cpp
@@ -0,0 +1,114 @@
+// Copyright Epic Games, Inc. All Rights Reserved.
+
+#include "Traits/SynchronizeUsingGroups.h"
+#include "Graph/SyncGroup_GraphInstanceComponent.h"
+
+namespace UE::AnimNext
+{
+	AUTO_REGISTER_ANIM_TRAIT(FSynchronizeUsingGroupsTrait)
+
+	// Trait implementation boilerplate
+	#define TRAIT_INTERFACE_ENUMERATOR(GeneratorMacro) \
+		GeneratorMacro(IGroupSynchronization) \
+		GeneratorMacro(ITimeline) \
+		GeneratorMacro(IUpdate) \
+
+	GENERATE_ANIM_TRAIT_IMPLEMENTATION(FSynchronizeUsingGroupsTrait, TRAIT_INTERFACE_ENUMERATOR)
+	#undef TRAIT_INTERFACE_ENUMERATOR
+
+	void FSynchronizeUsingGroupsTrait::PreUpdate(FUpdateTraversalContext& Context, const TTraitBinding<IUpdate>& Binding, const FTraitUpdateState& TraitState) const
+	{
+		FInstanceData* InstanceData = Binding.GetInstanceData<FInstanceData>();
+
+		TTraitBinding<IGroupSynchronization> GroupSyncTrait;
+		Binding.GetStackInterface(GroupSyncTrait);
+
+		const FName GroupName = GroupSyncTrait.GetGroupName(Context);
+		const bool bHasGroupName = !GroupName.IsNone();
+
+		// If we have a group name, we are active
+		// Freeze the timeline, our sync group will control it
+		InstanceData->bFreezeTimeline = bHasGroupName;
+
+		// Forward the PreUpdate call, if the timeline attempts to update, we'll do nothing if we are frozen
+		IUpdate::PreUpdate(Context, Binding, TraitState);
+
+		if (!bHasGroupName)
+		{
+			// If no group name is specified, this trait is inactive
+			return;
+		}
+
+		const EAnimGroupRole::Type GroupRole = GroupSyncTrait.GetGroupRole(Context);
+
+		// Append this trait to our group, we'll need to synchronize it
+		FSyncGroupGraphInstanceComponent& Component = Context.GetComponent<FSyncGroupGraphInstanceComponent>();
+		Component.RegisterWithGroup(GroupName, GroupRole, Binding.GetTraitPtr(), TraitState);
+	}
+
+	FName FSynchronizeUsingGroupsTrait::GetGroupName(const FExecutionContext& Context, const TTraitBinding<IGroupSynchronization>& Binding) const
+	{
+		const FSharedData* SharedData = Binding.GetSharedData<FSharedData>();
+		return SharedData->GroupName;
+	}
+
+	EAnimGroupRole::Type FSynchronizeUsingGroupsTrait::GetGroupRole(const FExecutionContext& Context, const TTraitBinding<IGroupSynchronization>& Binding) const
+	{
+		const FSharedData* SharedData = Binding.GetSharedData<FSharedData>();
+		return SharedData->GroupRole;
+	}
+
+	float FSynchronizeUsingGroupsTrait::AdvanceBy(const FExecutionContext& Context, const TTraitBinding<IGroupSynchronization>& Binding, float DeltaTime) const
+	{
+		FInstanceData* InstanceData = Binding.GetInstanceData<FInstanceData>();
+
+		// When the group advances the timeline, we thaw it to advance
+		InstanceData->bFreezeTimeline = false;
+
+		TTraitBinding<ITimeline> TimelineTrait;
+		Binding.GetStackInterface(TimelineTrait);
+
+		const float ProgressRatio = TimelineTrait.AdvanceBy(Context, DeltaTime);
+
+		InstanceData->bFreezeTimeline = true;
+
+		return ProgressRatio;
+	}
+
+	void FSynchronizeUsingGroupsTrait::AdvanceToRatio(const FExecutionContext& Context, const TTraitBinding<IGroupSynchronization>& Binding, float ProgressRatio) const
+	{
+		FInstanceData* InstanceData = Binding.GetInstanceData<FInstanceData>();
+
+		// When the group advances the timeline, we thaw it to advance
+		InstanceData->bFreezeTimeline = false;
+
+		TTraitBinding<ITimeline> TimelineTrait;
+		Binding.GetStackInterface(TimelineTrait);
+
+		TimelineTrait.AdvanceToRatio(Context, ProgressRatio);
+
+		InstanceData->bFreezeTimeline = true;
+	}
+
+	float FSynchronizeUsingGroupsTrait::AdvanceBy(const FExecutionContext& Context, const TTraitBinding<ITimeline>& Binding, float DeltaTime) const
+	{
+		const FInstanceData* InstanceData = Binding.GetInstanceData<FInstanceData>();
+		if (InstanceData->bFreezeTimeline)
+		{
+			return 0.0f;	// If the timeline is frozen, we don't advance
+		}
+
+		return ITimeline::AdvanceBy(Context, Binding, DeltaTime);
+	}
+
+	void FSynchronizeUsingGroupsTrait::AdvanceToRatio(const FExecutionContext& Context, const TTraitBinding<ITimeline>& Binding, float ProgressRatio) const
+	{
+		const FInstanceData* InstanceData = Binding.GetInstanceData<FInstanceData>();
+		if (InstanceData->bFreezeTimeline)
+		{
+			return;			// If the timeline is frozen, we don't advance
+		}
+
+		ITimeline::AdvanceToRatio(Context, Binding, ProgressRatio);
+	}
+}
diff --git a/Engine/Plugins/Experimental/AnimNextAnimGraph/Source/AnimNextAnimGraph/Private/Traits/SynchronizeUsingGroups.h b/Engine/Plugins/Experimental/AnimNextAnimGraph/Source/AnimNextAnimGraph/Private/Traits/SynchronizeUsingGroups.h
new file mode 100644
--- /dev/null
+++ b/Engine/Plugins/Experimental/AnimNextAnimGraph/Source/AnimNextAnimGraph/Private/Traits/SynchronizeUsingGroups.h
@@ -0,0 +1,60 @@
+// Copyright Epic Games, Inc. All Rights Reserved.
+
+#pragma once
+
+#include "CoreMinimal.h"
+#include "Animation/AnimationAsset.h"
+#include "TraitCore/Trait.h"
+#include "TraitInterfaces/IGroupSynchronization.h"
+#include "TraitInterfaces/IUpdate.h"
+#include "TraitInterfaces/ITimeline.h"
+
+#include "SynchronizeUsingGroups.generated.h"
+
+USTRUCT(meta = (DisplayName = "Synchronize Using Groups"))
+struct FAnimNextSynchronizeUsingGroupsTraitSharedData : public FAnimNextTraitSharedData
+{
+	GENERATED_BODY()
+
+	// The group name
+	// If no name is provided, this trait is inactive
+	UPROPERTY(EditAnywhere, Category = "Default", meta = (Inline))
+	FName GroupName;
+
+	// The role this player can assume within the group
+	UPROPERTY(EditAnywhere, Category = "Default", meta = (Inline))
+	TEnumAsByte<EAnimGroupRole::Type> GroupRole = EAnimGroupRole::CanBeLeader;
+};
+
+namespace UE::AnimNext
+{
+	/**
+	 * FSynchronizeUsingGroupsTrait
+	 * 
+	 * A trait that synchronizes animation sequence playback using named groups.
+	 */
+	struct FSynchronizeUsingGroupsTrait : FAdditiveTrait, IUpdate, IGroupSynchronization, ITimeline
+	{
+		DECLARE_ANIM_TRAIT(FSynchronizeUsingGroupsTrait, 0x6d318931, FAdditiveTrait)
+
+		using FSharedData = FAnimNextSynchronizeUsingGroupsTraitSharedData;
+
+		struct FInstanceData : FTraitInstanceData
+		{
+			bool bFreezeTimeline = false;
+		};
+
+		// IUpdate impl
+		virtual void PreUpdate(FUpdateTraversalContext& Context, const TTraitBinding<IUpdate>& Binding, const FTraitUpdateState& TraitState) const override;
+
+		// IGroupSynchronization impl
+		virtual FName GetGroupName(const FExecutionContext& Context, const TTraitBinding<IGroupSynchronization>& Binding) const override;
+		virtual EAnimGroupRole::Type GetGroupRole(const FExecutionContext& Context, const TTraitBinding<IGroupSynchronization>& Binding) const override;
+		virtual float AdvanceBy(const FExecutionContext& Context, const TTraitBinding<IGroupSynchronization>& Binding, float DeltaTime) const override;
+		virtual void AdvanceToRatio(const FExecutionContext& Context, const TTraitBinding<IGroupSynchronization>& Binding, float ProgressRatio) const override;
+
+		// ITimeline impl
+		virtual float AdvanceBy(const FExecutionContext& Context, const TTraitBinding<ITimeline>& Binding, float DeltaTime) const override;
+		virtual void AdvanceToRatio(const FExecutionContext& Context, const TTraitBinding<ITimeline>& Binding, float ProgressRatio) const override;
+	};
+}
diff --git a/Engine/Plugins/Experimental/AnimNextAnimGraph/Source/AnimNextAnimGraph/Public/AnimNextAnimGraphStats.h b/Engine/Plugins/Experimental/AnimNextAnimGraph/Source/AnimNextAnimGraph/Public/AnimNextAnimGraphStats.h
new file mode 100644
--- /dev/null
+++ b/Engine/Plugins/Experimental/AnimNextAnimGraph/Source/AnimNextAnimGraph/Public/AnimNextAnimGraphStats.h
@@ -0,0 +1,9 @@
+// Copyright Epic Games, Inc. All Rights Reserved.
+
+#pragma once
+
+#include "Stats/Stats2.h"
+#include "AnimNextStats.h"
+
+DECLARE_CYCLE_STAT_EXTERN(TEXT("AnimNext: Update Graph"), STAT_AnimNext_UpdateGraph, STATGROUP_AnimNext, ANIMNEXTANIMGRAPH_API);
+DECLARE_CYCLE_STAT_EXTERN(TEXT("AnimNext: Evaluate Graph"), STAT_AnimNext_EvaluateGraph, STATGROUP_AnimNext, ANIMNEXTANIMGRAPH_API);
diff --git a/Engine/Plugins/Experimental/AnimNextAnimGraph/Source/AnimNextAnimGraph/Public/TraitInterfaces/IContinuousBlend.h b/Engine/Plugins/Experimental/AnimNextAnimGraph/Source/AnimNextAnimGraph/Public/TraitInterfaces/IContinuousBlend.h
new file mode 100644
--- /dev/null
+++ b/Engine/Plugins/Experimental/AnimNextAnimGraph/Source/AnimNextAnimGraph/Public/TraitInterfaces/IContinuousBlend.h
@@ -0,0 +1,41 @@
+// Copyright Epic Games, Inc. All Rights Reserved.
+
+#pragma once
+
+#include "CoreMinimal.h"
+#include "TraitCore/ITraitInterface.h"
+#include "TraitCore/TraitBinding.h"
+
+namespace UE::AnimNext
+{
+	/**
+	 * IContinuousBlend
+	 *
+	 * This interface exposes continuous blend related information.
+	 */
+	struct ANIMNEXTANIMGRAPH_API IContinuousBlend : ITraitInterface
+	{
+		DECLARE_ANIM_TRAIT_INTERFACE(IContinuousBlend, 0xe7d79186)
+
+		// Returns the blend weight for the specified child
+		// Multiple children can have non-zero weight but their sum must be 1.0
+		// Returns -1.0 if the child index is invalid
+		virtual float GetBlendWeight(const FExecutionContext& Context, const TTraitBinding<IContinuousBlend>& Binding, int32 ChildIndex) const;
+	};
+
+	/**
+	 * Specialization for trait binding.
+	 */
+	template<>
+	struct TTraitBinding<IContinuousBlend> : FTraitBinding
+	{
+		// @see IContinuousBlend::GetBlendWeight
+		float GetBlendWeight(const FExecutionContext& Context, int32 ChildIndex) const
+		{
+			return GetInterface()->GetBlendWeight(Context, *this, ChildIndex);
+		}
+
+	protected:
+		const IContinuousBlend* GetInterface() const { return GetInterfaceTyped<IContinuousBlend>(); }
+	};
+}
diff --git a/Engine/Plugins/Experimental/AnimNextAnimGraph/Source/AnimNextAnimGraph/Public/TraitInterfaces/IDiscreteBlend.h b/Engine/Plugins/Experimental/AnimNextAnimGraph/Source/AnimNextAnimGraph/Public/TraitInterfaces/IDiscreteBlend.h
new file mode 100644
--- /dev/null
+++ b/Engine/Plugins/Experimental/AnimNextAnimGraph/Source/AnimNextAnimGraph/Public/TraitInterfaces/IDiscreteBlend.h
@@ -0,0 +1,95 @@
+// Copyright Epic Games, Inc. All Rights Reserved.
+
+#pragma once
+
+#include "CoreMinimal.h"
+#include "TraitCore/ITraitInterface.h"
+#include "TraitCore/TraitBinding.h"
+
+struct FAlphaBlend;
+
+namespace UE::AnimNext
+{
+	/**
+	 * IDiscreteBlend
+	 *
+	 * This interface exposes discrete blend related information.
+	 */
+	struct ANIMNEXTANIMGRAPH_API IDiscreteBlend : ITraitInterface
+	{
+		DECLARE_ANIM_TRAIT_INTERFACE(IDiscreteBlend, 0x2d395d56)
+
+		// Returns the blend weight for the specified child
+		// Multiple children can have non-zero weight but their sum must be 1.0
+		// Returns -1.0 if the child index is invalid
+		virtual float GetBlendWeight(const FExecutionContext& Context, const TTraitBinding<IDiscreteBlend>& Binding, int32 ChildIndex) const;
+
+		// Returns the blend alpha for the specified child
+		// Returns nullptr if the child index is invalid
+		// Allows additive traits to query the internal state of the base trait
+		virtual const FAlphaBlend* GetBlendState(const FExecutionContext& Context, const TTraitBinding<IDiscreteBlend>& Binding, int32 ChildIndex) const;
+
+		// Returns the blend destination child index (aka the active child index)
+		// Returns INDEX_NONE if no child is active
+		virtual int32 GetBlendDestinationChildIndex(const FExecutionContext& Context, const TTraitBinding<IDiscreteBlend>& Binding) const;
+
+		// Called when a blend transition between children occurs
+		// OldChildIndex can be INDEX_NONE if there was no previously active child
+		// NewChildIndex can be larger than the current number of known children to support a dynamic number of children at runtime
+		// When this occurs, the number of children increments by one
+		// The number of children never shrinks
+		virtual void OnBlendTransition(const FExecutionContext& Context, const TTraitBinding<IDiscreteBlend>& Binding, int32 OldChildIndex, int32 NewChildIndex) const;
+
+		// Called when the blend for specified child is initiated
+		virtual void OnBlendInitiated(const FExecutionContext& Context, const TTraitBinding<IDiscreteBlend>& Binding, int32 ChildIndex) const;
+
+		// Called when the blend for specified child terminates
+		virtual void OnBlendTerminated(const FExecutionContext& Context, const TTraitBinding<IDiscreteBlend>& Binding, int32 ChildIndex) const;
+	};
+
+	/**
+	 * Specialization for trait binding.
+	 */
+	template<>
+	struct TTraitBinding<IDiscreteBlend> : FTraitBinding
+	{
+		// @see IDiscreteBlend::GetBlendWeight
+		float GetBlendWeight(const FExecutionContext& Context, int32 ChildIndex) const
+		{
+			return GetInterface()->GetBlendWeight(Context, *this, ChildIndex);
+		}
+
+		// @see IDiscreteBlend::GetBlendState
+		const FAlphaBlend* GetBlendState(const FExecutionContext& Context, int32 ChildIndex) const
+		{
+			return GetInterface()->GetBlendState(Context, *this, ChildIndex);
+		}
+
+		// @see IDiscreteBlend::GetBlendDestinationChildIndex
+		int32 GetBlendDestinationChildIndex(const FExecutionContext& Context) const
+		{
+			return GetInterface()->GetBlendDestinationChildIndex(Context, *this);
+		}
+
+		// @see IDiscreteBlend::OnBlendTransition
+		void OnBlendTransition(const FExecutionContext& Context, int32 OldChildIndex, int32 NewChildIndex) const
+		{
+			GetInterface()->OnBlendTransition(Context, *this, OldChildIndex, NewChildIndex);
+		}
+
+		// @see IDiscreteBlend::OnBlendInitiated
+		void OnBlendInitiated(const FExecutionContext& Context, int32 ChildIndex) const
+		{
+			GetInterface()->OnBlendInitiated(Context, *this, ChildIndex);
+		}
+
+		// @see IDiscreteBlend::OnBlendTerminated
+		void OnBlendTerminated(const FExecutionContext& Context, int32 ChildIndex) const
+		{
+			GetInterface()->OnBlendTerminated(Context, *this, ChildIndex);
+		}
+
+	protected:
+		const IDiscreteBlend* GetInterface() const { return GetInterfaceTyped<IDiscreteBlend>(); }
+	};
+}
diff --git a/Engine/Plugins/Experimental/AnimNextAnimGraph/Source/AnimNextAnimGraph/Public/TraitInterfaces/IEvaluate.h b/Engine/Plugins/Experimental/AnimNextAnimGraph/Source/AnimNextAnimGraph/Public/TraitInterfaces/IEvaluate.h
new file mode 100644
--- /dev/null
+++ b/Engine/Plugins/Experimental/AnimNextAnimGraph/Source/AnimNextAnimGraph/Public/TraitInterfaces/IEvaluate.h
@@ -0,0 +1,111 @@
+// Copyright Epic Games, Inc. All Rights Reserved.
+
+#pragma once
+
+#include "CoreMinimal.h"
+#include "TraitCore/ExecutionContext.h"
+#include "TraitCore/ITraitInterface.h"
+#include "TraitCore/TraitBinding.h"
+#include "EvaluationVM/EvaluationProgram.h"
+#include "EvaluationVM/KeyframeState.h"
+
+namespace UE::AnimNext
+{
+	/**
+	 * FEvaluateTraversalContext
+	 *
+	 * Contains all relevant transient data for an evaluate traversal and wraps the execution context.
+	 */
+	struct FEvaluateTraversalContext final : FExecutionContext
+	{
+		// Appends a new task into the evaluation program, tasks mutate state in the order they have been appended in
+		// This means that child nodes need to evaluate first, tasks will usually be appended in IEvaluate::PostEvaluate
+		// Tasks are moved into their final memory location, caller can allocate the task anywhere, it is no longer needed after this operation
+		// @see FEvaluationProgram, FEvaluationTask, FEvaluationVM
+		template<class TaskType>
+		void AppendTask(TaskType&& Task) { EvaluationProgram.AppendTask(MoveTemp(Task)); }
+
+	private:
+		explicit FEvaluateTraversalContext(FEvaluationProgram& InEvaluationProgram);
+
+		FEvaluationProgram& EvaluationProgram;
+
+		friend ANIMNEXTANIMGRAPH_API FEvaluationProgram EvaluateGraph(const FWeakTraitPtr& GraphRootPtr);
+	};
+
+	/**
+	 * IEvaluate
+	 * 
+	 * This interface is called during the evaluation traversal. It aims to produce an evaluation program.
+	 * 
+	 * When a node is visited, PreEvaluate is first called on its top trait. It is responsible for forwarding
+	 * the call to the next trait that implements this interface on the trait stack of the node. Once
+	 * all traits have had the chance to PreEvaluate, the children of the trait are queried through
+	 * the IHierarchy interface. The children will then evaluate and PostEvaluate will then be called afterwards
+	 * on the original trait.
+	 * 
+	 * The execution context contains what to evaluate.
+	 * @see FEvaluationProgram
+	 */
+	struct ANIMNEXTANIMGRAPH_API IEvaluate : ITraitInterface
+	{
+		DECLARE_ANIM_TRAIT_INTERFACE(IEvaluate, 0xa303e9e7)
+
+		// Called before a traits children are evaluated
+		virtual void PreEvaluate(FEvaluateTraversalContext& Context, const TTraitBinding<IEvaluate>& Binding) const;
+
+		// Called after a traits children have been evaluated
+		virtual void PostEvaluate(FEvaluateTraversalContext& Context, const TTraitBinding<IEvaluate>& Binding) const;
+	};
+
+	/**
+	 * Specialization for trait binding.
+	 */
+	template<>
+	struct TTraitBinding<IEvaluate> : FTraitBinding
+	{
+		// @see IEvaluate::PreEvaluate
+		void PreEvaluate(FEvaluateTraversalContext& Context) const
+		{
+			GetInterface()->PreEvaluate(Context, *this);
+		}
+
+		// @see IEvaluate::PostEvaluate
+		void PostEvaluate(FEvaluateTraversalContext& Context) const
+		{
+			GetInterface()->PostEvaluate(Context, *this);
+		}
+
+	protected:
+		const IEvaluate* GetInterface() const { return GetInterfaceTyped<IEvaluate>(); }
+	};
+
+	/**
+	 * Evaluates a sub-graph starting at its root and produces an evaluation program.
+	 * Evaluation should be deterministic and repeated calls should yield the same evaluation program.
+	 *
+	 * For each node:
+	 *     - We call PreEvaluate on all its traits
+	 *     - We call GetChildren on all its traits
+	 *     - We evaluate all children found
+	 *     - We call PostEvaluate on all its traits
+	 *
+	 * @see IEvaluate::PreEvaluate, IEvaluate::PostEvaluate, IHierarchy::GetChildren
+	 */
+	[[nodiscard]] ANIMNEXTANIMGRAPH_API FEvaluationProgram EvaluateGraph(FAnimNextGraphInstancePtr& GraphInstance);
+
+	/**
+	 * Evaluates a sub-graph starting at its root and produces an evaluation program.
+	 * Evaluation starts at the top of the stack that includes the graph root trait.
+	 * Evaluation should be deterministic and repeated calls should yield the same evaluation program.
+	 *
+	 * For each node:
+	 *     - We call PreEvaluate on all its traits
+	 *     - We call GetChildren on all its traits
+	 *     - We evaluate all children found
+	 *     - We call PostEvaluate on all its traits
+	 *
+	 * @see IEvaluate::PreEvaluate, IEvaluate::PostEvaluate, IHierarchy::GetChildren
+	 */
+	[[nodiscard]] ANIMNEXTANIMGRAPH_API FEvaluationProgram EvaluateGraph(const FWeakTraitPtr& GraphRootPtr);
+}
diff --git a/Engine/Plugins/Experimental/AnimNextAnimGraph/Source/AnimNextAnimGraph/Public/TraitInterfaces/IGroupSynchronization.h b/Engine/Plugins/Experimental/AnimNextAnimGraph/Source/AnimNextAnimGraph/Public/TraitInterfaces/IGroupSynchronization.h
new file mode 100644
--- /dev/null
+++ b/Engine/Plugins/Experimental/AnimNextAnimGraph/Source/AnimNextAnimGraph/Public/TraitInterfaces/IGroupSynchronization.h
@@ -0,0 +1,69 @@
+// Copyright Epic Games, Inc. All Rights Reserved.
+
+#pragma once
+
+#include "CoreMinimal.h"
+#include "Animation/AnimationAsset.h"
+#include "TraitCore/ITraitInterface.h"
+#include "TraitCore/TraitBinding.h"
+
+namespace UE::AnimNext
+{
+	/**
+	 * IGroupSynchronization
+	 *
+	 * This interface exposes group synchronization related information and behavior.
+	 */
+	struct ANIMNEXTANIMGRAPH_API IGroupSynchronization : ITraitInterface
+	{
+		DECLARE_ANIM_TRAIT_INTERFACE(IGroupSynchronization, 0xf607d0fd)
+
+		// Returns the group name used for synchronization
+		virtual FName GetGroupName(const FExecutionContext& Context, const TTraitBinding<IGroupSynchronization>& Binding) const;
+
+		// Returns the group role used for synchronization
+		virtual EAnimGroupRole::Type GetGroupRole(const FExecutionContext& Context, const TTraitBinding<IGroupSynchronization>& Binding) const;
+
+		// Called by the sync group graph instance component once a group has been synchronized to advance time on the leader
+		// Returns the progress ratio of playback: 0.0 = start of animation, 1.0 = end of animation
+		virtual float AdvanceBy(const FExecutionContext& Context, const TTraitBinding<IGroupSynchronization>& Binding, float DeltaTime) const;
+
+		// Called by the sync group graph instance component once a group has been synchronized to advance time on each follower
+		// Progress ratio must be between [0.0, 1.0]
+		virtual void AdvanceToRatio(const FExecutionContext& Context, const TTraitBinding<IGroupSynchronization>& Binding, float ProgressRatio) const;
+	};
+
+	/**
+	 * Specialization for trait binding.
+	 */
+	template<>
+	struct TTraitBinding<IGroupSynchronization> : FTraitBinding
+	{
+		// @see IGroupSynchronization::GetGroupName
+		FName GetGroupName(const FExecutionContext& Context) const
+		{
+			return GetInterface()->GetGroupName(Context, *this);
+		}
+
+		// @see IGroupSynchronization::GetGroupRole
+		EAnimGroupRole::Type GetGroupRole(const FExecutionContext& Context) const
+		{
+			return GetInterface()->GetGroupRole(Context, *this);
+		}
+
+		// @see IGroupSynchronization::AdvanceBy
+		float AdvanceBy(const FExecutionContext& Context, float DeltaTime) const
+		{
+			return GetInterface()->AdvanceBy(Context, *this, DeltaTime);
+		}
+
+		// @see IGroupSynchronization::AdvanceToRatio
+		void AdvanceToRatio(const FExecutionContext& Context, float ProgressRatio) const
+		{
+			GetInterface()->AdvanceToRatio(Context, *this, ProgressRatio);
+		}
+
+	protected:
+		const IGroupSynchronization* GetInterface() const { return GetInterfaceTyped<IGroupSynchronization>(); }
+	};
+}
diff --git a/Engine/Plugins/Experimental/AnimNextAnimGraph/Source/AnimNextAnimGraph/Public/TraitInterfaces/ISmoothBlend.h b/Engine/Plugins/Experimental/AnimNextAnimGraph/Source/AnimNextAnimGraph/Public/TraitInterfaces/ISmoothBlend.h
new file mode 100644
--- /dev/null
+++ b/Engine/Plugins/Experimental/AnimNextAnimGraph/Source/AnimNextAnimGraph/Public/TraitInterfaces/ISmoothBlend.h
@@ -0,0 +1,39 @@
+// Copyright Epic Games, Inc. All Rights Reserved.
+
+#pragma once
+
+#include "CoreMinimal.h"
+#include "TraitCore/ITraitInterface.h"
+#include "TraitCore/TraitBinding.h"
+
+namespace UE::AnimNext
+{
+	/**
+	 * ISmoothBlend
+	 *
+	 * This interface exposes blend smoothing related information.
+	 */
+	struct ANIMNEXTANIMGRAPH_API ISmoothBlend : ITraitInterface
+	{
+		DECLARE_ANIM_TRAIT_INTERFACE(ISmoothBlend, 0x1c2c1739)
+
+		// Returns the desired blend time for the specified child
+		virtual float GetBlendTime(const FExecutionContext& Context, const TTraitBinding<ISmoothBlend>& Binding, int32 ChildIndex) const;
+	};
+
+	/**
+	 * Specialization for trait binding.
+	 */
+	template<>
+	struct TTraitBinding<ISmoothBlend> : FTraitBinding
+	{
+		// @see ISmoothBlend::GetBlendTime
+		float GetBlendTime(const FExecutionContext& Context, int32 ChildIndex) const
+		{
+			return GetInterface()->GetBlendTime(Context, *this, ChildIndex);
+		}
+
+	protected:
+		const ISmoothBlend* GetInterface() const { return GetInterfaceTyped<ISmoothBlend>(); }
+	};
+}
diff --git a/Engine/Plugins/Experimental/AnimNextAnimGraph/Source/AnimNextAnimGraph/Public/TraitInterfaces/ITimeline.h b/Engine/Plugins/Experimental/AnimNextAnimGraph/Source/AnimNextAnimGraph/Public/TraitInterfaces/ITimeline.h
new file mode 100644
--- /dev/null
+++ b/Engine/Plugins/Experimental/AnimNextAnimGraph/Source/AnimNextAnimGraph/Public/TraitInterfaces/ITimeline.h
@@ -0,0 +1,59 @@
+// Copyright Epic Games, Inc. All Rights Reserved.
+
+#pragma once
+
+#include "CoreMinimal.h"
+#include "TraitCore/ITraitInterface.h"
+#include "TraitCore/TraitBinding.h"
+
+namespace UE::AnimNext
+{
+	/**
+	 * ITimeline
+	 *
+	 * This interface exposes timeline related information.
+	 */
+	struct ANIMNEXTANIMGRAPH_API ITimeline : ITraitInterface
+	{
+		DECLARE_ANIM_TRAIT_INTERFACE(ITimeline, 0x53760727)
+
+		// Returns the play rate of this timeline
+		virtual float GetPlayRate(const FExecutionContext& Context, const TTraitBinding<ITimeline>& Binding) const;
+
+		// Advances time by the provided delta time (positive or negative) on this timeline
+		// Returns the progress ratio of playback: 0.0 = start of animation, 1.0 = end of animation
+		virtual float AdvanceBy(const FExecutionContext& Context, const TTraitBinding<ITimeline>& Binding, float DeltaTime) const;
+
+		// Advances time to the specified progress ratio on this timeline
+		// Progress ratio must be between [0.0, 1.0]
+		virtual void AdvanceToRatio(const FExecutionContext& Context, const TTraitBinding<ITimeline>& Binding, float ProgressRatio) const;
+	};
+
+	/**
+	 * Specialization for trait binding.
+	 */
+	template<>
+	struct TTraitBinding<ITimeline> : FTraitBinding
+	{
+		// @see ITimeline::GetPlayRate
+		float GetPlayRate(const FExecutionContext& Context) const
+		{
+			return GetInterface()->GetPlayRate(Context, *this);
+		}
+
+		// @see ITimeline::AdvanceBy
+		float AdvanceBy(const FExecutionContext& Context, float DeltaTime) const
+		{
+			return GetInterface()->AdvanceBy(Context, *this, DeltaTime);
+		}
+
+		// @see ITimeline::AdvanceToRatio
+		void AdvanceToRatio(const FExecutionContext& Context, float ProgressRatio) const
+		{
+			GetInterface()->AdvanceToRatio(Context, *this, ProgressRatio);
+		}
+
+	protected:
+		const ITimeline* GetInterface() const { return GetInterfaceTyped<ITimeline>(); }
+	};
+}
diff --git a/Engine/Plugins/Experimental/AnimNextAnimGraph/Source/AnimNextAnimGraph/Public/TraitInterfaces/IUpdate.h b/Engine/Plugins/Experimental/AnimNextAnimGraph/Source/AnimNextAnimGraph/Public/TraitInterfaces/IUpdate.h
new file mode 100644
--- /dev/null
+++ b/Engine/Plugins/Experimental/AnimNextAnimGraph/Source/AnimNextAnimGraph/Public/TraitInterfaces/IUpdate.h
@@ -0,0 +1,258 @@
+// Copyright Epic Games, Inc. All Rights Reserved.
+
+#pragma once
+
+#include "CoreMinimal.h"
+#include "TraitCore/ExecutionContext.h"
+#include "TraitCore/ITraitInterface.h"
+#include "TraitCore/TraitBinding.h"
+
+#include <type_traits>
+
+struct FAnimNextGraphInstancePtr;
+class FMemStack;
+
+namespace UE::AnimNext
+{
+	struct FUpdateTraversalQueue;
+
+	namespace Private
+	{
+		struct FUpdateEntry;
+	}
+
+	/**
+	  * FTraitUpdateState
+	  * 
+	  * State that propagates down the graph during the update traversal.
+	  * 
+	  * We align this to 16 bytes and maintain that same size to ensure efficient copying.
+	  */
+	struct alignas(16) FTraitUpdateState final
+	{
+		FTraitUpdateState() = default;
+
+		explicit FTraitUpdateState(float InDeltaTime)
+			: DeltaTime(InDeltaTime)
+		{}
+
+		// Returns the delta time for this trait
+		float GetDeltaTime() const { return DeltaTime; }
+
+		// Returns the total weight for this trait
+		float GetTotalWeight() const { return TotalWeight; }
+
+		// Returns the total trajectory weight for this trait
+		float GetTotalTrajectoryWeight() const { return TotalTrajectoryWeight; }
+
+		// Returns whether or not this trait is blending out
+		bool IsBlendingOut() const { return !!bIsBlendingOut; }
+
+		// Creates a new instance of the update state with the total weight scaled by the supplied weight
+		FTraitUpdateState WithWeight(float Weight) const
+		{
+			FTraitUpdateState Result(*this);
+			Result.TotalWeight *= Weight;
+			return Result;
+		}
+
+		// Creates a new instance of the update state with the total trajectory weight scaled by the supplied weight
+		FTraitUpdateState WithTrajectoryWeight(float Weight) const
+		{
+			FTraitUpdateState Result(*this);
+			Result.TotalTrajectoryWeight *= Weight;
+			return Result;
+		}
+
+		// Creates a new instance of the update state with the delta time scaled by the supplied scale factor
+		FTraitUpdateState WithTimeScale(float TimeScale) const
+		{
+			FTraitUpdateState Result(*this);
+			Result.DeltaTime *= TimeScale;
+			return Result;
+		}
+
+		// Creates a new instance of the update state marking it as blending out
+		FTraitUpdateState AsBlendingOut() const
+		{
+			FTraitUpdateState Result(*this);
+			Result.bIsBlendingOut = true;
+			return Result;
+		}
+
+	private:
+		// The amount of time to move forward by at allowing parents to time scale their children
+		float DeltaTime = 0.0f;
+
+		// The current total weight factoring in all inherited blend weights
+		float TotalWeight = 1.0f;
+
+		// The current total trajectory weight factoring in all inherited blend weights
+		float TotalTrajectoryWeight = 1.0f;
+
+		// Whether or not we are blending out
+		int32 bIsBlendingOut = false;
+	};
+
+	static_assert(sizeof(FTraitUpdateState) <= 16, "Keep the size to 16 bytes for efficient copying");
+
+	/**
+	 * FUpdateTraversalContext
+	 *
+	 * Contains all relevant transient data for an update traversal and wraps the execution context.
+	 */
+	struct FUpdateTraversalContext final : FExecutionContext
+	{
+	private:
+		// Constructs a new traversal context
+		FUpdateTraversalContext() = default;
+
+		// Pops entries from the traversal queue and pushes them onto the update stack
+		void PushQueuedUpdateEntries(FUpdateTraversalQueue& TraversalQueue);
+
+		// Pushes an entry onto the update stack
+		void PushUpdateEntry(Private::FUpdateEntry* Entry);
+
+		// Pops an entry from the top of the update stack
+		Private::FUpdateEntry* PopUpdateEntry();
+
+		// Pushes an entry onto the free entry stack
+		void PushFreeEntry(Private::FUpdateEntry* Entry);
+
+		// Returns a new entry suitable for update queuing
+		// If an entry isn't found in the free stack, a new one is allocated from the memstack
+		Private::FUpdateEntry* GetNewEntry(const FWeakTraitPtr& TraitPtr, const FTraitUpdateState& TraitState);
+
+		// The head pointer of the update stack
+		// This is the traversal execution stack and it contains entries that are
+		// pending their pre-update call and entries waiting for post-update to be called
+		Private::FUpdateEntry* UpdateStackHead = nullptr;
+
+		// The head pointer of the free entry stack
+		// Entries are allocated from the memstack and are re-used in LIFO since they'll
+		// be warmer in the CPU cache
+		Private::FUpdateEntry* FreeEntryStackHead = nullptr;
+
+		friend ANIMNEXTANIMGRAPH_API void UpdateGraph(FAnimNextGraphInstancePtr& GraphInstance, float DeltaTime);
+		friend FUpdateTraversalQueue;
+	};
+
+	/**
+	 * FUpdateTraversalQueue
+	 *
+	 * A queue of children to traverse.
+	 * @see IUpdateTraversal::QueueChildrenForTraversal
+	 */
+	struct FUpdateTraversalQueue final
+	{
+		// Queued a child for traversal
+		// Children are processed in the same order they are queued
+		ANIMNEXTANIMGRAPH_API void Push(const FWeakTraitPtr& ChildPtr, const FTraitUpdateState& ChildTraitState);
+
+	private:
+		explicit FUpdateTraversalQueue(FUpdateTraversalContext& InTraversalContext);
+
+		// The current traversal context
+		FUpdateTraversalContext& TraversalContext;
+
+		// The head pointer of the queued update stack
+		// When a child is queued for traversal, it is first appended to this stack
+		// Once all children have been queued and pre-update terminates, this stack is
+		// emptied and pushed onto the update stack
+		Private::FUpdateEntry* QueuedUpdateStackHead = nullptr;
+
+		friend ANIMNEXTANIMGRAPH_API void UpdateGraph(FAnimNextGraphInstancePtr& GraphInstance, float DeltaTime);
+		friend FUpdateTraversalContext;
+	};
+
+	/**
+	 * IUpdate
+	 *
+	 * This interface is called during the update traversal.
+	 *
+	 * When a node is visited, PreUpdate is first called on its top trait. It is responsible for forwarding
+	 * the call to the next trait that implements this interface on the trait stack of the node. Once
+	 * all traits have had the chance to PreUpdate, the children will then evaluate and PostUpdate will
+	 * then be called afterwards on the original trait.
+	 */
+	struct ANIMNEXTANIMGRAPH_API IUpdate : ITraitInterface
+	{
+		DECLARE_ANIM_TRAIT_INTERFACE(IUpdate, 0x59d24dc5)
+
+		// Called before a traits children are updated
+		virtual void PreUpdate(FUpdateTraversalContext& Context, const TTraitBinding<IUpdate>& Binding, const FTraitUpdateState& TraitState) const;
+
+		// Called after a traits children have been updated
+		virtual void PostUpdate(FUpdateTraversalContext& Context, const TTraitBinding<IUpdate>& Binding, const FTraitUpdateState& TraitState) const;
+	};
+
+	/**
+	 * IUpdateTraversal
+	 *
+	 * This interface is called during the update traversal.
+	 *
+	 * If a trait needs to modify the update state of its children, it needs to implement this interface.
+	 * There is no need to call the Super to forward the call down the trait stack as the update traversal
+	 * takes care of it.
+	 * 
+	 * If this interface is not implemented, the IHierarchy interface will be used to retrieve and queue
+	 * the children with the same update state as the owning trait.
+	 */
+	struct ANIMNEXTANIMGRAPH_API IUpdateTraversal : ITraitInterface
+	{
+		DECLARE_ANIM_TRAIT_INTERFACE(IUpdateTraversal, 0x256c21b1)
+
+		// Called after PreUpdate to request that children be queued with the provided context
+		virtual void QueueChildrenForTraversal(FUpdateTraversalContext& Context, const TTraitBinding<IUpdateTraversal>& Binding, const FTraitUpdateState& TraitState, FUpdateTraversalQueue& TraversalQueue) const;
+	};
+
+	/**
+	 * Specialization for trait binding.
+	 */
+	template<>
+	struct TTraitBinding<IUpdate> : FTraitBinding
+	{
+		// @see IUpdate::PreUpdate
+		void PreUpdate(FUpdateTraversalContext& Context, const FTraitUpdateState& TraitState) const
+		{
+			GetInterface()->PreUpdate(Context, *this, TraitState);
+		}
+
+		// @see IUpdate::PostUpdate
+		void PostUpdate(FUpdateTraversalContext& Context, const FTraitUpdateState& TraitState) const
+		{
+			GetInterface()->PostUpdate(Context, *this, TraitState);
+		}
+
+	protected:
+		const IUpdate* GetInterface() const { return GetInterfaceTyped<IUpdate>(); }
+	};
+
+	/**
+	 * Specialization for trait binding.
+	 */
+	template<>
+	struct TTraitBinding<IUpdateTraversal> : FTraitBinding
+	{
+		// @see IUpdateTraversal::QueueChildrenForTraversal
+		void QueueChildrenForTraversal(FUpdateTraversalContext& Context, const FTraitUpdateState& TraitState, FUpdateTraversalQueue& TraversalQueue) const
+		{
+			GetInterface()->QueueChildrenForTraversal(Context, *this, TraitState, TraversalQueue);
+		}
+
+	protected:
+		const IUpdateTraversal* GetInterface() const { return GetInterfaceTyped<IUpdateTraversal>(); }
+	};
+
+	/**
+	 * Updates a sub-graph starting at its root.
+	 *
+	 * For each node:
+	 *     - We call PreUpdate on all its traits
+	 *     - We update all children
+	 *     - We call PostUpdate on all its traits
+	 *
+	 * @see IUpdate::PreUpdate, IUpdate::PostUpdate, IHierarchy::GetChildren
+	 */
+	ANIMNEXTANIMGRAPH_API void UpdateGraph(FAnimNextGraphInstancePtr& GraphInstance, float DeltaTime);
+}
diff --git a/Engine/Plugins/Experimental/AnimNextAnimGraph/Source/AnimNextAnimGraphTestSuite/AnimNextAnimGraphTestSuite.Build.cs b/Engine/Plugins/Experimental/AnimNextAnimGraph/Source/AnimNextAnimGraphTestSuite/AnimNextAnimGraphTestSuite.Build.cs
new file mode 100644
--- /dev/null
+++ b/Engine/Plugins/Experimental/AnimNextAnimGraph/Source/AnimNextAnimGraphTestSuite/AnimNextAnimGraphTestSuite.Build.cs
@@ -0,0 +1,37 @@
+// Copyright Epic Games, Inc. All Rights Reserved.
+
+namespace UnrealBuildTool.Rules
+{
+	public class AnimNextAnimGraphTestSuite : ModuleRules
+	{
+		public AnimNextAnimGraphTestSuite(ReadOnlyTargetRules Target) : base(Target)
+		{
+			PCHUsage = ModuleRules.PCHUsageMode.UseExplicitOrSharedPCHs;
+
+			PublicDependencyModuleNames.AddRange(
+				new string[]
+				{
+					"Core",
+					"CoreUObject",
+					"Engine",
+					"AnimNext",
+					"AnimNextTestSuite",
+					"AnimNextAnimGraph",
+					"RigVM",
+				}
+			);
+
+			if (Target.bBuildEditor == true)
+			{
+				PrivateDependencyModuleNames.AddRange(
+					new string[]
+					{
+						"AnimNextUncookedOnly",
+						"AnimNextEditor",
+						"RigVMDeveloper",
+					}
+				);
+			}
+		}
+	}
+}
\ No newline at end of file
diff --git a/Engine/Plugins/Experimental/AnimNextAnimGraph/Source/AnimNextAnimGraphTestSuite/Private/AnimNextAnimGraphTraitGraphTest.cpp b/Engine/Plugins/Experimental/AnimNextAnimGraph/Source/AnimNextAnimGraphTestSuite/Private/AnimNextAnimGraphTraitGraphTest.cpp
new file mode 100644
--- /dev/null
+++ b/Engine/Plugins/Experimental/AnimNextAnimGraph/Source/AnimNextAnimGraphTestSuite/Private/AnimNextAnimGraphTraitGraphTest.cpp
@@ -0,0 +1,454 @@
+// Copyright Epic Games, Inc. All Rights Reserved.
+
+#include "AnimNextAnimGraphTraitGraphTest.h"
+
+#include "AnimNextRuntimeTest.h"
+#include "AnimNextTest.h"
+#include "AssetToolsModule.h"
+#include "Context.h"
+#include "UncookedOnlyUtils.h"
+#include "TraitCore/TraitRegistry.h"
+#include "TraitInterfaces/IEvaluate.h"
+#include "TraitInterfaces/IUpdate.h"
+#include "Editor/Transactor.h"
+#include "AnimNextExecuteContext.h"
+#include "Graph/AnimNextGraph.h"
+#include "Graph/AnimNextGraph_EditorData.h"
+#include "Graph/GraphFactory.h"
+#include "Graph/RigDecorator_AnimNextCppTrait.h"
+#include "Graph/RigUnit_AnimNextGraphRoot.h"
+#include "Graph/RigUnit_AnimNextTraitStack.h"
+#include "Misc/AutomationTest.h"
+#include "Param/ParamStack.h"
+#include "Param/RigVMDispatch_GetParameter.h"
+#include "RigVMFunctions/Math/RigVMFunction_MathInt.h"
+#include "RigVMFunctions/Math/RigVMFunction_MathFloat.h"
+
+#if WITH_DEV_AUTOMATION_TESTS
+
+//****************************************************************************
+// AnimNext Runtime Trait Graph Tests
+//****************************************************************************
+
+namespace UE::AnimNext
+{
+	struct FTestTrait final : FBaseTrait, IEvaluate, IUpdate
+	{
+		DECLARE_ANIM_TRAIT(FTestTrait, 0x41cecb7f, FBaseTrait)
+
+		using FSharedData = FTestTraitSharedData;
+
+		// IUpdate impl
+		virtual void PostUpdate(FUpdateTraversalContext& Context, const TTraitBinding<IUpdate>& Binding, const FTraitUpdateState& TraitState) const override
+		{
+			IUpdate::PostUpdate(Context, Binding, TraitState);
+
+			const FSharedData* SharedData = Binding.GetSharedData<FSharedData>();
+			UE::AnimNext::FParamStack& ParamStack = UE::AnimNext::FParamStack::Get();
+
+			int32& UpdateCount = ParamStack.GetMutableParam<int32>("UpdateCount");
+			UpdateCount++;
+
+			ParamStack.GetMutableParam<int32>("SomeInt32") = SharedData->SomeInt32;
+			ParamStack.GetMutableParam<float>("SomeFloat") = SharedData->SomeFloat;
+
+			ParamStack.GetMutableParam<int32>("SomeLatentInt32") = SharedData->GetSomeLatentInt32(Binding);				// MathAdd with constants, latent
+			ParamStack.GetMutableParam<int32>("SomeOtherLatentInt32") = SharedData->GetSomeOtherLatentInt32(Binding);	// GetParameter, latent
+			ParamStack.GetMutableParam<float>("SomeLatentFloat") = SharedData->GetSomeLatentFloat(Binding);				// Inline value, not latent
+		}
+
+		// IEvaluate impl
+		virtual void PostEvaluate(FEvaluateTraversalContext& Context, const TTraitBinding<IEvaluate>& Binding) const override
+		{
+			IEvaluate::PostEvaluate(Context, Binding);
+
+			UE::AnimNext::FParamStack& ParamStack = UE::AnimNext::FParamStack::Get();
+
+			int32& EvaluateCount = ParamStack.GetMutableParam<int32>("EvaluateCount");
+			EvaluateCount++;
+		}
+	};
+
+	// Trait implementation boilerplate
+	#define TRAIT_INTERFACE_ENUMERATOR(GeneratorMacro) \
+		GeneratorMacro(IEvaluate) \
+		GeneratorMacro(IUpdate) \
+
+	GENERATE_ANIM_TRAIT_IMPLEMENTATION(FTestTrait, TRAIT_INTERFACE_ENUMERATOR)
+	#undef TRAIT_INTERFACE_ENUMERATOR
+}
+
+IMPLEMENT_SIMPLE_AUTOMATION_TEST(FAnimationAnimNextRuntimeTest_GraphAddTrait, "Animation.AnimNext.Runtime.Graph.AddTrait", EAutomationTestFlags::EditorContext | EAutomationTestFlags::EngineFilter)
+
+bool FAnimationAnimNextRuntimeTest_GraphAddTrait::RunTest(const FString& InParameters)
+{
+	using namespace UE::AnimNext;
+
+	{
+		AUTO_REGISTER_ANIM_TRAIT(FTestTrait)
+
+		FScopedClearNodeTemplateRegistry ScopedClearNodeTemplateRegistry;
+
+		UFactory* GraphFactory = NewObject<UAnimNextGraphFactory>();
+		UAnimNextGraph* AnimNextGraph = CastChecked<UAnimNextGraph>(GraphFactory->FactoryCreateNew(UAnimNextGraph::StaticClass(), GetTransientPackage(), TEXT("TestAnimNextGraph"), RF_Transient, nullptr, nullptr, NAME_None));
+		UE_RETURN_ON_ERROR(AnimNextGraph != nullptr, "FAnimationAnimNextRuntimeTest_GraphAddTrait -> Failed to create animation graph");
+
+		UAnimNextGraph_EditorData* EditorData = UncookedOnly::FUtils::GetEditorData(AnimNextGraph);
+		UE_RETURN_ON_ERROR(EditorData != nullptr, "FAnimationAnimNextRuntimeTest_GraphAddTrait -> Failed to find animation graph editor data");
+
+		UE_RETURN_ON_ERROR(EditorData->AddAnimationGraph(FRigUnit_AnimNextGraphRoot::DefaultEntryPoint, false) != nullptr, "FAnimationAnimNextRuntimeTest_GraphAddTrait -> Failed to add animation graph");
+
+		URigVMController* Controller = EditorData->GetRigVMClient()->GetController(EditorData->GetRigVMClient()->GetDefaultModel());
+		UE_RETURN_ON_ERROR(Controller != nullptr, "FAnimationAnimNextRuntimeTest_GraphAddTrait -> Failed to get RigVM controller");
+
+		// Create an empty trait stack node
+		URigVMUnitNode* DecoratorStackNode = Controller->AddUnitNode(FRigUnit_AnimNextTraitStack::StaticStruct(), FRigVMStruct::ExecuteName, FVector2D(0.0f, 0.0f), FString(), false);
+		UE_RETURN_ON_ERROR(DecoratorStackNode != nullptr, "FAnimationAnimNextRuntimeTest_GraphAddDecorator -> Failed to create decorator stack node");
+
+		// Add a trait
+		const UScriptStruct* CppDecoratorStruct = FRigDecorator_AnimNextCppDecorator::StaticStruct();
+		UE_RETURN_ON_ERROR(CppDecoratorStruct != nullptr, "FAnimationAnimNextRuntimeTest_GraphAddDecorator -> Failed to get find Cpp decorator static struct");
+
+		const FTrait* Trait = FTraitRegistry::Get().Find(FTestTrait::TraitUID);
+		UE_RETURN_ON_ERROR(Trait != nullptr, "FAnimationAnimNextRuntimeTest_GraphAddTrait -> Failed to find test trait");
+
+		UScriptStruct* ScriptStruct = Trait->GetTraitSharedDataStruct();
+		UE_RETURN_ON_ERROR(ScriptStruct != nullptr, "FAnimationAnimNextRuntimeTest_GraphAddTrait -> Failed to find trait shared data struct");
+
+		FString DefaultValue;
+		{
+			const FRigDecorator_AnimNextCppDecorator DefaultCppDecoratorStructInstance;
+			FRigDecorator_AnimNextCppDecorator CppDecoratorStructInstance;
+			CppDecoratorStructInstance.DecoratorSharedDataStruct = ScriptStruct;
+
+			UE_RETURN_ON_ERROR(CppDecoratorStructInstance.CanBeAddedToNode(DecoratorStackNode, nullptr), "FAnimationAnimNextRuntimeTest_GraphAddTrait -> Trait cannot be added to decorator stack node");
+
+			const FProperty* Prop = FAnimNextCppDecoratorWrapper::StaticStruct()->FindPropertyByName(GET_MEMBER_NAME_STRING_CHECKED(FAnimNextCppDecoratorWrapper, CppDecorator));
+			UE_RETURN_ON_ERROR(Prop != nullptr, "FAnimationAnimNextRuntimeTest_GraphAddTrait -> Failed to find wrapper property");
+
+			Prop->ExportText_Direct(DefaultValue, &CppDecoratorStructInstance, &DefaultCppDecoratorStructInstance, nullptr, PPF_None);
+		}
+
+		FString DisplayNameMetadata;
+		ScriptStruct->GetStringMetaDataHierarchical(FRigVMStruct::DisplayNameMetaName, &DisplayNameMetadata);
+		const FString DisplayName = DisplayNameMetadata.IsEmpty() ? Trait->GetTraitName() : DisplayNameMetadata;
+
+		const FName DecoratorName = Controller->AddDecorator(
+			DecoratorStackNode->GetFName(),
+			*CppDecoratorStruct->GetPathName(),
+			*DisplayName,
+			DefaultValue, INDEX_NONE, true, true);
+		UE_RETURN_ON_ERROR(DecoratorName == DisplayName, TEXT("FAnimationAnimNextRuntimeTest_GraphAddTrait -> Unexpected decorator name"));
+
+		URigVMPin* DecoratorPin = DecoratorStackNode->FindPin(*DisplayName);
+		UE_RETURN_ON_ERROR(DecoratorPin != nullptr, TEXT("FAnimationAnimNextRuntimeTest_GraphAddTrait -> Failed to find decorator pin"));
+
+		// Our first pin is the hard coded output result, decorator pins follow
+		UE_RETURN_ON_ERROR(DecoratorStackNode->GetPins().Num() == 2, TEXT("FAnimationAnimNextRuntimeTest_GraphAddTrait -> Unexpected number of pins"));
+		UE_RETURN_ON_ERROR(DecoratorPin->IsDecoratorPin() == true, TEXT("FAnimationAnimNextRuntimeTest_GraphAddTrait -> Unexpected pin type"));
+		UE_RETURN_ON_ERROR(DecoratorPin->GetFName() == DecoratorName, TEXT("FAnimationAnimNextRuntimeTest_GraphAddTrait -> Unexpected pin name"));
+		UE_RETURN_ON_ERROR(DecoratorPin->GetCPPTypeObject() == FRigDecorator_AnimNextCppDecorator::StaticStruct(), TEXT("FAnimationAnimNextRuntimeTest_GraphAddTrait -> Unexpected pin type"));
+
+		// Our first sub-pin is the hard coded script struct member that parametrizes the decorator, dynamic decorator sub-pins follow
+		UE_RETURN_ON_ERROR(DecoratorPin->GetSubPins().Num() == 6, TEXT("FAnimationAnimNextRuntimeTest_GraphAddTrait -> Unexpected decorator sub pins"));
+
+		// SomeInt32
+		UE_RETURN_ON_ERROR(DecoratorPin->GetSubPins()[1]->GetCPPType() == TEXT("int32"), TEXT("FAnimationAnimNextRuntimeTest_GraphAddTrait -> Unexpected decorator pin type"));
+		UE_RETURN_ON_ERROR(DecoratorPin->GetSubPins()[1]->GetDefaultValue() == TEXT("3"), TEXT("FAnimationAnimNextRuntimeTest_GraphAddTrait -> Unexpected decorator pin value"));
+		UE_RETURN_ON_ERROR(!DecoratorPin->GetSubPins()[1]->IsLazy(), TEXT("FAnimationAnimNextRuntimeTest_GraphAddTrait -> Expected non-lazy decorator pin"));
+
+		// SomeFloat
+		UE_RETURN_ON_ERROR(DecoratorPin->GetSubPins()[2]->GetCPPType() == TEXT("float"), TEXT("FAnimationAnimNextRuntimeTest_GraphAddTrait -> Unexpected decorator pin type"));
+		UE_RETURN_ON_ERROR(DecoratorPin->GetSubPins()[2]->GetDefaultValue() == TEXT("34.000000"), TEXT("FAnimationAnimNextRuntimeTest_GraphAddTrait -> Unexpected decorator pin value"));
+		UE_RETURN_ON_ERROR(!DecoratorPin->GetSubPins()[2]->IsLazy(), TEXT("FAnimationAnimNextRuntimeTest_GraphAddTrait -> Expected non-lazy decorator pin"));
+
+		// SomeLatentInt32
+		UE_RETURN_ON_ERROR(DecoratorPin->GetSubPins()[3]->GetCPPType() == TEXT("int32"), TEXT("FAnimationAnimNextRuntimeTest_GraphAddTrait -> Unexpected decorator pin type"));
+		UE_RETURN_ON_ERROR(DecoratorPin->GetSubPins()[3]->GetDefaultValue() == TEXT("3"), TEXT("FAnimationAnimNextRuntimeTest_GraphAddTrait -> Unexpected decorator pin value"));
+		UE_RETURN_ON_ERROR(DecoratorPin->GetSubPins()[3]->IsLazy(), TEXT("FAnimationAnimNextRuntimeTest_GraphAddTrait -> Expected lazy decorator pin"));
+
+		// SomeOtherLatentInt32
+		UE_RETURN_ON_ERROR(DecoratorPin->GetSubPins()[4]->GetCPPType() == TEXT("int32"), TEXT("FAnimationAnimNextRuntimeTest_GraphAddTrait -> Unexpected decorator pin type"));
+		UE_RETURN_ON_ERROR(DecoratorPin->GetSubPins()[4]->GetDefaultValue() == TEXT("3"), TEXT("FAnimationAnimNextRuntimeTest_GraphAddTrait -> Unexpected decorator pin value"));
+		UE_RETURN_ON_ERROR(DecoratorPin->GetSubPins()[4]->IsLazy(), TEXT("FAnimationAnimNextRuntimeTest_GraphAddTrait -> Expected lazy decorator pin"));
+
+		// SomeLatentFloat
+		UE_RETURN_ON_ERROR(DecoratorPin->GetSubPins()[5]->GetCPPType() == TEXT("float"), TEXT("FAnimationAnimNextRuntimeTest_GraphAddTrait -> Unexpected decorator pin type"));
+		UE_RETURN_ON_ERROR(DecoratorPin->GetSubPins()[5]->GetDefaultValue() == TEXT("34.000000"), TEXT("FAnimationAnimNextRuntimeTest_GraphAddTrait -> Unexpected decorator pin value"));
+		UE_RETURN_ON_ERROR(DecoratorPin->GetSubPins()[5]->IsLazy(), TEXT("FAnimationAnimNextRuntimeTest_GraphAddTrait -> Expected lazy decorator pin"));
+	}
+
+	Tests::FUtils::CleanupAfterTests();
+
+	return true;
+}
+
+IMPLEMENT_SIMPLE_AUTOMATION_TEST(FAnimationAnimNextRuntimeTest_GraphExecute, "Animation.AnimNext.Runtime.Graph.Execute", EAutomationTestFlags::EditorContext | EAutomationTestFlags::EngineFilter)
+
+bool FAnimationAnimNextRuntimeTest_GraphExecute::RunTest(const FString& InParameters)
+{
+	using namespace UE::AnimNext;
+
+	{
+		AUTO_REGISTER_ANIM_TRAIT(FTestTrait)
+
+		FScopedClearNodeTemplateRegistry ScopedClearNodeTemplateRegistry;
+
+		UFactory* GraphFactory = NewObject<UAnimNextGraphFactory>();
+		UAnimNextGraph* AnimNextGraph = CastChecked<UAnimNextGraph>(GraphFactory->FactoryCreateNew(UAnimNextGraph::StaticClass(), GetTransientPackage(), TEXT("TestAnimNextGraph"), RF_Transient, nullptr, nullptr, NAME_None));
+		UE_RETURN_ON_ERROR(AnimNextGraph != nullptr, "FAnimationAnimNextRuntimeTest_GraphExecute -> Failed to create animation graph");
+
+		UAnimNextGraph_EditorData* EditorData = UncookedOnly::FUtils::GetEditorData(AnimNextGraph);
+		UE_RETURN_ON_ERROR(EditorData != nullptr, "FAnimationAnimNextRuntimeTest_GraphExecute -> Failed to find animation graph editor data");
+
+		UE_RETURN_ON_ERROR(EditorData->AddAnimationGraph(FRigUnit_AnimNextGraphRoot::DefaultEntryPoint, false) != nullptr, "FAnimationAnimNextRuntimeTest_GraphExecute -> Failed to add animation graph");
+
+		URigVMController* Controller = EditorData->GetRigVMClient()->GetController(EditorData->GetRigVMClient()->GetDefaultModel());
+		UE_RETURN_ON_ERROR(Controller != nullptr, "FAnimationAnimNextRuntimeTest_GraphExecute -> Failed to get RigVM controller");
+
+		// Find graph entry point
+		URigVMNode* MainEntryPointNode = Controller->GetGraph()->FindNodeByName(FRigUnit_AnimNextGraphRoot::StaticStruct()->GetFName());
+		UE_RETURN_ON_ERROR(MainEntryPointNode != nullptr, "FAnimationAnimNextRuntimeTest_GraphExecute -> Failed to find main entry point node");
+
+		URigVMPin* BeginExecutePin = MainEntryPointNode->FindPin(GET_MEMBER_NAME_STRING_CHECKED(FRigUnit_AnimNextGraphRoot, Result));
+		UE_RETURN_ON_ERROR(BeginExecutePin != nullptr && BeginExecutePin->GetDirection() == ERigVMPinDirection::Input, "FAnimationAnimNextRuntimeTest_GraphExecute -> Failed to create entry point");
+
+		// Create an empty decorator stack node
+		URigVMUnitNode* DecoratorStackNode = Controller->AddUnitNode(FRigUnit_AnimNextTraitStack::StaticStruct(), FRigVMStruct::ExecuteName, FVector2D(0.0f, 0.0f), FString(), false);
+		UE_RETURN_ON_ERROR(DecoratorStackNode != nullptr, "FAnimationAnimNextRuntimeTest_GraphExecute -> Failed to create decorator stack node");
+
+		// Link our stack result to our entry point
+		Controller->AddLink(DecoratorStackNode->GetPins()[0], MainEntryPointNode->GetPins()[0]);
+
+		// Add a trait
+		const UScriptStruct* CppDecoratorStruct = FRigDecorator_AnimNextCppDecorator::StaticStruct();
+		UE_RETURN_ON_ERROR(CppDecoratorStruct != nullptr, "FAnimationAnimNextRuntimeTest_GraphExecute -> Failed to get find Cpp decorator static struct");
+
+		const FTrait* Trait = FTraitRegistry::Get().Find(FTestTrait::TraitUID);
+		UE_RETURN_ON_ERROR(Trait != nullptr, "FAnimationAnimNextRuntimeTest_GraphExecute -> Failed to find test trait");
+
+		UScriptStruct* ScriptStruct = Trait->GetTraitSharedDataStruct();
+		UE_RETURN_ON_ERROR(ScriptStruct != nullptr, "FAnimationAnimNextRuntimeTest_GraphExecute -> Failed to find trait shared data struct");
+
+		FString DefaultValue;
+		{
+			const FRigDecorator_AnimNextCppDecorator DefaultCppDecoratorStructInstance;
+			FRigDecorator_AnimNextCppDecorator CppDecoratorStructInstance;
+			CppDecoratorStructInstance.DecoratorSharedDataStruct = ScriptStruct;
+
+			UE_RETURN_ON_ERROR(CppDecoratorStructInstance.CanBeAddedToNode(DecoratorStackNode, nullptr), "FAnimationAnimNextRuntimeTest_GraphExecute -> Trait cannot be added to decorator stack node");
+
+			const FProperty* Prop = FAnimNextCppDecoratorWrapper::StaticStruct()->FindPropertyByName(GET_MEMBER_NAME_STRING_CHECKED(FAnimNextCppDecoratorWrapper, CppDecorator));
+			UE_RETURN_ON_ERROR(Prop != nullptr, "FAnimationAnimNextRuntimeTest_GraphExecute -> Failed to find wrapper property");
+
+			Prop->ExportText_Direct(DefaultValue, &CppDecoratorStructInstance, &DefaultCppDecoratorStructInstance, nullptr, PPF_None);
+		}
+
+		FString DisplayNameMetadata;
+		ScriptStruct->GetStringMetaDataHierarchical(FRigVMStruct::DisplayNameMetaName, &DisplayNameMetadata);
+		const FString DisplayName = DisplayNameMetadata.IsEmpty() ? Trait->GetTraitName() : DisplayNameMetadata;
+
+		const FName DecoratorName = Controller->AddDecorator(
+			DecoratorStackNode->GetFName(),
+			*CppDecoratorStruct->GetPathName(),
+			*DisplayName,
+			DefaultValue, INDEX_NONE, true, true);
+		UE_RETURN_ON_ERROR(DecoratorName == DisplayName, TEXT("FAnimationAnimNextRuntimeTest_GraphExecute -> Unexpected decorator name"));
+
+		URigVMPin* DecoratorPin = DecoratorStackNode->FindPin(*DisplayName);
+		UE_RETURN_ON_ERROR(DecoratorPin != nullptr, TEXT("FAnimationAnimNextRuntimeTest_GraphExecute -> Failed to find decorator pin"));
+
+		// Set some values on our decorator
+		Controller->SetPinDefaultValue(DecoratorPin->GetSubPins()[1]->GetPinPath(), TEXT("78"));
+		Controller->SetPinDefaultValue(DecoratorPin->GetSubPins()[2]->GetPinPath(), TEXT("142.33"));
+
+		TSharedRef<FParamStack> ParamStack = MakeShared<FParamStack>();
+		FParamStack::AttachToCurrentThread(ParamStack);
+
+		FAnimNextGraphInstancePtr GraphInstance;
+		AnimNextGraph->AllocateInstance(GraphInstance);
+
+		FParamStack::FPushedLayerHandle LayerHandle = ParamStack->PushValues(
+			"UpdateCount", (int32)0,
+			"EvaluateCount", (int32)0,
+			"SomeInt32", (int32)0,
+			"SomeFloat", 0.0f,
+			"SomeLatentInt32", (int32)0,
+			"SomeOtherLatentInt32", (int32)0,
+			"SomeLatentFloat", 0.0f
+		);
+
+		{
+			UE::AnimNext::UpdateGraph(GraphInstance, 1.0f / 30.0f);
+			(void)UE::AnimNext::EvaluateGraph(GraphInstance);
+		}
+
+		AddErrorIfFalse(ParamStack->GetParam<int32>("UpdateCount") == 1, "FAnimationAnimNextRuntimeTest_GraphExecute -> Unexpected update count");
+		AddErrorIfFalse(ParamStack->GetParam<int32>("EvaluateCount") == 1, "FAnimationAnimNextRuntimeTest_GraphExecute -> Unexpected evaluate count");
+		AddErrorIfFalse(ParamStack->GetParam<int32>("SomeInt32") == 78, "FAnimationAnimNextRuntimeTest_GraphExecute -> Unexpected SomeInt32 value");
+		AddErrorIfFalse(ParamStack->GetParam<float>("SomeFloat") == 142.33f, "FAnimationAnimNextRuntimeTest_GraphExecute -> Unexpected SomeFloat value");
+		AddErrorIfFalse(ParamStack->GetParam<int32>("SomeLatentInt32") == 3, "FAnimationAnimNextRuntimeTest_GraphExecute -> Unexpected SomeLatentInt32 value");
+		AddErrorIfFalse(ParamStack->GetParam<int32>("SomeOtherLatentInt32") == 3, "FAnimationAnimNextRuntimeTest_GraphExecute -> Unexpected SomeOtherLatentInt32 value");
+		AddErrorIfFalse(ParamStack->GetParam<float>("SomeLatentFloat") == 34.0f, "FAnimationAnimNextRuntimeTest_GraphExecute -> Unexpected SomeLatentFloat value");
+
+		ParamStack->PopLayer(LayerHandle);
+		GraphInstance.Release();
+
+		FParamStack::DetachFromCurrentThread();
+	}
+
+	Tests::FUtils::CleanupAfterTests();
+
+	return true;
+}
+
+IMPLEMENT_SIMPLE_AUTOMATION_TEST(FAnimationAnimNextRuntimeTest_GraphExecuteLatent, "Animation.AnimNext.Runtime.Graph.ExecuteLatent", EAutomationTestFlags::EditorContext | EAutomationTestFlags::EngineFilter)
+
+bool FAnimationAnimNextRuntimeTest_GraphExecuteLatent::RunTest(const FString& InParameters)
+{
+	using namespace UE::AnimNext;
+
+	{
+		AUTO_REGISTER_ANIM_TRAIT(FTestTrait)
+		FScopedClearNodeTemplateRegistry ScopedClearNodeTemplateRegistry;
+
+		UFactory* GraphFactory = NewObject<UAnimNextGraphFactory>();
+		UAnimNextGraph* AnimNextGraph = CastChecked<UAnimNextGraph>(GraphFactory->FactoryCreateNew(UAnimNextGraph::StaticClass(), GetTransientPackage(), TEXT("TestAnimNextGraph"), RF_Transient, nullptr, nullptr, NAME_None));
+		UE_RETURN_ON_ERROR(AnimNextGraph != nullptr, "FAnimationAnimNextRuntimeTest_GraphExecuteLatent -> Failed to create animation graph");
+
+		UAnimNextGraph_EditorData* EditorData = UncookedOnly::FUtils::GetEditorData(AnimNextGraph);
+		UE_RETURN_ON_ERROR(EditorData != nullptr, "FAnimationAnimNextRuntimeTest_GraphExecuteLatent -> Failed to find animation graph editor data");
+
+		UE_RETURN_ON_ERROR(EditorData->AddAnimationGraph(FRigUnit_AnimNextGraphRoot::DefaultEntryPoint, false) != nullptr, "FAnimationAnimNextRuntimeTest_GraphExecuteLatent -> Failed to add animation graph");
+
+		URigVMController* Controller = EditorData->GetRigVMClient()->GetController(EditorData->GetRigVMClient()->GetDefaultModel());
+		UE_RETURN_ON_ERROR(Controller != nullptr, "FAnimationAnimNextRuntimeTest_GraphExecuteLatent -> Failed to get RigVM controller");
+
+		// Find graph entry point
+		URigVMNode* MainEntryPointNode = Controller->GetGraph()->FindNodeByName(FRigUnit_AnimNextGraphRoot::StaticStruct()->GetFName());
+		UE_RETURN_ON_ERROR(MainEntryPointNode != nullptr, "FAnimationAnimNextRuntimeTest_GraphExecuteLatent -> Failed to find main entry point node");
+
+		URigVMPin* BeginExecutePin = MainEntryPointNode->FindPin(GET_MEMBER_NAME_STRING_CHECKED(FRigUnit_AnimNextGraphRoot, Result));
+		UE_RETURN_ON_ERROR(BeginExecutePin != nullptr && BeginExecutePin->GetDirection() == ERigVMPinDirection::Input, "FAnimationAnimNextRuntimeTest_GraphExecuteLatent -> Failed to create entry point");
+
+		// Create an empty decorator stack node
+		URigVMUnitNode* DecoratorStackNode = Controller->AddUnitNode(FRigUnit_AnimNextTraitStack::StaticStruct(), FRigVMStruct::ExecuteName, FVector2D(0.0f, 0.0f), FString(), false);
+		UE_RETURN_ON_ERROR(DecoratorStackNode != nullptr, "FAnimationAnimNextRuntimeTest_GraphExecuteLatent -> Failed to create decorator stack node");
+
+		// Link our stack result to our entry point
+		Controller->AddLink(DecoratorStackNode->GetPins()[0], MainEntryPointNode->GetPins()[0]);
+
+		// Add a trait
+		const UScriptStruct* CppDecoratorStruct = FRigDecorator_AnimNextCppDecorator::StaticStruct();
+		UE_RETURN_ON_ERROR(CppDecoratorStruct != nullptr, "FAnimationAnimNextRuntimeTest_GraphExecuteLatent -> Failed to get find Cpp decorator static struct");
+
+		const FTrait* Trait = FTraitRegistry::Get().Find(FTestTrait::TraitUID);
+		UE_RETURN_ON_ERROR(Trait != nullptr, "FAnimationAnimNextRuntimeTest_GraphExecuteLatent -> Failed to find test trait");
+
+		UScriptStruct* ScriptStruct = Trait->GetTraitSharedDataStruct();
+		UE_RETURN_ON_ERROR(ScriptStruct != nullptr, "FAnimationAnimNextRuntimeTest_GraphExecuteLatent -> Failed to find trait shared data struct");
+
+		FString DefaultValue;
+		{
+			const FRigDecorator_AnimNextCppDecorator DefaultCppDecoratorStructInstance;
+			FRigDecorator_AnimNextCppDecorator CppDecoratorStructInstance;
+			CppDecoratorStructInstance.DecoratorSharedDataStruct = ScriptStruct;
+
+			UE_RETURN_ON_ERROR(CppDecoratorStructInstance.CanBeAddedToNode(DecoratorStackNode, nullptr), "FAnimationAnimNextRuntimeTest_GraphExecuteLatent -> Trait cannot be added to decorator stack node");
+
+			const FProperty* Prop = FAnimNextCppDecoratorWrapper::StaticStruct()->FindPropertyByName(GET_MEMBER_NAME_STRING_CHECKED(FAnimNextCppDecoratorWrapper, CppDecorator));
+			UE_RETURN_ON_ERROR(Prop != nullptr, "FAnimationAnimNextRuntimeTest_GraphExecuteLatent -> Failed to find wrapper property");
+
+			Prop->ExportText_Direct(DefaultValue, &CppDecoratorStructInstance, &DefaultCppDecoratorStructInstance, nullptr, PPF_None);
+		}
+
+		FString DisplayNameMetadata;
+		ScriptStruct->GetStringMetaDataHierarchical(FRigVMStruct::DisplayNameMetaName, &DisplayNameMetadata);
+		const FString DisplayName = DisplayNameMetadata.IsEmpty() ? Trait->GetTraitName() : DisplayNameMetadata;
+
+		const FName DecoratorName = Controller->AddDecorator(
+			DecoratorStackNode->GetFName(),
+			*CppDecoratorStruct->GetPathName(),
+			*DisplayName,
+			DefaultValue, INDEX_NONE, true, true);
+		UE_RETURN_ON_ERROR(DecoratorName == DisplayName, TEXT("FAnimationAnimNextRuntimeTest_GraphExecuteLatent -> Unexpected decorator name"));
+
+		// Set some values on our decorator
+		URigVMPin* DecoratorPin = DecoratorStackNode->FindPin(*DisplayName);
+		UE_RETURN_ON_ERROR(DecoratorPin != nullptr, TEXT("FAnimationAnimNextRuntimeTest_GraphExecuteLatent -> Failed to find decorator pin"));
+
+		Controller->SetPinDefaultValue(DecoratorPin->GetSubPins()[1]->GetPinPath(), TEXT("78"));		// SomeInt32
+		Controller->SetPinDefaultValue(DecoratorPin->GetSubPins()[2]->GetPinPath(), TEXT("142.33"));	// SomeFloat
+		Controller->SetPinDefaultValue(DecoratorPin->GetSubPins()[5]->GetPinPath(), TEXT("1123.31"));	// SomeLatentFloat, inline value on latent pin
+
+		// Set some latent values on our decorator
+		{
+			FRigVMFunction_MathIntAdd IntAdd;
+			IntAdd.A = 10;
+			IntAdd.B = 23;
+
+			URigVMUnitNode* IntAddNode = Controller->AddUnitNodeWithDefaults(FRigVMFunction_MathIntAdd::StaticStruct(), FRigStructScope(IntAdd), FRigVMStruct::ExecuteName, FVector2D::ZeroVector, FString(), false);
+			UE_RETURN_ON_ERROR(IntAddNode != nullptr, TEXT("FAnimationAnimNextRuntimeTest_GraphExecuteLatent -> Failed to create Int add node"));
+
+			Controller->AddLink(
+				IntAddNode->FindPin(GET_MEMBER_NAME_STRING_CHECKED(FRigVMFunction_MathIntAdd, Result)),
+				DecoratorPin->GetSubPins()[3]);	// SomeLatentInt32
+		}
+
+		{
+			const FRigVMDispatchFactory* GetParameterFactory = FRigVMRegistry::Get().RegisterFactory(FRigVMDispatch_GetParameter::StaticStruct());
+			const FName GetParameterNotation = GetParameterFactory->GetTemplate()->GetNotation();
+
+			URigVMNode* GetParameterNode = Controller->AddTemplateNode(GetParameterNotation);
+			UE_RETURN_ON_ERROR(GetParameterNode != nullptr, TEXT("FAnimationAnimNextRuntimeTest_GraphExecuteLatent -> Failed to create GetParameter node"));
+
+			UE_RETURN_ON_ERROR(Controller->ResolveWildCardPin(GetParameterNode->FindPin(TEXT("Value")), RigVMTypeUtils::TypeIndex::Int32), TEXT("FAnimationAnimNextRuntimeTest_GraphExecuteLatent -> Failed to resolve wildcard for GetParameter node"));
+			Controller->SetPinDefaultValue(GetParameterNode->FindPin(TEXT("Parameter"))->GetPinPath(), TEXT("SomeSourceInt"));
+
+			Controller->AddLink(
+				GetParameterNode->FindPin(TEXT("Value")),
+				DecoratorPin->GetSubPins()[4]);	// SomeOtherLatentInt32
+		}
+
+		TSharedRef<FParamStack> ParamStack = MakeShared<FParamStack>();
+		FParamStack::AttachToCurrentThread(ParamStack);
+
+		FAnimNextGraphInstancePtr GraphInstance;
+		AnimNextGraph->AllocateInstance(GraphInstance);
+
+		FParamStack::FPushedLayerHandle LayerHandle = ParamStack->PushValues(
+			"UpdateCount", (int32)0,
+			"EvaluateCount", (int32)0,
+			"SomeSourceInt", (int32)1223,
+			"SomeInt32", (int32)0,
+			"SomeFloat", 0.0f,
+			"SomeLatentInt32", (int32)0,
+			"SomeOtherLatentInt32", (int32)0,
+			"SomeLatentFloat", 0.0f
+		);
+
+		{
+			UE::AnimNext::UpdateGraph(GraphInstance, 1.0f / 30.0f);
+			(void)UE::AnimNext::EvaluateGraph(GraphInstance);
+		}
+
+		AddErrorIfFalse(ParamStack->GetParam<int32>("UpdateCount") == 1, "FAnimationAnimNextRuntimeTest_GraphExecuteLatent -> Unexpected update count");
+		AddErrorIfFalse(ParamStack->GetParam<int32>("EvaluateCount") == 1, "FAnimationAnimNextRuntimeTest_GraphExecuteLatent -> Unexpected evaluate count");
+		AddErrorIfFalse(ParamStack->GetParam<int32>("SomeInt32") == 78, "FAnimationAnimNextRuntimeTest_GraphExecuteLatent -> Unexpected SomeInt32 value");
+		AddErrorIfFalse(ParamStack->GetParam<float>("SomeFloat") == 142.33f, "FAnimationAnimNextRuntimeTest_GraphExecuteLatent -> Unexpected SomeFloat value");
+		AddErrorIfFalse(ParamStack->GetParam<int32>("SomeLatentInt32") == 33, "FAnimationAnimNextRuntimeTest_GraphExecuteLatent -> Unexpected SomeLatentInt32 value");
+		AddErrorIfFalse(ParamStack->GetParam<int32>("SomeOtherLatentInt32") == 1223, "FAnimationAnimNextRuntimeTest_GraphExecuteLatent -> Unexpected SomeOtherLatentInt32 value");
+		AddErrorIfFalse(ParamStack->GetParam<float>("SomeLatentFloat") == 1123.31f, "FAnimationAnimNextRuntimeTest_GraphExecuteLatent -> Unexpected SomeLatentFloat value");
+
+		ParamStack->PopLayer(LayerHandle);
+		GraphInstance.Release();
+
+		FParamStack::DetachFromCurrentThread();
+	}
+
+	Tests::FUtils::CleanupAfterTests();
+
+	return true;
+}
+
+#endif
diff --git a/Engine/Plugins/Experimental/AnimNextAnimGraph/Source/AnimNextAnimGraphTestSuite/Private/AnimNextAnimGraphTraitGraphTest.h b/Engine/Plugins/Experimental/AnimNextAnimGraph/Source/AnimNextAnimGraphTestSuite/Private/AnimNextAnimGraphTraitGraphTest.h
new file mode 100644
--- /dev/null
+++ b/Engine/Plugins/Experimental/AnimNextAnimGraph/Source/AnimNextAnimGraphTestSuite/Private/AnimNextAnimGraphTraitGraphTest.h
@@ -0,0 +1,40 @@
+// Copyright Epic Games, Inc. All Rights Reserved.
+
+#pragma once
+
+#include "CoreMinimal.h"
+
+#include "TraitCore/TraitBinding.h"
+#include "TraitCore/TraitSharedData.h"
+#include "Graph/AnimNextGraph.h"
+
+#include "AnimNextAnimGraphTraitGraphTest.generated.h"
+
+USTRUCT()
+struct FTestTraitSharedData : public FAnimNextTraitSharedData
+{
+	GENERATED_BODY()
+
+	UPROPERTY(meta = (Input, Inline))
+	int32 SomeInt32 = 3;
+
+	UPROPERTY(meta = (Input, Inline))
+	float SomeFloat = 34.0f;
+
+	UPROPERTY(meta = (Input))
+	int32 SomeLatentInt32 = 3;			// MathAdd with constants, latent
+
+	UPROPERTY(meta = (Input))
+	int32 SomeOtherLatentInt32 = 3;		// GetParameter, latent
+
+	UPROPERTY(meta = (Input))
+	float SomeLatentFloat = 34.0f;		// Inline value, not latent
+
+	#define TRAIT_LATENT_PROPERTIES_ENUMERATOR(GeneratorMacro) \
+		GeneratorMacro(SomeLatentInt32) \
+		GeneratorMacro(SomeOtherLatentInt32) \
+		GeneratorMacro(SomeLatentFloat) \
+
+	GENERATE_TRAIT_LATENT_PROPERTIES(FTestTraitSharedData, TRAIT_LATENT_PROPERTIES_ENUMERATOR)
+	#undef TRAIT_LATENT_PROPERTIES_ENUMERATOR
+};
diff --git a/Engine/Plugins/Experimental/AnimNextAnimGraph/Source/AnimNextAnimGraphTestSuite/Private/AnimNextAnimGraphTraitInterfacesTest.cpp b/Engine/Plugins/Experimental/AnimNextAnimGraph/Source/AnimNextAnimGraphTestSuite/Private/AnimNextAnimGraphTraitInterfacesTest.cpp
new file mode 100644
--- /dev/null
+++ b/Engine/Plugins/Experimental/AnimNextAnimGraph/Source/AnimNextAnimGraphTestSuite/Private/AnimNextAnimGraphTraitInterfacesTest.cpp
@@ -0,0 +1,1005 @@
+// Copyright Epic Games, Inc. All Rights Reserved.
+
+#include "AnimNextAnimGraphTraitInterfacesTest.h"
+#include "AnimNextRuntimeTest.h"
+#include "AnimNextTest.h"
+
+#include "Misc/AutomationTest.h"
+
+#if WITH_DEV_AUTOMATION_TESTS
+
+#include "Serialization/MemoryReader.h"
+
+#include "TraitCore/Trait.h"
+#include "TraitCore/TraitReader.h"
+#include "TraitCore/TraitRegistry.h"
+#include "TraitCore/TraitWriter.h"
+#include "TraitCore/ExecutionContext.h"
+#include "TraitCore/ITraitInterface.h"
+#include "TraitCore/IScopedTraitInterface.h"
+#include "TraitCore/NodeInstance.h"
+#include "TraitCore/NodeTemplateBuilder.h"
+#include "TraitCore/NodeTemplateRegistry.h"
+#include "TraitInterfaces/IEvaluate.h"
+#include "TraitInterfaces/IHierarchy.h"
+#include "TraitInterfaces/IUpdate.h"
+#include "Graph/AnimNextGraph.h"
+#include "Graph/GraphFactory.h"
+
+//****************************************************************************
+// AnimNext Runtime TraitInterfaces Tests
+//****************************************************************************
+
+namespace UE::AnimNext
+{
+	namespace Private
+	{
+		static TArray<FTraitUID>* UpdatedTraits = nullptr;
+		static TArray<FTraitUID>* EvaluatedTraits = nullptr;
+
+		static FName TestTag(TEXT("MyTag"));
+		static TArray<bool>* IsTagInScope = nullptr;
+		static bool AutoPopTag = false;
+	}
+
+	struct FTraitWithNoChildren : FBaseTrait, IUpdate, IEvaluate
+	{
+		DECLARE_ANIM_TRAIT(FTraitWithNoChildren, 0x79080a53, FBaseTrait)
+
+		// IUpdate impl
+		virtual void PreUpdate(FUpdateTraversalContext& Context, const TTraitBinding<IUpdate>& Binding, const FTraitUpdateState& TraitState) const override
+		{
+			if (Private::UpdatedTraits != nullptr)
+			{
+				Private::UpdatedTraits->Add(FTraitWithNoChildren::TraitUID);
+			}
+
+			IUpdate::PreUpdate(Context, Binding, TraitState);
+		}
+
+		virtual void PostUpdate(FUpdateTraversalContext& Context, const TTraitBinding<IUpdate>& Binding, const FTraitUpdateState& TraitState) const override
+		{
+			if (Private::UpdatedTraits != nullptr)
+			{
+				Private::UpdatedTraits->Add(FTraitWithNoChildren::TraitUID);
+			}
+
+			IUpdate::PostUpdate(Context, Binding, TraitState);
+		}
+
+		// IEvaluate impl
+		virtual void PreEvaluate(FEvaluateTraversalContext& Context, const TTraitBinding<IEvaluate>& Binding) const override
+		{
+			if (Private::EvaluatedTraits != nullptr)
+			{
+				Private::EvaluatedTraits->Add(FTraitWithNoChildren::TraitUID);
+			}
+
+			IEvaluate::PreEvaluate(Context, Binding);
+		}
+
+		virtual void PostEvaluate(FEvaluateTraversalContext& Context, const TTraitBinding<IEvaluate>& Binding) const override
+		{
+			if (Private::EvaluatedTraits != nullptr)
+			{
+				Private::EvaluatedTraits->Add(FTraitWithNoChildren::TraitUID);
+			}
+
+			IEvaluate::PostEvaluate(Context, Binding);
+		}
+	};
+
+	// Trait implementation boilerplate
+	#define TRAIT_INTERFACE_ENUMERATOR(GeneratorMacro) \
+		GeneratorMacro(IEvaluate) \
+		GeneratorMacro(IUpdate) \
+
+	GENERATE_ANIM_TRAIT_IMPLEMENTATION(FTraitWithNoChildren, TRAIT_INTERFACE_ENUMERATOR)
+	#undef TRAIT_INTERFACE_ENUMERATOR
+
+	// This trait does not update or evaluate
+	struct FTraitWithOneChild : FBaseTrait, IHierarchy
+	{
+		DECLARE_ANIM_TRAIT(FTraitWithOneChild, 0xab680b35, FBaseTrait)
+
+		using FSharedData = FTraitWithOneChildSharedData;
+
+		struct FInstanceData : FTrait::FInstanceData
+		{
+			FTraitPtr Child;
+
+			void Construct(const FExecutionContext& Context, const FTraitBinding& Binding)
+			{
+				Child = Context.AllocateNodeInstance(Binding.GetTraitPtr(), Binding.GetSharedData<FSharedData>()->Child);
+			}
+		};
+
+		// IHierarchy impl
+		virtual uint32 GetNumChildren(const FExecutionContext& Context, const TTraitBinding<IHierarchy>& Binding) const override
+		{
+			return 1;
+		}
+
+		virtual void GetChildren(const FExecutionContext& Context, const TTraitBinding<IHierarchy>& Binding, FChildrenArray& Children) const override
+		{
+			const FInstanceData* InstanceData = Binding.GetInstanceData<FInstanceData>();
+
+			Children.Add(InstanceData->Child);
+		}
+	};
+
+	// Trait implementation boilerplate
+	#define TRAIT_INTERFACE_ENUMERATOR(GeneratorMacro) \
+		GeneratorMacro(IHierarchy) \
+
+	GENERATE_ANIM_TRAIT_IMPLEMENTATION(FTraitWithOneChild, TRAIT_INTERFACE_ENUMERATOR)
+	#undef TRAIT_INTERFACE_ENUMERATOR
+
+	struct FTraitWithChildren : FBaseTrait, IHierarchy, IUpdate, IUpdateTraversal, IEvaluate
+	{
+		DECLARE_ANIM_TRAIT(FTraitWithChildren, 0x4b296948, FBaseTrait)
+
+		using FSharedData = FTraitWithChildrenSharedData;
+
+		struct FInstanceData : FTrait::FInstanceData
+		{
+			FTraitPtr Children[2];
+
+			void Construct(const FExecutionContext& Context, const FTraitBinding& Binding)
+			{
+				Children[0] = Context.AllocateNodeInstance(Binding.GetTraitPtr(), Binding.GetSharedData<FSharedData>()->Children[0]);
+				Children[1] = Context.AllocateNodeInstance(Binding.GetTraitPtr(), Binding.GetSharedData<FSharedData>()->Children[1]);
+			}
+		};
+
+		// IHierarchy impl
+		virtual uint32 GetNumChildren(const FExecutionContext& Context, const TTraitBinding<IHierarchy>& Binding) const override
+		{
+			return 2;
+		}
+
+		virtual void GetChildren(const FExecutionContext& Context, const TTraitBinding<IHierarchy>& Binding, FChildrenArray& Children) const override
+		{
+			const FInstanceData* InstanceData = Binding.GetInstanceData<FInstanceData>();
+
+			Children.Add(InstanceData->Children[0]);
+			Children.Add(InstanceData->Children[1]);
+		}
+
+		// IUpdate impl
+		virtual void PreUpdate(FUpdateTraversalContext& Context, const TTraitBinding<IUpdate>& Binding, const FTraitUpdateState& TraitState) const override
+		{
+			if (Private::UpdatedTraits != nullptr)
+			{
+				Private::UpdatedTraits->Add(FTraitWithChildren::TraitUID);
+			}
+
+			IUpdate::PreUpdate(Context, Binding, TraitState);
+		}
+
+		virtual void PostUpdate(FUpdateTraversalContext& Context, const TTraitBinding<IUpdate>& Binding, const FTraitUpdateState& TraitState) const override
+		{
+			if (Private::UpdatedTraits != nullptr)
+			{
+				Private::UpdatedTraits->Add(FTraitWithChildren::TraitUID);
+			}
+
+			IUpdate::PostUpdate(Context, Binding, TraitState);
+		}
+
+		// IUpdateTraversal impl
+		virtual void QueueChildrenForTraversal(FUpdateTraversalContext& Context, const TTraitBinding<IUpdateTraversal>& Binding, const FTraitUpdateState& TraitState, FUpdateTraversalQueue& TraversalQueue) const override
+		{
+			const FInstanceData* InstanceData = Binding.GetInstanceData<FInstanceData>();
+
+			TraversalQueue.Push(InstanceData->Children[0], TraitState);
+			TraversalQueue.Push(InstanceData->Children[1], TraitState);
+		}
+
+		// IEvaluate impl
+		virtual void PreEvaluate(FEvaluateTraversalContext& Context, const TTraitBinding<IEvaluate>& Binding) const override
+		{
+			if (Private::EvaluatedTraits != nullptr)
+			{
+				Private::EvaluatedTraits->Add(FTraitWithChildren::TraitUID);
+			}
+
+			IEvaluate::PreEvaluate(Context, Binding);
+		}
+
+		virtual void PostEvaluate(FEvaluateTraversalContext& Context, const TTraitBinding<IEvaluate>& Binding) const override
+		{
+			if (Private::EvaluatedTraits != nullptr)
+			{
+				Private::EvaluatedTraits->Add(FTraitWithChildren::TraitUID);
+			}
+
+			IEvaluate::PostEvaluate(Context, Binding);
+		}
+	};
+
+	// Trait implementation boilerplate
+	#define TRAIT_INTERFACE_ENUMERATOR(GeneratorMacro) \
+		GeneratorMacro(IEvaluate) \
+		GeneratorMacro(IHierarchy) \
+		GeneratorMacro(IUpdate) \
+		GeneratorMacro(IUpdateTraversal) \
+
+	GENERATE_ANIM_TRAIT_IMPLEMENTATION(FTraitWithChildren, TRAIT_INTERFACE_ENUMERATOR)
+	#undef TRAIT_INTERFACE_ENUMERATOR
+
+	struct IScopedTagInterface : IScopedTraitInterface
+	{
+		DECLARE_ANIM_TRAIT_INTERFACE(IScopedTagInterface, 0xb8f4b397)
+
+		virtual FName GetTag(const FExecutionContext& Context, const TTraitBinding<IScopedTagInterface>& Binding) const;
+
+		static bool IsTagInScope(const FExecutionContext& Context, FName Tag);
+	};
+
+	template<>
+	struct TTraitBinding<IScopedTagInterface> : FTraitBinding
+	{
+		FName GetTag(const FExecutionContext& Context) const
+		{
+			return GetInterface()->GetTag(Context, *this);
+		}
+
+	protected:
+		const IScopedTagInterface* GetInterface() const { return GetInterfaceTyped<IScopedTagInterface>(); }
+	};
+
+	FName IScopedTagInterface::GetTag(const FExecutionContext& Context, const TTraitBinding<IScopedTagInterface>& Binding) const
+	{
+		TTraitBinding<IScopedTagInterface> SuperBinding;
+		if (Binding.GetStackInterfaceSuper(SuperBinding))
+		{
+			return SuperBinding.GetTag(Context);
+		}
+
+		return NAME_None;
+	}
+
+	bool IScopedTagInterface::IsTagInScope(const FExecutionContext& Context, FName Tag)
+	{
+		bool bResult = false;
+
+		Context.ForEachScopedInterface<IScopedTagInterface>([&Context, Tag, &bResult](TTraitBinding<IScopedTagInterface>& InterfaceBinding)
+			{
+				if (Tag == InterfaceBinding.GetTag(Context))
+				{
+					// We found our tag, stop iterating
+					bResult = true;
+					return false;
+				}
+
+				// Keep searching
+				return true;
+			});
+
+		return bResult;
+	}
+
+	// Adds a scoped tag
+	struct FScopedTagTrait : FAdditiveTrait, IScopedTagInterface, IUpdate
+	{
+		DECLARE_ANIM_TRAIT(FScopedTagTrait, 0x4b296948, FAdditiveTrait)
+
+		// IScopedTagInterface impl
+		virtual FName GetTag(const FExecutionContext& Context, const TTraitBinding<IScopedTagInterface>& Binding) const override
+		{
+			return Private::TestTag;
+		}
+
+		// IUpdate impl
+		virtual void PreUpdate(FUpdateTraversalContext& Context, const TTraitBinding<IUpdate>& Binding, const FTraitUpdateState& TraitState) const override
+		{
+			Context.PushScopedInterface<IScopedTagInterface>(Binding);
+
+			IUpdate::PreUpdate(Context, Binding, TraitState);
+		}
+
+		virtual void PostUpdate(FUpdateTraversalContext& Context, const TTraitBinding<IUpdate>& Binding, const FTraitUpdateState& TraitState) const override
+		{
+			if (!Private::AutoPopTag)
+			{
+				ensure(Context.PopScopedInterface<IScopedTagInterface>(Binding));
+			}
+
+			IUpdate::PostUpdate(Context, Binding, TraitState);
+		}
+	};
+
+	// Trait implementation boilerplate
+	#define TRAIT_INTERFACE_ENUMERATOR(GeneratorMacro) \
+		GeneratorMacro(IScopedTagInterface) \
+		GeneratorMacro(IUpdate) \
+
+	GENERATE_ANIM_TRAIT_IMPLEMENTATION(FScopedTagTrait, TRAIT_INTERFACE_ENUMERATOR)
+	#undef TRAIT_INTERFACE_ENUMERATOR
+
+	// Tests if we have a scoped tag
+	struct FTestScopedTagTrait : FAdditiveTrait, IUpdate
+	{
+		DECLARE_ANIM_TRAIT(FTestScopedTagTrait, 0x922bd23d, FAdditiveTrait)
+
+		// IUpdate impl
+		virtual void PreUpdate(FUpdateTraversalContext& Context, const TTraitBinding<IUpdate>& Binding, const FTraitUpdateState& TraitState) const override
+		{
+			if (Private::IsTagInScope != nullptr)
+			{
+				Private::IsTagInScope->Add(IScopedTagInterface::IsTagInScope(Context, Private::TestTag));
+			}
+
+			IUpdate::PreUpdate(Context, Binding, TraitState);
+		}
+
+		virtual void PostUpdate(FUpdateTraversalContext& Context, const TTraitBinding<IUpdate>& Binding, const FTraitUpdateState& TraitState) const override
+		{
+			if (Private::IsTagInScope != nullptr)
+			{
+				Private::IsTagInScope->Add(IScopedTagInterface::IsTagInScope(Context, Private::TestTag));
+			}
+
+			IUpdate::PostUpdate(Context, Binding, TraitState);
+		}
+	};
+
+	// Trait implementation boilerplate
+	#define TRAIT_INTERFACE_ENUMERATOR(GeneratorMacro) \
+		GeneratorMacro(IUpdate) \
+
+	GENERATE_ANIM_TRAIT_IMPLEMENTATION(FTestScopedTagTrait, TRAIT_INTERFACE_ENUMERATOR)
+	#undef TRAIT_INTERFACE_ENUMERATOR
+}
+
+IMPLEMENT_SIMPLE_AUTOMATION_TEST(FAnimationAnimNextRuntimeTest_IHierarchy, "Animation.AnimNext.Runtime.IHierarchy", EAutomationTestFlags::EditorContext | EAutomationTestFlags::EngineFilter)
+
+bool FAnimationAnimNextRuntimeTest_IHierarchy::RunTest(const FString& InParameters)
+{
+	using namespace UE::AnimNext;
+
+	{
+		AUTO_REGISTER_ANIM_TRAIT(FTraitWithNoChildren)
+		AUTO_REGISTER_ANIM_TRAIT(FTraitWithOneChild)
+		AUTO_REGISTER_ANIM_TRAIT(FTraitWithChildren)
+
+		UFactory* GraphFactory = NewObject<UAnimNextGraphFactory>();
+		UAnimNextGraph* AnimNextGraph = CastChecked<UAnimNextGraph>(GraphFactory->FactoryCreateNew(UAnimNextGraph::StaticClass(), GetTransientPackage(), TEXT("TestAnimNextGraph"), RF_Transient, nullptr, nullptr, NAME_None));
+		UE_RETURN_ON_ERROR(AnimNextGraph != nullptr, "FAnimationAnimNextRuntimeTest_IHierarchy -> Failed to create animation graph");
+
+		FScopedClearNodeTemplateRegistry ScopedClearNodeTemplateRegistry;
+		FNodeTemplateRegistry& Registry = FNodeTemplateRegistry::Get();
+
+		// We create a few node templates
+		// Template A has a single trait with no children
+		TArray<FTraitUID> NodeTemplateTraitListA;
+		NodeTemplateTraitListA.Add(FTraitWithNoChildren::TraitUID);
+
+		// Template B has a single trait with one child
+		TArray<FTraitUID> NodeTemplateTraitListB;
+		NodeTemplateTraitListB.Add(FTraitWithOneChild::TraitUID);
+
+		// Template C has two traits, each with one child
+		TArray<FTraitUID> NodeTemplateTraitListC;
+		NodeTemplateTraitListC.Add(FTraitWithOneChild::TraitUID);
+		NodeTemplateTraitListC.Add(FTraitWithOneChild::TraitUID);
+
+		// Template D has a single trait with children
+		TArray<FTraitUID> NodeTemplateTraitListD;
+		NodeTemplateTraitListD.Add(FTraitWithChildren::TraitUID);
+
+		// Populate our node template registry
+		TArray<uint8> NodeTemplateBufferA, NodeTemplateBufferB, NodeTemplateBufferC, NodeTemplateBufferD;
+		const FNodeTemplate* NodeTemplateA = FNodeTemplateBuilder::BuildNodeTemplate(NodeTemplateTraitListA, NodeTemplateBufferA);
+		const FNodeTemplate* NodeTemplateB = FNodeTemplateBuilder::BuildNodeTemplate(NodeTemplateTraitListB, NodeTemplateBufferB);
+		const FNodeTemplate* NodeTemplateC = FNodeTemplateBuilder::BuildNodeTemplate(NodeTemplateTraitListC, NodeTemplateBufferC);
+		const FNodeTemplate* NodeTemplateD = FNodeTemplateBuilder::BuildNodeTemplate(NodeTemplateTraitListD, NodeTemplateBufferD);
+
+		// Build our graph, it as follow (each node template has a single node instance):
+		// NodeA has no children
+		// NodeB has one child: NodeA
+		// NodeC has two children: NodeA and NodeB (but both traits are base, only NodeB will be referenced)
+		// NodeD has two children: NodeA and NodeC
+
+		TArray<FNodeHandle> NodeHandles;
+
+		// Write our graph
+		TArray<uint8> GraphSharedDataArchiveBuffer;
+		TArray<TObjectPtr<UObject>> GraphReferencedObjects;
+		{
+			FTraitWriter TraitWriter;
+
+			NodeHandles.Add(TraitWriter.RegisterNode(*NodeTemplateA));
+			NodeHandles.Add(TraitWriter.RegisterNode(*NodeTemplateB));
+			NodeHandles.Add(TraitWriter.RegisterNode(*NodeTemplateC));
+			NodeHandles.Add(TraitWriter.RegisterNode(*NodeTemplateD));
+
+			// We don't have trait properties
+			TArray<TMap<FName, FString>> TraitPropertiesA;
+			TraitPropertiesA.AddDefaulted(NodeTemplateTraitListA.Num());
+
+			TArray<TMap<FName, FString>> TraitPropertiesB;
+			TraitPropertiesB.AddDefaulted(NodeTemplateTraitListB.Num());
+			TraitPropertiesB[0].Add(TEXT("Child"), ToString<FTraitWithOneChild::FSharedData>(TEXT("Child"), FAnimNextTraitHandle(NodeHandles[0])));
+
+			TArray<TMap<FName, FString>> TraitPropertiesC;
+			TraitPropertiesC.AddDefaulted(NodeTemplateTraitListC.Num());
+			TraitPropertiesC[0].Add(TEXT("Child"), ToString<FTraitWithOneChild::FSharedData>(TEXT("Child"), FAnimNextTraitHandle(NodeHandles[0])));
+			TraitPropertiesC[1].Add(TEXT("Child"), ToString<FTraitWithOneChild::FSharedData>(TEXT("Child"), FAnimNextTraitHandle(NodeHandles[1])));
+
+			TArray<TMap<FName, FString>> TraitPropertiesD;
+			TraitPropertiesD.AddDefaulted(NodeTemplateTraitListD.Num());
+			FAnimNextTraitHandle ChildrenHandlesD[2] = { FAnimNextTraitHandle(NodeHandles[0]), FAnimNextTraitHandle(NodeHandles[2], 1)};
+			TraitPropertiesD[0].Add(TEXT("Children"), ToString<FTraitWithChildren::FSharedData>(TEXT("Children"), ChildrenHandlesD));
+
+			TraitWriter.BeginNodeWriting();
+			TraitWriter.WriteNode(NodeHandles[0],
+				[&TraitPropertiesA](uint32 TraitIndex, FName PropertyName)
+				{
+					return TraitPropertiesA[TraitIndex][PropertyName];
+				},
+				[](uint32 TraitIndex, FName PropertyName)
+				{
+					return MAX_uint16;
+				});
+			TraitWriter.WriteNode(NodeHandles[1],
+				[&TraitPropertiesB](uint32 TraitIndex, FName PropertyName)
+				{
+					return TraitPropertiesB[TraitIndex][PropertyName];
+				},
+				[](uint32 TraitIndex, FName PropertyName)
+				{
+					return MAX_uint16;
+				});
+			TraitWriter.WriteNode(NodeHandles[2],
+				[&TraitPropertiesC](uint32 TraitIndex, FName PropertyName)
+				{
+					return TraitPropertiesC[TraitIndex][PropertyName];
+				},
+				[](uint32 TraitIndex, FName PropertyName)
+				{
+					return MAX_uint16;
+				});
+			TraitWriter.WriteNode(NodeHandles[3],
+				[&TraitPropertiesD](uint32 TraitIndex, FName PropertyName)
+				{
+					return TraitPropertiesD[TraitIndex][PropertyName];
+				},
+				[](uint32 TraitIndex, FName PropertyName)
+				{
+					return MAX_uint16;
+				});
+			TraitWriter.EndNodeWriting();
+
+			AddErrorIfFalse(TraitWriter.GetErrorState() == FTraitWriter::EErrorState::None, "FAnimationAnimNextRuntimeTest_IHierarchy -> Failed to write traits");
+			GraphSharedDataArchiveBuffer = TraitWriter.GetGraphSharedData();
+			GraphReferencedObjects = TraitWriter.GetGraphReferencedObjects();
+		}
+
+		// Read our graph
+		FTestUtils::LoadFromArchiveBuffer(*AnimNextGraph, NodeHandles, GraphSharedDataArchiveBuffer);
+
+		FAnimNextGraphInstancePtr GraphInstance;
+		AnimNextGraph->AllocateInstance(GraphInstance);
+
+		FExecutionContext Context(GraphInstance);
+
+		{
+			FMemMark Mark(FMemStack::Get());
+
+			FAnimNextTraitHandle RootHandle(NodeHandles[3]);	// Point to NodeD, first base trait
+
+			FTraitPtr NodeDPtr = Context.AllocateNodeInstance(*GraphInstance.GetImpl(), RootHandle);
+			AddErrorIfFalse(NodeDPtr.IsValid(), "FAnimationAnimNextRuntimeTest_IHierarchy -> Failed to allocate root node instance");
+
+			FTraitStackBinding StackNodeD;
+			AddErrorIfFalse(Context.GetStack(NodeDPtr, StackNodeD), "FAnimationAnimNextRuntimeTest_IHierarchy -> Failed to bind to trait stack");
+
+			TTraitBinding<IHierarchy> HierarchyBindingNodeD;
+			AddErrorIfFalse(StackNodeD.GetInterface(HierarchyBindingNodeD), "FAnimationAnimNextRuntimeTest_IHierarchy -> IHierarchy not found");
+
+			FChildrenArray ChildrenNodeD;
+			HierarchyBindingNodeD.GetChildren(Context, ChildrenNodeD);
+
+			AddErrorIfFalse(ChildrenNodeD.Num() == 2, "FAnimationAnimNextRuntimeTest_IHierarchy -> Expected 2 children");
+			AddErrorIfFalse(HierarchyBindingNodeD.GetNumChildren(Context) == 2, "FAnimationAnimNextRuntimeTest_IHierarchy -> Expected 2 children");
+			AddErrorIfFalse(ChildrenNodeD[0].IsValid() && ChildrenNodeD[0].GetNodeInstance()->GetNodeHandle() == NodeHandles[0], "FAnimationAnimNextRuntimeTest_IHierarchy -> Expected child: NodeA");
+			AddErrorIfFalse(ChildrenNodeD[1].IsValid() && ChildrenNodeD[1].GetNodeInstance()->GetNodeHandle() == NodeHandles[2], "FAnimationAnimNextRuntimeTest_IHierarchy -> Expected child: NodeC");
+
+			ChildrenNodeD.Reset();
+			IHierarchy::GetStackChildren(Context, StackNodeD, ChildrenNodeD);
+
+			AddErrorIfFalse(ChildrenNodeD.Num() == 2, "FAnimationAnimNextRuntimeTest_IHierarchy -> Expected 2 children");
+			AddErrorIfFalse(IHierarchy::GetNumStackChildren(Context, StackNodeD) == 2, "FAnimationAnimNextRuntimeTest_IHierarchy -> Expected 2 children");
+			AddErrorIfFalse(ChildrenNodeD[0].IsValid() && ChildrenNodeD[0].GetNodeInstance()->GetNodeHandle() == NodeHandles[0], "FAnimationAnimNextRuntimeTest_IHierarchy -> Expected child: NodeA");
+			AddErrorIfFalse(ChildrenNodeD[1].IsValid() && ChildrenNodeD[1].GetNodeInstance()->GetNodeHandle() == NodeHandles[2], "FAnimationAnimNextRuntimeTest_IHierarchy -> Expected child: NodeC");
+
+			{
+				FTraitStackBinding StackNodeC;
+				AddErrorIfFalse(Context.GetStack(ChildrenNodeD[1], StackNodeC), "FAnimationAnimNextRuntimeTest_IHierarchy -> Failed to bind to trait stack");
+
+				TTraitBinding<IHierarchy> HierarchyBindingNodeC;
+				AddErrorIfFalse(StackNodeC.GetInterface(HierarchyBindingNodeC), "FAnimationAnimNextRuntimeTest_IHierarchy -> IHierarchy not found");
+
+				FChildrenArray ChildrenNodeC;
+				HierarchyBindingNodeC.GetChildren(Context, ChildrenNodeC);
+
+				AddErrorIfFalse(ChildrenNodeC.Num() == 1, "FAnimationAnimNextRuntimeTest_IHierarchy -> Expected 1 child");
+				AddErrorIfFalse(HierarchyBindingNodeC.GetNumChildren(Context) == 1, "FAnimationAnimNextRuntimeTest_IHierarchy -> Expected 1 child");
+				AddErrorIfFalse(ChildrenNodeC[0].IsValid() && ChildrenNodeC[0].GetNodeInstance()->GetNodeHandle() == NodeHandles[1], "FAnimationAnimNextRuntimeTest_IHierarchy -> Expected child: NodeB");
+
+				ChildrenNodeC.Reset();
+				IHierarchy::GetStackChildren(Context, StackNodeC, ChildrenNodeC);
+
+				AddErrorIfFalse(ChildrenNodeC.Num() == 1, "FAnimationAnimNextRuntimeTest_IHierarchy -> Expected 1 child");
+				AddErrorIfFalse(IHierarchy::GetNumStackChildren(Context, StackNodeC) == 1, "FAnimationAnimNextRuntimeTest_IHierarchy -> Expected 1 child");
+				AddErrorIfFalse(ChildrenNodeC[0].IsValid() && ChildrenNodeC[0].GetNodeInstance()->GetNodeHandle() == NodeHandles[1], "FAnimationAnimNextRuntimeTest_IHierarchy -> Expected child: NodeB");
+
+				{
+					FTraitStackBinding StackNodeB;
+					AddErrorIfFalse(Context.GetStack(ChildrenNodeC[0], StackNodeB), "FAnimationAnimNextRuntimeTest_IHierarchy -> Failed to bind to trait stack");
+
+					TTraitBinding<IHierarchy> HierarchyBindingNodeB;
+					AddErrorIfFalse(StackNodeB.GetInterface(HierarchyBindingNodeB), "FAnimationAnimNextRuntimeTest_IHierarchy -> IHierarchy not found");
+
+					FChildrenArray ChildrenNodeB;
+					HierarchyBindingNodeB.GetChildren(Context, ChildrenNodeB);
+
+					AddErrorIfFalse(ChildrenNodeB.Num() == 1, "FAnimationAnimNextRuntimeTest_IHierarchy -> Expected 1 child");
+					AddErrorIfFalse(HierarchyBindingNodeB.GetNumChildren(Context) == 1, "FAnimationAnimNextRuntimeTest_IHierarchy -> Expected 1 child");
+					AddErrorIfFalse(ChildrenNodeB[0].IsValid() && ChildrenNodeB[0].GetNodeInstance()->GetNodeHandle() == NodeHandles[0], "FAnimationAnimNextRuntimeTest_IHierarchy -> Expected child: NodeA");
+
+					ChildrenNodeB.Reset();
+					IHierarchy::GetStackChildren(Context, StackNodeB, ChildrenNodeB);
+
+					AddErrorIfFalse(ChildrenNodeB.Num() == 1, "FAnimationAnimNextRuntimeTest_IHierarchy -> Expected 1 child");
+					AddErrorIfFalse(IHierarchy::GetNumStackChildren(Context, StackNodeB) == 1, "FAnimationAnimNextRuntimeTest_IHierarchy -> Expected 1 child");
+					AddErrorIfFalse(ChildrenNodeB[0].IsValid() && ChildrenNodeB[0].GetNodeInstance()->GetNodeHandle() == NodeHandles[0], "FAnimationAnimNextRuntimeTest_IHierarchy -> Expected child: NodeA");
+				}
+			}
+		}
+
+		Registry.Unregister(NodeTemplateA);
+		Registry.Unregister(NodeTemplateB);
+		Registry.Unregister(NodeTemplateC);
+		Registry.Unregister(NodeTemplateD);
+
+		AddErrorIfFalse(Registry.GetNum() == 0, "FAnimationAnimNextRuntimeTest_IHierarchy -> Registry should contain 0 templates");
+	}
+
+	Tests::FUtils::CleanupAfterTests();
+	
+	return true;
+}
+
+IMPLEMENT_SIMPLE_AUTOMATION_TEST(FAnimationAnimNextRuntimeTest_IUpdate, "Animation.AnimNext.Runtime.IUpdate", EAutomationTestFlags::EditorContext | EAutomationTestFlags::EngineFilter)
+
+bool FAnimationAnimNextRuntimeTest_IUpdate::RunTest(const FString& InParameters)
+{
+	using namespace UE::AnimNext;
+	{
+		AUTO_REGISTER_ANIM_TRAIT(FTraitWithNoChildren)
+		AUTO_REGISTER_ANIM_TRAIT(FTraitWithOneChild)
+		AUTO_REGISTER_ANIM_TRAIT(FTraitWithChildren)
+
+		UFactory* GraphFactory = NewObject<UAnimNextGraphFactory>();
+		UAnimNextGraph* AnimNextGraph = CastChecked<UAnimNextGraph>(GraphFactory->FactoryCreateNew(UAnimNextGraph::StaticClass(), GetTransientPackage(), TEXT("TestAnimNextGraph"), RF_Transient, nullptr, nullptr, NAME_None));
+		UE_RETURN_ON_ERROR(AnimNextGraph != nullptr, "FAnimationAnimNextRuntimeTest_IUpdate -> Failed to create animation graph");
+
+		FScopedClearNodeTemplateRegistry ScopedClearNodeTemplateRegistry;
+		FNodeTemplateRegistry& Registry = FNodeTemplateRegistry::Get();
+
+		// We create a few node templates
+		// Template A has a single trait with no children
+		TArray<FTraitUID> NodeTemplateTraitListA;
+		NodeTemplateTraitListA.Add(FTraitWithNoChildren::TraitUID);
+
+		// Template B has a single trait with one child, it doesn't update
+		TArray<FTraitUID> NodeTemplateTraitListB;
+		NodeTemplateTraitListB.Add(FTraitWithOneChild::TraitUID);
+
+		// Template C has a single trait with children
+		TArray<FTraitUID> NodeTemplateTraitListC;
+		NodeTemplateTraitListC.Add(FTraitWithChildren::TraitUID);
+
+		// Populate our node template registry
+		TArray<uint8> NodeTemplateBufferA, NodeTemplateBufferB, NodeTemplateBufferC;
+		const FNodeTemplate* NodeTemplateA = FNodeTemplateBuilder::BuildNodeTemplate(NodeTemplateTraitListA, NodeTemplateBufferA);
+		const FNodeTemplate* NodeTemplateB = FNodeTemplateBuilder::BuildNodeTemplate(NodeTemplateTraitListB, NodeTemplateBufferB);
+		const FNodeTemplate* NodeTemplateC = FNodeTemplateBuilder::BuildNodeTemplate(NodeTemplateTraitListC, NodeTemplateBufferC);
+
+		// Build our graph, it as follow (each node template has a single node instance):
+		// NodeA has no children
+		// NodeB has one child: NodeA (it doesn't update)
+		// NodeC (root) has two children: NodeA and NodeB
+
+		TArray<FNodeHandle> NodeHandles;
+
+		// Write our graph
+		TArray<uint8> GraphSharedDataArchiveBuffer;
+		TArray<TObjectPtr<UObject>> GraphReferencedObjects;
+		{
+			FTraitWriter TraitWriter;
+
+			NodeHandles.Add(TraitWriter.RegisterNode(*NodeTemplateC));	// Root node
+			NodeHandles.Add(TraitWriter.RegisterNode(*NodeTemplateA));
+			NodeHandles.Add(TraitWriter.RegisterNode(*NodeTemplateB));
+
+			// We don't have trait properties
+			TArray<TMap<FName, FString>> TraitPropertiesA;
+			TraitPropertiesA.AddDefaulted(NodeTemplateTraitListA.Num());
+
+			TArray<TMap<FName, FString>> TraitPropertiesB;
+			TraitPropertiesB.AddDefaulted(NodeTemplateTraitListB.Num());
+			TraitPropertiesB[0].Add(TEXT("Child"), ToString<FTraitWithOneChild::FSharedData>(TEXT("Child"), FAnimNextTraitHandle(NodeHandles[1])));
+
+			TArray<TMap<FName, FString>> TraitPropertiesC;
+			TraitPropertiesC.AddDefaulted(NodeTemplateTraitListC.Num());
+			FAnimNextTraitHandle ChildrenHandlesC[2] = { FAnimNextTraitHandle(NodeHandles[1]), FAnimNextTraitHandle(NodeHandles[2])};
+			TraitPropertiesC[0].Add(TEXT("Children"), ToString<FTraitWithChildren::FSharedData>(TEXT("Children"), ChildrenHandlesC));
+
+			TraitWriter.BeginNodeWriting();
+			TraitWriter.WriteNode(NodeHandles[0],
+				[&TraitPropertiesC](uint32 TraitIndex, FName PropertyName)
+				{
+					return TraitPropertiesC[TraitIndex][PropertyName];
+				},
+				[](uint32 TraitIndex, FName PropertyName)
+				{
+					return MAX_uint16;
+				});
+			TraitWriter.WriteNode(NodeHandles[1],
+				[&TraitPropertiesA](uint32 TraitIndex, FName PropertyName)
+				{
+					return TraitPropertiesA[TraitIndex][PropertyName];
+				},
+				[](uint32 TraitIndex, FName PropertyName)
+				{
+					return MAX_uint16;
+				});
+			TraitWriter.WriteNode(NodeHandles[2],
+				[&TraitPropertiesB](uint32 TraitIndex, FName PropertyName)
+				{
+					return TraitPropertiesB[TraitIndex][PropertyName];
+				},
+				[](uint32 TraitIndex, FName PropertyName)
+				{
+					return MAX_uint16;
+				});
+			TraitWriter.EndNodeWriting();
+
+			AddErrorIfFalse(TraitWriter.GetErrorState() == FTraitWriter::EErrorState::None, "FAnimationAnimNextRuntimeTest_IUpdate -> Failed to write traits");
+			GraphSharedDataArchiveBuffer = TraitWriter.GetGraphSharedData();
+			GraphReferencedObjects = TraitWriter.GetGraphReferencedObjects();
+		}
+
+		// Read our graph
+		FTestUtils::LoadFromArchiveBuffer(*AnimNextGraph, NodeHandles, GraphSharedDataArchiveBuffer);
+
+		FAnimNextGraphInstancePtr GraphInstance;
+		AnimNextGraph->AllocateInstance(GraphInstance);
+
+		FExecutionContext Context(GraphInstance);
+
+		{
+			TArray<FTraitUID> UpdatedTraits;
+			Private::UpdatedTraits = &UpdatedTraits;
+
+			// Call pre/post update on our graph
+			UpdateGraph(GraphInstance, 0.0333f);
+
+			AddErrorIfFalse(UpdatedTraits.Num() == 6, "FAnimationAnimNextRuntimeTest_IUpdate -> Expected 6 nodes to have been visited during the update traversal");
+			AddErrorIfFalse(UpdatedTraits[0] == FTraitWithChildren::TraitUID, "FAnimationAnimNextRuntimeTest_IUpdate -> Unexpected update order");		// NodeC
+			AddErrorIfFalse(UpdatedTraits[1] == FTraitWithNoChildren::TraitUID, "FAnimationAnimNextRuntimeTest_IUpdate -> Unexpected update order");		// NodeA
+			AddErrorIfFalse(UpdatedTraits[2] == FTraitWithNoChildren::TraitUID, "FAnimationAnimNextRuntimeTest_IUpdate -> Unexpected update order");		// NodeA
+			AddErrorIfFalse(UpdatedTraits[3] == FTraitWithNoChildren::TraitUID, "FAnimationAnimNextRuntimeTest_IUpdate -> Unexpected update order");		// NodeB -> NodeA
+			AddErrorIfFalse(UpdatedTraits[4] == FTraitWithNoChildren::TraitUID, "FAnimationAnimNextRuntimeTest_IUpdate -> Unexpected update order");		// NodeB -> NodeA
+			AddErrorIfFalse(UpdatedTraits[5] == FTraitWithChildren::TraitUID, "FAnimationAnimNextRuntimeTest_IUpdate -> Unexpected update order");		// NodeC
+
+			Private::UpdatedTraits = nullptr;
+		}
+
+		Registry.Unregister(NodeTemplateA);
+		Registry.Unregister(NodeTemplateB);
+		Registry.Unregister(NodeTemplateC);
+
+		AddErrorIfFalse(Registry.GetNum() == 0, "FAnimationAnimNextRuntimeTest_IUpdate -> Registry should contain 0 templates");
+	}
+	Tests::FUtils::CleanupAfterTests();
+
+	return true;
+}
+
+IMPLEMENT_SIMPLE_AUTOMATION_TEST(FAnimationAnimNextRuntimeTest_IEvaluate, "Animation.AnimNext.Runtime.IEvaluate", EAutomationTestFlags::EditorContext | EAutomationTestFlags::EngineFilter)
+
+bool FAnimationAnimNextRuntimeTest_IEvaluate::RunTest(const FString& InParameters)
+{
+	using namespace UE::AnimNext;
+	{
+		AUTO_REGISTER_ANIM_TRAIT(FTraitWithNoChildren)
+		AUTO_REGISTER_ANIM_TRAIT(FTraitWithOneChild)
+		AUTO_REGISTER_ANIM_TRAIT(FTraitWithChildren)
+
+		UFactory* GraphFactory = NewObject<UAnimNextGraphFactory>();
+		UAnimNextGraph* AnimNextGraph = CastChecked<UAnimNextGraph>(GraphFactory->FactoryCreateNew(UAnimNextGraph::StaticClass(), GetTransientPackage(), TEXT("TestAnimNextGraph"), RF_Transient, nullptr, nullptr, NAME_None));
+		UE_RETURN_ON_ERROR(AnimNextGraph != nullptr, "FAnimationAnimNextRuntimeTest_IEvaluate -> Failed to create animation graph");
+
+		FScopedClearNodeTemplateRegistry ScopedClearNodeTemplateRegistry;
+		FNodeTemplateRegistry& Registry = FNodeTemplateRegistry::Get();
+
+		// We create a few node templates
+		// Template A has a single trait with no children
+		TArray<FTraitUID> NodeTemplateTraitListA;
+		NodeTemplateTraitListA.Add(FTraitWithNoChildren::TraitUID);
+
+		// Template B has a single trait with one child, it doesn't evaluate
+		TArray<FTraitUID> NodeTemplateTraitListB;
+		NodeTemplateTraitListB.Add(FTraitWithOneChild::TraitUID);
+
+		// Template C has a single trait with children
+		TArray<FTraitUID> NodeTemplateTraitListC;
+		NodeTemplateTraitListC.Add(FTraitWithChildren::TraitUID);
+
+		// Populate our node template registry
+		TArray<uint8> NodeTemplateBufferA, NodeTemplateBufferB, NodeTemplateBufferC;
+		const FNodeTemplate* NodeTemplateA = FNodeTemplateBuilder::BuildNodeTemplate(NodeTemplateTraitListA, NodeTemplateBufferA);
+		const FNodeTemplate* NodeTemplateB = FNodeTemplateBuilder::BuildNodeTemplate(NodeTemplateTraitListB, NodeTemplateBufferB);
+		const FNodeTemplate* NodeTemplateC = FNodeTemplateBuilder::BuildNodeTemplate(NodeTemplateTraitListC, NodeTemplateBufferC);
+
+		// Build our graph, it as follow (each node template has a single node instance):
+		// NodeA has no children
+		// NodeB has one child: NodeA (it doesn't evaluate)
+		// NodeC (root) has two children: NodeA and NodeB
+
+		TArray<FNodeHandle> NodeHandles;
+
+		// Write our graph
+		TArray<uint8> GraphSharedDataArchiveBuffer;
+		TArray<TObjectPtr<UObject>> GraphReferencedObjects;
+		{
+			FTraitWriter TraitWriter;
+
+			NodeHandles.Add(TraitWriter.RegisterNode(*NodeTemplateC));	// Root node
+			NodeHandles.Add(TraitWriter.RegisterNode(*NodeTemplateA));
+			NodeHandles.Add(TraitWriter.RegisterNode(*NodeTemplateB));
+
+			// We don't have trait properties
+			TArray<TMap<FName, FString>> TraitPropertiesA;
+			TraitPropertiesA.AddDefaulted(NodeTemplateTraitListA.Num());
+
+			TArray<TMap<FName, FString>> TraitPropertiesB;
+			TraitPropertiesB.AddDefaulted(NodeTemplateTraitListB.Num());
+			TraitPropertiesB[0].Add(TEXT("Child"), ToString<FTraitWithOneChild::FSharedData>(TEXT("Child"), FAnimNextTraitHandle(NodeHandles[1])));
+
+			TArray<TMap<FName, FString>> TraitPropertiesC;
+			TraitPropertiesC.AddDefaulted(NodeTemplateTraitListC.Num());
+
+			FAnimNextTraitHandle ChildrenHandlesC[2] = { FAnimNextTraitHandle(NodeHandles[1]), FAnimNextTraitHandle(NodeHandles[2])};
+			TraitPropertiesC[0].Add(TEXT("Children"), ToString<FTraitWithChildren::FSharedData>(TEXT("Children"), ChildrenHandlesC));
+
+			TraitWriter.BeginNodeWriting();
+			TraitWriter.WriteNode(NodeHandles[0],
+				[&TraitPropertiesC](uint32 TraitIndex, FName PropertyName)
+				{
+					return TraitPropertiesC[TraitIndex][PropertyName];
+				},
+				[](uint32 TraitIndex, FName PropertyName)
+				{
+					return MAX_uint16;
+				});
+			TraitWriter.WriteNode(NodeHandles[1],
+				[&TraitPropertiesA](uint32 TraitIndex, FName PropertyName)
+				{
+					return TraitPropertiesA[TraitIndex][PropertyName];
+				},
+				[](uint32 TraitIndex, FName PropertyName)
+				{
+					return MAX_uint16;
+				});
+			TraitWriter.WriteNode(NodeHandles[2],
+				[&TraitPropertiesB](uint32 TraitIndex, FName PropertyName)
+				{
+					return TraitPropertiesB[TraitIndex][PropertyName];
+				},
+				[](uint32 TraitIndex, FName PropertyName)
+				{
+					return MAX_uint16;
+				});
+			TraitWriter.EndNodeWriting();
+
+			AddErrorIfFalse(TraitWriter.GetErrorState() == FTraitWriter::EErrorState::None, "FAnimationAnimNextRuntimeTest_IEvaluate -> Failed to write traits");
+			GraphSharedDataArchiveBuffer = TraitWriter.GetGraphSharedData();
+			GraphReferencedObjects = TraitWriter.GetGraphReferencedObjects();
+		}
+
+		// Read our graph
+		FTestUtils::LoadFromArchiveBuffer(*AnimNextGraph, NodeHandles, GraphSharedDataArchiveBuffer);
+
+		FAnimNextGraphInstancePtr GraphInstance;
+		AnimNextGraph->AllocateInstance(GraphInstance);
+
+		{
+			TArray<FTraitUID> EvaluatedTraits;
+			Private::EvaluatedTraits = &EvaluatedTraits;
+
+			// Call pre/post evaluate on our graph
+			(void)EvaluateGraph(GraphInstance);
+
+			AddErrorIfFalse(EvaluatedTraits.Num() == 6, "FAnimationAnimNextRuntimeTest_IEvaluate -> Expected 6 nodes to have been visited during the evaluate traversal");
+			AddErrorIfFalse(EvaluatedTraits[0] == FTraitWithChildren::TraitUID, "FAnimationAnimNextRuntimeTest_IEvaluate -> Unexpected evaluate order");		// NodeC
+			AddErrorIfFalse(EvaluatedTraits[1] == FTraitWithNoChildren::TraitUID, "FAnimationAnimNextRuntimeTest_IEvaluate -> Unexpected evaluate order");			// NodeA
+			AddErrorIfFalse(EvaluatedTraits[2] == FTraitWithNoChildren::TraitUID, "FAnimationAnimNextRuntimeTest_IEvaluate -> Unexpected evaluate order");			// NodeA
+			AddErrorIfFalse(EvaluatedTraits[3] == FTraitWithNoChildren::TraitUID, "FAnimationAnimNextRuntimeTest_IEvaluate -> Unexpected evaluate order");			// NodeB -> NodeA
+			AddErrorIfFalse(EvaluatedTraits[4] == FTraitWithNoChildren::TraitUID, "FAnimationAnimNextRuntimeTest_IEvaluate -> Unexpected evaluate order");			// NodeB -> NodeA
+			AddErrorIfFalse(EvaluatedTraits[5] == FTraitWithChildren::TraitUID, "FAnimationAnimNextRuntimeTest_IEvaluate -> Unexpected evaluate order");		// NodeC
+
+			Private::EvaluatedTraits = nullptr;
+		}
+
+		Registry.Unregister(NodeTemplateA);
+		Registry.Unregister(NodeTemplateB);
+		Registry.Unregister(NodeTemplateC);
+
+		AddErrorIfFalse(Registry.GetNum() == 0, "FAnimationAnimNextRuntimeTest_IEvaluate -> Registry should contain 0 templates");
+	}
+	Tests::FUtils::CleanupAfterTests();
+
+	return true;
+}
+IMPLEMENT_SIMPLE_AUTOMATION_TEST(FAnimationAnimNextRuntimeTest_IScopedInterface, "Animation.AnimNext.Runtime.IScopedInterface", EAutomationTestFlags::EditorContext | EAutomationTestFlags::EngineFilter)
+
+bool FAnimationAnimNextRuntimeTest_IScopedInterface::RunTest(const FString& InParameters)
+{
+	using namespace UE::AnimNext;
+	{
+		AUTO_REGISTER_ANIM_TRAIT(FTraitWithOneChild)
+		AUTO_REGISTER_ANIM_TRAIT(FScopedTagTrait)
+		AUTO_REGISTER_ANIM_TRAIT(FTestScopedTagTrait)
+
+		UFactory* GraphFactory = NewObject<UAnimNextGraphFactory>();
+		UAnimNextGraph* AnimNextGraph = CastChecked<UAnimNextGraph>(GraphFactory->FactoryCreateNew(UAnimNextGraph::StaticClass(), GetTransientPackage(), TEXT("TestAnimNextGraph"), RF_Transient, nullptr, nullptr, NAME_None));
+		UE_RETURN_ON_ERROR(AnimNextGraph != nullptr, "FAnimationAnimNextRuntimeTest_IScopedInterface -> Failed to create animation graph");
+
+		FScopedClearNodeTemplateRegistry ScopedClearNodeTemplateRegistry;
+		FNodeTemplateRegistry& Registry = FNodeTemplateRegistry::Get();
+
+		// We create a few node templates
+		// Template A has a single child and tests for our tag
+		TArray<FTraitUID> NodeTemplateTraitList0;
+		NodeTemplateTraitList0.Add(FTraitWithOneChild::TraitUID);
+		NodeTemplateTraitList0.Add(FTestScopedTagTrait::TraitUID);
+
+		// Template B has a single child, it tests and pushes our tag
+		TArray<FTraitUID> NodeTemplateTraitList1;
+		NodeTemplateTraitList1.Add(FTraitWithOneChild::TraitUID);
+		NodeTemplateTraitList1.Add(FTestScopedTagTrait::TraitUID);	// Test after push/pop
+		NodeTemplateTraitList1.Add(FScopedTagTrait::TraitUID);
+		NodeTemplateTraitList1.Add(FTestScopedTagTrait::TraitUID);	// Test before push/pop
+
+		// Populate our node template registry
+		TArray<uint8> NodeTemplateBuffer0, NodeTemplateBuffer1;
+		const FNodeTemplate* NodeTemplate0 = FNodeTemplateBuilder::BuildNodeTemplate(NodeTemplateTraitList0, NodeTemplateBuffer0);
+		const FNodeTemplate* NodeTemplate1 = FNodeTemplateBuilder::BuildNodeTemplate(NodeTemplateTraitList1, NodeTemplateBuffer1);
+
+		// Build our graph, it as follow:
+		// NodeA has no child (tag is scoped)
+		// NodeB has one child: NodeA (NodeB adds the scoped tag)
+		// NodeC (root) has one child: NodeB (no tag scoped)
+
+		TArray<FNodeHandle> NodeHandles;
+
+		// Write our graph
+		TArray<uint8> GraphSharedDataArchiveBuffer;
+		TArray<TObjectPtr<UObject>> GraphReferencedObjects;
+		{
+			FTraitWriter TraitWriter;
+
+			NodeHandles.Add(TraitWriter.RegisterNode(*NodeTemplate0));	// NodeC (root node)
+			NodeHandles.Add(TraitWriter.RegisterNode(*NodeTemplate1));	// NodeB
+			NodeHandles.Add(TraitWriter.RegisterNode(*NodeTemplate0));	// NodeA
+
+			// We don't have trait properties
+			TArray<TMap<FName, FString>> TraitPropertiesA;
+			TraitPropertiesA.AddDefaulted(NodeTemplateTraitList0.Num());
+			TraitPropertiesA[0].Add(TEXT("Child"), ToString<FTraitWithOneChild::FSharedData>(TEXT("Child"), FAnimNextTraitHandle()));
+
+			TArray<TMap<FName, FString>> TraitPropertiesB;
+			TraitPropertiesB.AddDefaulted(NodeTemplateTraitList1.Num());
+			TraitPropertiesB[0].Add(TEXT("Child"), ToString<FTraitWithOneChild::FSharedData>(TEXT("Child"), FAnimNextTraitHandle(NodeHandles[2])));
+
+			TArray<TMap<FName, FString>> TraitPropertiesC;
+			TraitPropertiesC.AddDefaulted(NodeTemplateTraitList0.Num());
+			TraitPropertiesC[0].Add(TEXT("Child"), ToString<FTraitWithOneChild::FSharedData>(TEXT("Child"), FAnimNextTraitHandle(NodeHandles[1])));
+
+			TraitWriter.BeginNodeWriting();
+			TraitWriter.WriteNode(NodeHandles[0],
+				[&TraitPropertiesC](uint32 TraitIndex, FName PropertyName)
+				{
+					return TraitPropertiesC[TraitIndex][PropertyName];
+				},
+				[](uint32 TraitIndex, FName PropertyName)
+				{
+					return MAX_uint16;
+				});
+			TraitWriter.WriteNode(NodeHandles[1],
+				[&TraitPropertiesB](uint32 TraitIndex, FName PropertyName)
+				{
+					return TraitPropertiesB[TraitIndex][PropertyName];
+				},
+				[](uint32 TraitIndex, FName PropertyName)
+				{
+					return MAX_uint16;
+				});
+			TraitWriter.WriteNode(NodeHandles[2],
+				[&TraitPropertiesA](uint32 TraitIndex, FName PropertyName)
+				{
+					return TraitPropertiesA[TraitIndex][PropertyName];
+				},
+				[](uint32 TraitIndex, FName PropertyName)
+				{
+					return MAX_uint16;
+				});
+			TraitWriter.EndNodeWriting();
+
+			AddErrorIfFalse(TraitWriter.GetErrorState() == FTraitWriter::EErrorState::None, "FAnimationAnimNextRuntimeTest_IScopedInterface -> Failed to write traits");
+			GraphSharedDataArchiveBuffer = TraitWriter.GetGraphSharedData();
+			GraphReferencedObjects = TraitWriter.GetGraphReferencedObjects();
+		}
+
+		// Read our graph
+		FTestUtils::LoadFromArchiveBuffer(*AnimNextGraph, NodeHandles, GraphSharedDataArchiveBuffer);
+
+		FAnimNextGraphInstancePtr GraphInstance;
+		AnimNextGraph->AllocateInstance(GraphInstance);
+
+		FExecutionContext Context(GraphInstance);
+
+		{
+			TArray<bool> IsTagInScope;
+			Private::IsTagInScope = &IsTagInScope;
+
+			// Call pre/post update on our graph
+			Private::AutoPopTag = true;
+			UpdateGraph(GraphInstance, 0.0333f);
+
+			AddErrorIfFalse(IsTagInScope.Num() == 8, "FAnimationAnimNextRuntimeTest_IScopedInterface -> Unexpected number of entries");
+			AddErrorIfFalse(IsTagInScope[0] == false, "FAnimationAnimNextRuntimeTest_IScopedInterface -> Unexpected scoped tag state");		// NodeC::PreUpdate (template 0)
+			AddErrorIfFalse(IsTagInScope[1] == false, "FAnimationAnimNextRuntimeTest_IScopedInterface -> Unexpected scoped tag state");		// NodeB::Before::PreUpdate (template 1)
+			AddErrorIfFalse(IsTagInScope[2] == true, "FAnimationAnimNextRuntimeTest_IScopedInterface -> Unexpected scoped tag state");		// NodeB::After::PreUpdate (template 1)
+			AddErrorIfFalse(IsTagInScope[3] == true, "FAnimationAnimNextRuntimeTest_IScopedInterface -> Unexpected scoped tag state");		// NodeA::PreUpdate (template 0)
+			AddErrorIfFalse(IsTagInScope[4] == true, "FAnimationAnimNextRuntimeTest_IScopedInterface -> Unexpected scoped tag state");		// NodeA::PostUpdate (template 0)
+			AddErrorIfFalse(IsTagInScope[5] == true, "FAnimationAnimNextRuntimeTest_IScopedInterface -> Unexpected scoped tag state");		// NodeB::Before::PostUpdate (template 1)
+			AddErrorIfFalse(IsTagInScope[6] == true, "FAnimationAnimNextRuntimeTest_IScopedInterface -> Unexpected scoped tag state");		// NodeB::After::PostUpdate (template 1)
+			AddErrorIfFalse(IsTagInScope[7] == false, "FAnimationAnimNextRuntimeTest_IScopedInterface -> Unexpected scoped tag state");		// NodeC::PostUpdate (template 0)
+
+			IsTagInScope.Reset();
+			Private::AutoPopTag = false;
+
+			// Call pre/post update on our graph
+			UpdateGraph(GraphInstance, 0.0333f);
+
+			AddErrorIfFalse(IsTagInScope.Num() == 8, "FAnimationAnimNextRuntimeTest_IScopedInterface -> Unexpected number of entries");
+			AddErrorIfFalse(IsTagInScope[0] == false, "FAnimationAnimNextRuntimeTest_IScopedInterface -> Unexpected scoped tag state");		// NodeC::PreUpdate (template 0)
+			AddErrorIfFalse(IsTagInScope[1] == false, "FAnimationAnimNextRuntimeTest_IScopedInterface -> Unexpected scoped tag state");		// NodeB::Before::PreUpdate (template 1)
+			AddErrorIfFalse(IsTagInScope[2] == true, "FAnimationAnimNextRuntimeTest_IScopedInterface -> Unexpected scoped tag state");		// NodeB::After::PreUpdate (template 1)
+			AddErrorIfFalse(IsTagInScope[3] == true, "FAnimationAnimNextRuntimeTest_IScopedInterface -> Unexpected scoped tag state");		// NodeA::PreUpdate (template 0)
+			AddErrorIfFalse(IsTagInScope[4] == true, "FAnimationAnimNextRuntimeTest_IScopedInterface -> Unexpected scoped tag state");		// NodeA::PostUpdate (template 0)
+			AddErrorIfFalse(IsTagInScope[5] == true, "FAnimationAnimNextRuntimeTest_IScopedInterface -> Unexpected scoped tag state");		// NodeB::Before::PostUpdate (template 1)
+			AddErrorIfFalse(IsTagInScope[6] == false, "FAnimationAnimNextRuntimeTest_IScopedInterface -> Unexpected scoped tag state");		// NodeB::After::PostUpdate (template 1)
+			AddErrorIfFalse(IsTagInScope[7] == false, "FAnimationAnimNextRuntimeTest_IScopedInterface -> Unexpected scoped tag state");		// NodeC::PostUpdate (template 0)
+
+			Private::UpdatedTraits = nullptr;
+		}
+
+		Registry.Unregister(NodeTemplate0);
+		Registry.Unregister(NodeTemplate1);
+
+		AddErrorIfFalse(Registry.GetNum() == 0, "FAnimationAnimNextRuntimeTest_IScopedInterface -> Registry should contain 0 templates");
+	}
+	Tests::FUtils::CleanupAfterTests();
+
+	return true;
+}
+
+#endif
diff --git a/Engine/Plugins/Experimental/AnimNextAnimGraph/Source/AnimNextAnimGraphTestSuite/Private/AnimNextAnimGraphTraitInterfacesTest.h b/Engine/Plugins/Experimental/AnimNextAnimGraph/Source/AnimNextAnimGraphTestSuite/Private/AnimNextAnimGraphTraitInterfacesTest.h
new file mode 100644
--- /dev/null
+++ b/Engine/Plugins/Experimental/AnimNextAnimGraph/Source/AnimNextAnimGraphTestSuite/Private/AnimNextAnimGraphTraitInterfacesTest.h
@@ -0,0 +1,28 @@
+// Copyright Epic Games, Inc. All Rights Reserved.
+
+#pragma once
+
+#include "CoreMinimal.h"
+
+#include "TraitCore/TraitHandle.h"
+#include "TraitCore/TraitSharedData.h"
+
+#include "AnimNextAnimGraphTraitInterfacesTest.generated.h"
+
+USTRUCT()
+struct FTraitWithOneChildSharedData : public FAnimNextTraitSharedData
+{
+	GENERATED_BODY()
+
+	UPROPERTY()
+	FAnimNextTraitHandle Child;
+};
+
+USTRUCT()
+struct FTraitWithChildrenSharedData : public FAnimNextTraitSharedData
+{
+	GENERATED_BODY()
+
+	UPROPERTY()
+	FAnimNextTraitHandle Children[2];
+};
diff --git a/Engine/Plugins/Experimental/AnimNextAnimGraph/Source/AnimNextAnimGraphTestSuite/Private/Module.cpp b/Engine/Plugins/Experimental/AnimNextAnimGraph/Source/AnimNextAnimGraphTestSuite/Private/Module.cpp
new file mode 100644
--- /dev/null
+++ b/Engine/Plugins/Experimental/AnimNextAnimGraph/Source/AnimNextAnimGraphTestSuite/Private/Module.cpp
@@ -0,0 +1,17 @@
+// Copyright Epic Games, Inc. All Rights Reserved.
+
+#include "CoreMinimal.h"
+#include "Modules/ModuleInterface.h"
+#include "Modules/ModuleManager.h"
+
+namespace UE::AnimNext::Tests
+{
+
+class FModule : public IModuleInterface
+{
+public:
+};
+
+}
+
+IMPLEMENT_MODULE(UE::AnimNext::Tests::FModule, AnimNextAnimGraphTestSuite)
diff --git a/Engine/Plugins/Experimental/AnimNextAnimGraph/Source/AnimNextAnimGraphUncookedOnly/AnimNextAnimGraphUncookedOnly.Build.cs b/Engine/Plugins/Experimental/AnimNextAnimGraph/Source/AnimNextAnimGraphUncookedOnly/AnimNextAnimGraphUncookedOnly.Build.cs
new file mode 100644
--- /dev/null
+++ b/Engine/Plugins/Experimental/AnimNextAnimGraph/Source/AnimNextAnimGraphUncookedOnly/AnimNextAnimGraphUncookedOnly.Build.cs
@@ -0,0 +1,37 @@
+// Copyright Epic Games, Inc. All Rights Reserved.
+
+namespace UnrealBuildTool.Rules
+{
+	public class AnimNextAnimGraphUncookedOnly : ModuleRules
+	{
+		public AnimNextAnimGraphUncookedOnly(ReadOnlyTargetRules Target) : base(Target)
+		{
+			PrivateDependencyModuleNames.AddRange(
+				new string[]
+				{
+					"Core",
+					"CoreUObject",
+					"Engine",
+					"AnimNext",
+					"AnimNextUncookedOnly",
+					"AnimNextAnimGraph",
+					"ControlRigDeveloper",
+					"BlueprintGraph",	// For K2 schema
+					"AnimGraph",
+					"RigVM",
+					"ControlRig",
+				}
+			);
+
+			if (Target.bBuildEditor)
+			{
+				PrivateDependencyModuleNames.AddRange(
+					new string[]
+					{
+						"UnrealEd",
+					}
+				);
+			}
+		}
+	}
+}
\ No newline at end of file
diff --git a/Engine/Plugins/Experimental/AnimNextAnimGraph/Source/AnimNextAnimGraphUncookedOnly/Internal/Graph/AnimGraph/AnimGraphNode_AnimNextGraph.h b/Engine/Plugins/Experimental/AnimNextAnimGraph/Source/AnimNextAnimGraphUncookedOnly/Internal/Graph/AnimGraph/AnimGraphNode_AnimNextGraph.h
new file mode 100644
--- /dev/null
+++ b/Engine/Plugins/Experimental/AnimNextAnimGraph/Source/AnimNextAnimGraphUncookedOnly/Internal/Graph/AnimGraph/AnimGraphNode_AnimNextGraph.h
@@ -0,0 +1,56 @@
+// Copyright Epic Games, Inc. All Rights Reserved.
+
+#pragma once
+
+#include "AnimGraphNode_CustomProperty.h"
+#include "Graph/AnimGraph/AnimNode_AnimNextGraph.h"
+#include "SSearchableComboBox.h"
+#include "AnimGraphNode_AnimNextGraph.generated.h"
+
+struct FVariableMappingInfo;
+
+UCLASS(MinimalAPI)
+class UAnimGraphNode_AnimNextGraph : public UAnimGraphNode_CustomProperty
+{
+	GENERATED_UCLASS_BODY()
+
+	UPROPERTY(EditAnywhere, Category = Settings)
+	FAnimNode_AnimNextGraph Node;
+
+	virtual void CustomizeDetails(IDetailLayoutBuilder& DetailBuilder) override;
+	virtual void PostEditChangeProperty(FPropertyChangedEvent& PropertyChangedEvent) override;
+	virtual void PostReconstructNode() override;
+	virtual void GetRequiredExtensions(TArray<TSubclassOf<UAnimBlueprintExtension>>& OutExtensions) const override;
+
+private:
+	// UEdGraphNode interface
+	virtual FText GetNodeTitle(ENodeTitleType::Type TitleType) const override;
+	virtual void CustomizePinData(UEdGraphPin* Pin, FName SourcePropertyName, int32 ArrayIndex) const override;
+	virtual FText GetTooltipText() const override;
+	virtual FText GetMenuCategory() const override;
+	virtual void PreloadRequiredAssets() override;
+	//virtual bool IsSinkNode() const { return true; }
+
+	virtual FAnimNode_CustomProperty* GetCustomPropertyNode() override { return &Node; }
+	virtual const FAnimNode_CustomProperty* GetCustomPropertyNode() const override { return &Node; }
+
+	// property related things
+	void GetVariables(bool bInput, TMap<FName, FRigVMExternalVariable>& OutParameters) const;
+
+	TMap<FName, FRigVMExternalVariable> InputVariables;
+	TMap<FName, FRigVMExternalVariable> OutputVariables;
+
+	void RebuildExposedProperties();
+	virtual void CreateCustomPins(TArray<UEdGraphPin*>* OldPins) override;
+	virtual void ValidateAnimNodeDuringCompilation(USkeleton* ForSkeleton, FCompilerResultsLog& MessageLog) override;
+	
+	// pin option related
+	bool IsPropertyExposeEnabled(FName PropertyName) const;
+	ECheckBoxState IsPropertyExposed(FName PropertyName) const;
+	void OnPropertyExposeCheckboxChanged(ECheckBoxState NewState, FName PropertyName);
+
+	bool IsInputProperty(const FName& PropertyName) const;
+
+	virtual UObject* GetJumpTargetForDoubleClick() const override;
+};
+
diff --git a/Engine/Plugins/Experimental/AnimNextAnimGraph/Source/AnimNextAnimGraphUncookedOnly/Private/Graph/AnimGraph/AnimGraphNode_AnimNextGraph.cpp b/Engine/Plugins/Experimental/AnimNextAnimGraph/Source/AnimNextAnimGraphUncookedOnly/Private/Graph/AnimGraph/AnimGraphNode_AnimNextGraph.cpp
new file mode 100644
--- /dev/null
+++ b/Engine/Plugins/Experimental/AnimNextAnimGraph/Source/AnimNextAnimGraphUncookedOnly/Private/Graph/AnimGraph/AnimGraphNode_AnimNextGraph.cpp
@@ -0,0 +1,216 @@
+// Copyright Epic Games, Inc. All Rights Reserved.
+
+#include "Graph/AnimGraph/AnimGraphNode_AnimNextGraph.h"
+#include "Kismet2/CompilerResultsLog.h"
+#include "DetailLayoutBuilder.h"
+#include "Widgets/Text/STextBlock.h"
+#include "Widgets/Input/SCheckBox.h"
+#include "Widgets/Input/SButton.h"
+#include "ScopedTransaction.h"
+#include "Kismet2/BlueprintEditorUtils.h"
+#include "RigVMBlueprintGeneratedClass.h"
+#include "ControlRigBlueprint.h"
+#include "Graph/AnimGraph/AnimBlueprintExtension_AnimNextParameters.h"
+
+#include UE_INLINE_GENERATED_CPP_BY_NAME(AnimGraphNode_AnimNextGraph)
+
+#define LOCTEXT_NAMESPACE "AnimGraphNode_AnimNextGraph"
+
+UAnimGraphNode_AnimNextGraph::UAnimGraphNode_AnimNextGraph(const FObjectInitializer& ObjectInitializer)
+	: Super(ObjectInitializer)
+{
+}
+
+FText UAnimGraphNode_AnimNextGraph::GetNodeTitle(ENodeTitleType::Type TitleType) const
+{
+	// display control rig here
+	return LOCTEXT("AnimGraphNode_AnimNextGraph_Title", "AnimNext Graph");
+}
+
+FText UAnimGraphNode_AnimNextGraph::GetTooltipText() const
+{
+	// display control rig here
+	return LOCTEXT("AnimGraphNode_AnimNextGraph_Tooltip", "Evaluates a AnimNext Graph");
+}
+
+FText UAnimGraphNode_AnimNextGraph::GetMenuCategory() const
+{
+	return LOCTEXT("Category", "AnimNext");
+}
+
+void UAnimGraphNode_AnimNextGraph::PreloadRequiredAssets()
+{
+	Super::PreloadRequiredAssets();
+
+	if (Node.AnimNextGraph)
+	{
+		PreloadObject(Node.AnimNextGraph);
+		PreloadObject(Node.AnimNextGraph->EditorData);
+	}
+}
+
+void UAnimGraphNode_AnimNextGraph::GetRequiredExtensions(TArray<TSubclassOf<UAnimBlueprintExtension>>& OutExtensions) const
+{
+	OutExtensions.Add(UAnimBlueprintExtension_AnimNextParameters::StaticClass());
+}
+
+void UAnimGraphNode_AnimNextGraph::CreateCustomPins(TArray<UEdGraphPin*>* OldPins)
+{
+	DECLARE_SCOPE_HIERARCHICAL_COUNTER_FUNC()
+	
+	// we do this to refresh input variables
+	//RebuildExposedProperties();
+
+	// Grab the SKELETON class here as when we are reconstructed during during BP compilation
+	// the full generated class is not yet present built.
+	if (!GetTargetSkeletonClass())
+	{
+		// Nothing to search for properties
+		return;
+	}
+}
+
+void UAnimGraphNode_AnimNextGraph::ValidateAnimNodeDuringCompilation(USkeleton* ForSkeleton, FCompilerResultsLog& MessageLog)
+{
+	if (UClass* TargetClass = GetTargetClass())
+	{
+	}
+
+	Super::ValidateAnimNodeDuringCompilation(ForSkeleton, MessageLog);
+}
+
+void UAnimGraphNode_AnimNextGraph::RebuildExposedProperties()
+{
+	DECLARE_SCOPE_HIERARCHICAL_COUNTER_FUNC()
+
+}
+
+bool UAnimGraphNode_AnimNextGraph::IsInputProperty(const FName& PropertyName) const
+{
+	DECLARE_SCOPE_HIERARCHICAL_COUNTER_FUNC()
+
+	// this is true for both input variables and controls
+	return InputVariables.Contains(PropertyName) || !OutputVariables.Contains(PropertyName);
+}
+
+bool UAnimGraphNode_AnimNextGraph::IsPropertyExposeEnabled(FName PropertyName) const
+{
+	DECLARE_SCOPE_HIERARCHICAL_COUNTER_FUNC()
+
+	// if known exposable, and and if it hasn't been exposed yet
+	if (CustomPinProperties.ContainsByPredicate([PropertyName](const FOptionalPinFromProperty& InOptionalPin){ return InOptionalPin.PropertyName == PropertyName; }))
+	{
+		return IsInputProperty(PropertyName);
+	}
+
+	return false;
+}
+
+ECheckBoxState UAnimGraphNode_AnimNextGraph::IsPropertyExposed(FName PropertyName) const
+{
+	if (CustomPinProperties.ContainsByPredicate([PropertyName](const FOptionalPinFromProperty& InOptionalPin){ return InOptionalPin.bShowPin && InOptionalPin.PropertyName == PropertyName; }))
+	{
+		return ECheckBoxState::Checked;
+	}
+
+	return ECheckBoxState::Unchecked;
+}
+
+void UAnimGraphNode_AnimNextGraph::OnPropertyExposeCheckboxChanged(ECheckBoxState NewState, FName PropertyName)
+{
+	DECLARE_SCOPE_HIERARCHICAL_COUNTER_FUNC()
+}
+
+void UAnimGraphNode_AnimNextGraph::CustomizeDetails(IDetailLayoutBuilder& DetailBuilder)
+{
+	DECLARE_SCOPE_HIERARCHICAL_COUNTER_FUNC()
+
+	Super::CustomizeDetails(DetailBuilder);
+
+	// We dont allow multi-select here
+	if (DetailBuilder.GetSelectedObjects().Num() > 1)
+	{
+		return;
+	}
+
+	// input/output exposure feature START
+	RebuildExposedProperties();
+
+	// ***************************************************************
+
+
+	TSharedRef<IPropertyHandle> ClassHandle = DetailBuilder.GetProperty(TEXT("Node.AnimNextGraph"), GetClass());
+	if (ClassHandle->IsValidHandle())
+	{
+		ClassHandle->SetOnPropertyValueChanged(FSimpleDelegate::CreateUObject(this, &UAnimGraphNode_AnimNextGraph::OnInstanceClassChanged, &DetailBuilder));
+	}
+}
+
+void UAnimGraphNode_AnimNextGraph::GetVariables(bool bInput, TMap<FName, FRigVMExternalVariable>& OutVariables) const
+{
+	DECLARE_SCOPE_HIERARCHICAL_COUNTER_FUNC()
+
+	OutVariables.Reset();
+
+	if (URigVMBlueprintGeneratedClass* TargetClass = Cast<URigVMBlueprintGeneratedClass>(GetTargetClass()))
+	{
+		if (UControlRigBlueprint* RigBlueprint = Cast<UControlRigBlueprint>(TargetClass->ClassGeneratedBy))
+		{
+			//RigBlueprint->CleanupVariables();
+			UControlRig* ControlRig = TargetClass->GetDefaultObject<UControlRig>();
+			if (ControlRig)
+			{
+				const TArray<FRigVMExternalVariable>& PublicVariables = ControlRig->GetPublicVariables();
+				for (const FRigVMExternalVariable& PublicVariable : PublicVariables)
+				{
+					if (!bInput || !PublicVariable.bIsReadOnly)
+					{
+						OutVariables.Add(PublicVariable.Name, PublicVariable);
+					}
+				}
+			}
+		}
+	}
+}
+
+void UAnimGraphNode_AnimNextGraph::PostEditChangeProperty(FPropertyChangedEvent& PropertyChangedEvent)
+{
+	DECLARE_SCOPE_HIERARCHICAL_COUNTER_FUNC()
+
+	Super::PostEditChangeProperty(PropertyChangedEvent);
+
+	bool bRequiresNodeReconstruct = false;
+	FProperty* ChangedProperty = PropertyChangedEvent.Property;
+
+	if (ChangedProperty)
+	{
+		if (ChangedProperty->GetFName() == GET_MEMBER_NAME_CHECKED(FAnimNode_AnimNextGraph, AnimNextGraph))
+		{
+			bRequiresNodeReconstruct = true;
+			RebuildExposedProperties();
+		}
+	}
+
+	if (bRequiresNodeReconstruct)
+	{
+		ReconstructNode();
+		FBlueprintEditorUtils::MarkBlueprintAsStructurallyModified(GetBlueprint());
+	}
+}
+
+void UAnimGraphNode_AnimNextGraph::PostReconstructNode()
+{
+	Super::PostReconstructNode();
+}
+
+void UAnimGraphNode_AnimNextGraph::CustomizePinData(UEdGraphPin* Pin, FName SourcePropertyName, int32 ArrayIndex) const
+{
+	Super::CustomizePinData(Pin, SourcePropertyName, ArrayIndex);
+}
+
+UObject* UAnimGraphNode_AnimNextGraph::GetJumpTargetForDoubleClick() const
+{
+	return Node.AnimNextGraph;
+}
+
+#undef LOCTEXT_NAMESPACE
diff --git a/Engine/Plugins/Experimental/AnimNextAnimGraph/Source/AnimNextAnimGraphUncookedOnly/Private/Module.cpp b/Engine/Plugins/Experimental/AnimNextAnimGraph/Source/AnimNextAnimGraphUncookedOnly/Private/Module.cpp
new file mode 100644
--- /dev/null
+++ b/Engine/Plugins/Experimental/AnimNextAnimGraph/Source/AnimNextAnimGraphUncookedOnly/Private/Module.cpp
@@ -0,0 +1,21 @@
+// Copyright Epic Games, Inc. All Rights Reserved.
+
+#include "CoreMinimal.h"
+#include "Modules/ModuleManager.h"
+
+namespace UE::AnimNext::UncookedOnly
+{
+	class FModule : public IModuleInterface
+	{
+	private:
+		virtual void StartupModule() override
+		{
+		}
+
+		virtual void ShutdownModule() override
+		{
+		}
+	};
+}
+
+IMPLEMENT_MODULE(UE::AnimNext::UncookedOnly::FModule, AnimNextAnimGraphUncookedOnly);
