Move module definitions into their own namespace to avoid clashing with the AnimNext core modules

diff --git a/Engine/Plugins/Experimental/AnimNextAnimGraph/Source/AnimNextAnimGraph/Private/Module.cpp b/Engine/Plugins/Experimental/AnimNextAnimGraph/Source/AnimNextAnimGraph/Private/Module.cpp
--- a/Engine/Plugins/Experimental/AnimNextAnimGraph/Source/AnimNextAnimGraph/Private/Module.cpp
+++ b/Engine/Plugins/Experimental/AnimNextAnimGraph/Source/AnimNextAnimGraph/Private/Module.cpp
@@ -14,7 +14,7 @@
 #include "TraitInterfaces/IUpdate.h"
 #include "EvaluationVM/EvaluationVM.h"
 
-namespace UE::AnimNext
+namespace UE::AnimNext::AnimGraph
 {
 
 class FModule : public IModuleInterface, IAnimNextAnimGraph
@@ -76,4 +76,4 @@ public:
 
 }
 
-IMPLEMENT_MODULE(UE::AnimNext::FModule, AnimNextAnimGraph)
+IMPLEMENT_MODULE(UE::AnimNext::AnimGraph::FModule, AnimNextAnimGraph)
diff --git a/Engine/Plugins/Experimental/AnimNextAnimGraph/Source/AnimNextAnimGraphTestSuite/Private/Module.cpp b/Engine/Plugins/Experimental/AnimNextAnimGraph/Source/AnimNextAnimGraphTestSuite/Private/Module.cpp
--- a/Engine/Plugins/Experimental/AnimNextAnimGraph/Source/AnimNextAnimGraphTestSuite/Private/Module.cpp
+++ b/Engine/Plugins/Experimental/AnimNextAnimGraph/Source/AnimNextAnimGraphTestSuite/Private/Module.cpp
@@ -4,14 +4,11 @@
 #include "Modules/ModuleInterface.h"
 #include "Modules/ModuleManager.h"
 
-namespace UE::AnimNext::Tests
+namespace UE::AnimNext::AnimGraph::Tests
 {
-
-class FModule : public IModuleInterface
-{
-public:
-};
-
+	class FModule : public IModuleInterface
+	{
+	};
 }
 
-IMPLEMENT_MODULE(UE::AnimNext::Tests::FModule, AnimNextAnimGraphTestSuite)
+IMPLEMENT_MODULE(UE::AnimNext::AnimGraph::Tests::FModule, AnimNextAnimGraphTestSuite)
diff --git a/Engine/Plugins/Experimental/AnimNextAnimGraph/Source/AnimNextAnimGraphUncookedOnly/Private/Module.cpp b/Engine/Plugins/Experimental/AnimNextAnimGraph/Source/AnimNextAnimGraphUncookedOnly/Private/Module.cpp
--- a/Engine/Plugins/Experimental/AnimNextAnimGraph/Source/AnimNextAnimGraphUncookedOnly/Private/Module.cpp
+++ b/Engine/Plugins/Experimental/AnimNextAnimGraph/Source/AnimNextAnimGraphUncookedOnly/Private/Module.cpp
@@ -3,19 +3,11 @@
 #include "CoreMinimal.h"
 #include "Modules/ModuleManager.h"
 
-namespace UE::AnimNext::UncookedOnly
+namespace UE::AnimNext::AnimGraph::UncookedOnly
 {
 	class FModule : public IModuleInterface
 	{
-	private:
-		virtual void StartupModule() override
-		{
-		}
-
-		virtual void ShutdownModule() override
-		{
-		}
 	};
 }
 
-IMPLEMENT_MODULE(UE::AnimNext::UncookedOnly::FModule, AnimNextAnimGraphUncookedOnly);
+IMPLEMENT_MODULE(UE::AnimNext::AnimGraph::UncookedOnly::FModule, AnimNextAnimGraphUncookedOnly);
