Expose sequence as a latent property in sequence player

diff --git a/Engine/Plugins/Experimental/AnimNextAnimGraph/Source/AnimNextAnimGraph/Private/Traits/SequencePlayer.cpp b/Engine/Plugins/Experimental/AnimNextAnimGraph/Source/AnimNextAnimGraph/Private/Traits/SequencePlayer.cpp
--- a/Engine/Plugins/Experimental/AnimNextAnimGraph/Source/AnimNextAnimGraph/Private/Traits/SequencePlayer.cpp
+++ b/Engine/Plugins/Experimental/AnimNextAnimGraph/Source/AnimNextAnimGraph/Private/Traits/SequencePlayer.cpp
@@ -15,20 +15,37 @@ namespace UE::AnimNext
 		GeneratorMacro(IEvaluate) \
 		GeneratorMacro(ITimeline) \
 		GeneratorMacro(IUpdate) \
+		GeneratorMacro(IGarbageCollection) \
 
 	GENERATE_ANIM_TRAIT_IMPLEMENTATION(FSequencePlayerTrait, TRAIT_INTERFACE_ENUMERATOR, NULL_ANIM_TRAIT_EVENT_ENUMERATOR)
 	#undef TRAIT_INTERFACE_ENUMERATOR
 
 	void FSequencePlayerTrait::FInstanceData::Construct(const FExecutionContext& Context, const FTraitBinding& Binding)
 	{
+		FTrait::FInstanceData::Construct(Context, Binding);
+
+		IGarbageCollection::RegisterWithGC(Context, Binding);
+
 		const FSharedData* SharedData = Binding.GetSharedData<FSharedData>();
-		if (SharedData->AnimSequence != nullptr)
+		FInstanceData* InstanceData = Binding.GetInstanceData<FInstanceData>();
+
+		// Cache the anim sequence we'll play during construction, we don't allow it to change afterwards
+		InstanceData->AnimSequence = SharedData->GetAnimSequence(Binding);
+
+		if (InstanceData->AnimSequence)
 		{
-			const float SequenceLength = SharedData->AnimSequence->GetPlayLength();
+			const float SequenceLength = InstanceData->AnimSequence->GetPlayLength();
 			InternalTimeAccumulator = FMath::Clamp(SharedData->GetStartPosition(Binding), 0.0f, SequenceLength);
 		}
 	}
 
+	void FSequencePlayerTrait::FInstanceData::Destruct(const FExecutionContext& Context, const FTraitBinding& Binding)
+	{
+		FTrait::FInstanceData::Destruct(Context, Binding);
+
+		IGarbageCollection::UnregisterWithGC(Context, Binding);
+	}
+
 	void FSequencePlayerTrait::PreEvaluate(FEvaluateTraversalContext& Context, const TTraitBinding<IEvaluate>& Binding) const
 	{
 		const FSharedData* SharedData = Binding.GetSharedData<FSharedData>();
@@ -36,7 +53,7 @@ namespace UE::AnimNext
 		const FInstanceData* InstanceData = Binding.GetInstanceData<FInstanceData>();
 		const bool bInterpolate = true;
 
-		FAnimNextAnimSequenceKeyframeTask Task = FAnimNextAnimSequenceKeyframeTask::MakeFromSampleTime(SharedData->AnimSequence, InstanceData->InternalTimeAccumulator, bInterpolate);
+		FAnimNextAnimSequenceKeyframeTask Task = FAnimNextAnimSequenceKeyframeTask::MakeFromSampleTime(InstanceData->AnimSequence, InstanceData->InternalTimeAccumulator, bInterpolate);
 		Task.bExtractTrajectory = true;	/*Output.AnimInstanceProxy->ShouldExtractRootMotion()*/
 
 		Context.AppendTask(Task);
@@ -50,10 +67,11 @@ namespace UE::AnimNext
 
 	float FSequencePlayerTrait::AdvanceBy(const FExecutionContext& Context, const TTraitBinding<ITimeline>& Binding, float DeltaTime) const
 	{
-		const FSharedData* SharedData = Binding.GetSharedData<FSharedData>();
-		if (UAnimSequence* AnimSeq = SharedData->AnimSequence.Get())
+		FInstanceData* InstanceData = Binding.GetInstanceData<FInstanceData>();
+
+		if (UAnimSequence* AnimSeq = InstanceData->AnimSequence.Get())
 		{
-			FInstanceData* InstanceData = Binding.GetInstanceData<FInstanceData>();
+			const FSharedData* SharedData = Binding.GetSharedData<FSharedData>();
 
 			TTraitBinding<ITimeline> TimelineTrait;
 			Binding.GetStackInterface(TimelineTrait);
@@ -72,11 +90,10 @@ namespace UE::AnimNext
 
 	void FSequencePlayerTrait::AdvanceToRatio(const FExecutionContext& Context, const TTraitBinding<ITimeline>& Binding, float ProgressRatio) const
 	{
-		const FSharedData* SharedData = Binding.GetSharedData<FSharedData>();
-		if (UAnimSequence* AnimSeq = SharedData->AnimSequence.Get())
-		{
-			FInstanceData* InstanceData = Binding.GetInstanceData<FInstanceData>();
+		FInstanceData* InstanceData = Binding.GetInstanceData<FInstanceData>();
 
+		if (UAnimSequence* AnimSeq = InstanceData->AnimSequence.Get())
+		{
 			const float SequenceLength = AnimSeq->GetPlayLength();
 
 			InstanceData->InternalTimeAccumulator = FMath::Clamp(ProgressRatio, 0.0f, 1.0f) * SequenceLength;
@@ -91,4 +108,13 @@ namespace UE::AnimNext
 
 		TimelineTrait.AdvanceBy(Context, TraitState.GetDeltaTime());
 	}
+
+	void FSequencePlayerTrait::AddReferencedObjects(const FExecutionContext& Context, const TTraitBinding<IGarbageCollection>& Binding, FReferenceCollector& Collector) const
+	{
+		IGarbageCollection::AddReferencedObjects(Context, Binding, Collector);
+
+		FInstanceData* InstanceData = Binding.GetInstanceData<FInstanceData>();
+
+		Collector.AddReferencedObject(InstanceData->AnimSequence);
+	}
 }
diff --git a/Engine/Plugins/Experimental/AnimNextAnimGraph/Source/AnimNextAnimGraph/Private/Traits/SequencePlayer.h b/Engine/Plugins/Experimental/AnimNextAnimGraph/Source/AnimNextAnimGraph/Private/Traits/SequencePlayer.h
--- a/Engine/Plugins/Experimental/AnimNextAnimGraph/Source/AnimNextAnimGraph/Private/Traits/SequencePlayer.h
+++ b/Engine/Plugins/Experimental/AnimNextAnimGraph/Source/AnimNextAnimGraph/Private/Traits/SequencePlayer.h
@@ -5,6 +5,7 @@
 #include "CoreMinimal.h"
 #include "TraitCore/Trait.h"
 #include "TraitInterfaces/IEvaluate.h"
+#include "TraitInterfaces/IGarbageCollection.h"
 #include "TraitInterfaces/ITimeline.h"
 #include "TraitInterfaces/IUpdate.h"
 #include "Animation/AnimSequence.h"
@@ -17,7 +18,7 @@ struct FAnimNextSequencePlayerTraitSharedData : public FAnimNextTraitSharedData
 	GENERATED_BODY()
 
 	/** The sequence to play. */
-	UPROPERTY(EditAnywhere, Category = "Default", meta = (Inline))
+	UPROPERTY(EditAnywhere, Category = "Default")
 	TObjectPtr<UAnimSequence> AnimSequence;
 
 	/** The play rate multiplier at which this sequence plays. */
@@ -34,6 +35,7 @@ struct FAnimNextSequencePlayerTraitSharedData : public FAnimNextTraitSharedData
 
 	// Latent pin support boilerplate
 	#define TRAIT_LATENT_PROPERTIES_ENUMERATOR(GeneratorMacro) \
+		GeneratorMacro(AnimSequence) \
 		GeneratorMacro(PlayRate) \
 		GeneratorMacro(StartPosition) \
 		GeneratorMacro(bLoop) \
@@ -49,7 +51,7 @@ namespace UE::AnimNext
 	 * 
 	 * A trait that can play an animation sequence.
 	 */
-	struct FSequencePlayerTrait : FBaseTrait, IEvaluate, ITimeline, IUpdate
+	struct FSequencePlayerTrait : FBaseTrait, IEvaluate, ITimeline, IUpdate, IGarbageCollection
 	{
 		DECLARE_ANIM_TRAIT(FSequencePlayerTrait, 0x7a0dc157, FBaseTrait)
 
@@ -57,9 +59,13 @@ namespace UE::AnimNext
 
 		struct FInstanceData : FTrait::FInstanceData
 		{
+			// Cached value of the anim sequence we are playing
+			TObjectPtr<UAnimSequence> AnimSequence;
+
 			float InternalTimeAccumulator = 0.0f;
 
 			void Construct(const FExecutionContext& Context, const FTraitBinding& Binding);
+			void Destruct(const FExecutionContext& Context, const FTraitBinding& Binding);
 		};
 
 		// IEvaluate impl
@@ -72,5 +78,8 @@ namespace UE::AnimNext
 
 		// IUpdate impl
 		virtual void PreUpdate(FUpdateTraversalContext& Context, const TTraitBinding<IUpdate>& Binding, const FTraitUpdateState& TraitState) const override;
+
+		// IGarbageCollection impl
+		virtual void AddReferencedObjects(const FExecutionContext& Context, const TTraitBinding<IGarbageCollection>& Binding, FReferenceCollector& Collector) const override;
 	};
 }
diff --git a/Engine/Plugins/Experimental/AnimNextAnimGraph/Source/AnimNextAnimGraph/Private/Traits/SubGraphHost.cpp b/Engine/Plugins/Experimental/AnimNextAnimGraph/Source/AnimNextAnimGraph/Private/Traits/SubGraphHost.cpp
--- a/Engine/Plugins/Experimental/AnimNextAnimGraph/Source/AnimNextAnimGraph/Private/Traits/SubGraphHost.cpp
+++ b/Engine/Plugins/Experimental/AnimNextAnimGraph/Source/AnimNextAnimGraph/Private/Traits/SubGraphHost.cpp
@@ -233,6 +233,8 @@ namespace UE::AnimNext
 
 	void FSubGraphHostTrait::AddReferencedObjects(const FExecutionContext& Context, const TTraitBinding<IGarbageCollection>& Binding, FReferenceCollector& Collector) const
 	{
+		IGarbageCollection::AddReferencedObjects(Context, Binding, Collector);
+
 		FInstanceData* InstanceData = Binding.GetInstanceData<FInstanceData>();
 
 		for (FSubGraphSlot& SubGraphEntry : InstanceData->SubGraphSlots)
