AnimNext: Added DeadBlending

diff --git a/Engine/Plugins/Experimental/AnimNext/Source/AnimNext/Private/EvaluationVM/Tasks/StoreKeyframe.cpp b/Engine/Plugins/Experimental/AnimNext/Source/AnimNext/Private/EvaluationVM/Tasks/StoreKeyframe.cpp
new file mode 100644
--- /dev/null
+++ b/Engine/Plugins/Experimental/AnimNext/Source/AnimNext/Private/EvaluationVM/Tasks/StoreKeyframe.cpp
@@ -0,0 +1,41 @@
+// Copyright Epic Games, Inc. All Rights Reserved.
+
+#include "EvaluationVM/Tasks/StoreKeyframe.h"
+
+#include "EvaluationVM/EvaluationVM.h"
+#include "EvaluationVM/KeyframeState.h"
+#include "TransformArrayOperations.h"
+
+FAnimNextSwapTransformsTask FAnimNextSwapTransformsTask::Make(UE::AnimNext::FTransformArraySoAHeap* A, UE::AnimNext::FTransformArraySoAHeap* B)
+{
+	FAnimNextSwapTransformsTask Task;
+	Task.A = A;
+	Task.B = B;
+	return Task;
+}
+
+void FAnimNextSwapTransformsTask::Execute(UE::AnimNext::FEvaluationVM& VM) const
+{
+	Swap(*A, *B);
+}
+
+FAnimNextStoreKeyframeTransformsTask FAnimNextStoreKeyframeTransformsTask::Make(UE::AnimNext::FTransformArraySoAHeap* Dest)
+{
+	FAnimNextStoreKeyframeTransformsTask Task;
+	Task.Dest = Dest;
+	return Task;
+}
+
+void FAnimNextStoreKeyframeTransformsTask::Execute(UE::AnimNext::FEvaluationVM& VM) const
+{
+	using namespace UE::AnimNext;
+
+	if (EnumHasAnyFlags(VM.GetFlags(), EEvaluationFlags::Bones))
+	{
+		if (const TUniquePtr<FKeyframeState>* Keyframe = VM.PeekValue<TUniquePtr<FKeyframeState>>(KEYFRAME_STACK_NAME, 0))
+		{
+			Dest->SetNumUninitialized((*Keyframe)->Pose.LocalTransforms.Num());
+			CopyTransforms(Dest->GetView(), (*Keyframe)->Pose.LocalTransforms.GetConstView());
+		}
+	}
+}
diff --git a/Engine/Plugins/Experimental/AnimNext/Source/AnimNext/Public/EvaluationVM/Tasks/StoreKeyframe.h b/Engine/Plugins/Experimental/AnimNext/Source/AnimNext/Public/EvaluationVM/Tasks/StoreKeyframe.h
new file mode 100644
--- /dev/null
+++ b/Engine/Plugins/Experimental/AnimNext/Source/AnimNext/Public/EvaluationVM/Tasks/StoreKeyframe.h
@@ -0,0 +1,46 @@
+// Copyright Epic Games, Inc. All Rights Reserved.
+
+#pragma once
+
+#include "CoreMinimal.h"
+#include "EvaluationVM/EvaluationTask.h"
+#include "TransformArray.h"
+
+#include "StoreKeyframe.generated.h"
+
+/**
+ * Swap the two given Transform Arrays
+ */
+USTRUCT()
+struct ANIMNEXT_API FAnimNextSwapTransformsTask : public FAnimNextEvaluationTask
+{
+	GENERATED_BODY()
+
+	DECLARE_ANIM_EVALUATION_TASK(FAnimNextSwapTransformsTask)
+
+	static FAnimNextSwapTransformsTask Make(UE::AnimNext::FTransformArraySoAHeap* A, UE::AnimNext::FTransformArraySoAHeap* B);
+
+	// Task entry point
+	virtual void Execute(UE::AnimNext::FEvaluationVM& VM) const override;
+
+	UE::AnimNext::FTransformArraySoAHeap* A = nullptr;
+	UE::AnimNext::FTransformArraySoAHeap* B = nullptr;
+};
+
+/*
+ * Store the pose on top of the stack in the given Transform Array
+ */
+USTRUCT()
+struct ANIMNEXT_API FAnimNextStoreKeyframeTransformsTask : public FAnimNextEvaluationTask
+{
+	GENERATED_BODY()
+
+	DECLARE_ANIM_EVALUATION_TASK(FAnimNextStoreKeyframeTransformsTask)
+
+	static FAnimNextStoreKeyframeTransformsTask Make(UE::AnimNext::FTransformArraySoAHeap* Dest);
+
+	// Task entry point
+	virtual void Execute(UE::AnimNext::FEvaluationVM& VM) const override;
+
+	UE::AnimNext::FTransformArraySoAHeap* Dest = nullptr;
+};
diff --git a/Engine/Plugins/Experimental/AnimNext/Source/AnimNext/Public/TransformArray.h b/Engine/Plugins/Experimental/AnimNext/Source/AnimNext/Public/TransformArray.h
--- a/Engine/Plugins/Experimental/AnimNext/Source/AnimNext/Public/TransformArray.h
+++ b/Engine/Plugins/Experimental/AnimNext/Source/AnimNext/Public/TransformArray.h
@@ -318,6 +318,17 @@ struct TTransformArraySoA
 		return *this;
 	}
 
+	inline bool IsEmpty() const
+	{
+		return Num() == 0;
+	}
+
+	inline void Empty(int32 Slack = 0)
+	{
+		AllocatedMemory.Empty(Slack);
+		UpdateViews(AllocatedMemory.GetData(), 0);
+	}
+
 	inline void Reset(int32 NumTransforms)
 	{
 		constexpr int32 TransformSize = sizeof(FVector) + sizeof(FQuat) + sizeof(FVector);
diff --git a/Engine/Plugins/Experimental/AnimNext/Source/AnimNext/Public/TransformArrayView.h b/Engine/Plugins/Experimental/AnimNext/Source/AnimNext/Public/TransformArrayView.h
--- a/Engine/Plugins/Experimental/AnimNext/Source/AnimNext/Public/TransformArrayView.h
+++ b/Engine/Plugins/Experimental/AnimNext/Source/AnimNext/Public/TransformArrayView.h
@@ -93,6 +93,7 @@ namespace UE::AnimNext
 		TArrayView<FQuat> Rotations;
 		TArrayView<FVector> Scales3D;
 
+		bool IsEmpty() const { return Rotations.IsEmpty(); }
 		int32 Num() const { return Rotations.Num(); }
 
 		inline FTransformSoAAdapter operator[](int32 Index)
@@ -119,6 +120,7 @@ namespace UE::AnimNext
 			, Scales3D(Other.Scales3D)
 		{}
 
+		bool IsEmpty() const { return Rotations.IsEmpty(); }
 		int32 Num() const { return Rotations.Num(); }
 
 		inline const FTransformSoAAdapterConst operator[](int32 Index) const
diff --git a/Engine/Plugins/Experimental/AnimNextAnimGraph/Source/AnimNextAnimGraph/Private/EvaluationVM/Tasks/DeadBlending.cpp b/Engine/Plugins/Experimental/AnimNextAnimGraph/Source/AnimNextAnimGraph/Private/EvaluationVM/Tasks/DeadBlending.cpp
new file mode 100644
--- /dev/null
+++ b/Engine/Plugins/Experimental/AnimNextAnimGraph/Source/AnimNextAnimGraph/Private/EvaluationVM/Tasks/DeadBlending.cpp
@@ -0,0 +1,579 @@
+// Copyright Epic Games, Inc. All Rights Reserved.
+
+#include "EvaluationVM/Tasks/DeadBlending.h"
+
+#include "EvaluationVM/EvaluationVM.h"
+#include "EvaluationVM/KeyframeState.h"
+
+#include "Traits/Inertialization.h"
+
+#ifndef UE_ANIM_NEXT_DEAD_BLENDING_ISPC
+#define UE_ANIM_NEXT_DEAD_BLENDING_ISPC INTEL_ISPC
+//#define UE_ANIM_NEXT_DEAD_BLENDING_ISPC 0
+#endif
+
+#if UE_ANIM_NEXT_DEAD_BLENDING_ISPC
+#include "DeadBlending.ispc.generated.h"
+#endif
+
+namespace UE::AnimNext::DeadBlending::Private
+{
+#if !UE_ANIM_NEXT_DEAD_BLENDING_ISPC
+	static constexpr float Ln2 = 0.69314718056f;
+
+	static inline FVector3f VectorDivMax(const float V, const FVector3f W, const float Epsilon = UE_SMALL_NUMBER)
+	{
+		return FVector3f(
+			V / FMath::Max(W.X, Epsilon),
+			V / FMath::Max(W.Y, Epsilon),
+			V / FMath::Max(W.Z, Epsilon));
+	}
+
+	static inline FVector3f VectorDivMax(const FVector3f V, const FVector3f W, const float Epsilon = UE_SMALL_NUMBER)
+	{
+		return FVector3f(
+			V.X / FMath::Max(W.X, Epsilon),
+			V.Y / FMath::Max(W.Y, Epsilon),
+			V.Z / FMath::Max(W.Z, Epsilon));
+	}
+
+	static inline FVector VectorDivMax(const FVector V, const FVector W, const float Epsilon = UE_SMALL_NUMBER)
+	{
+		return FVector(
+			V.X / FMath::Max(W.X, Epsilon),
+			V.Y / FMath::Max(W.Y, Epsilon),
+			V.Z / FMath::Max(W.Z, Epsilon));
+	}
+
+	static inline FVector3f VectorInvExpApprox(const FVector3f V)
+	{
+		return FVector3f(
+			FMath::InvExpApprox(V.X),
+			FMath::InvExpApprox(V.Y),
+			FMath::InvExpApprox(V.Z));
+	}
+
+	static inline FVector VectorEerp(const FVector V, const FVector W, const float Alpha, const float Epsilon = UE_SMALL_NUMBER)
+	{
+		if (FVector::DistSquared(V, W) < Epsilon)
+		{
+			return FVector(
+				FMath::Lerp(FMath::Max(V.X, Epsilon), FMath::Max(W.X, Epsilon), Alpha),
+				FMath::Lerp(FMath::Max(V.Y, Epsilon), FMath::Max(W.Y, Epsilon), Alpha),
+				FMath::Lerp(FMath::Max(V.Z, Epsilon), FMath::Max(W.Z, Epsilon), Alpha));
+		}
+		else
+		{
+			return FVector(
+				FMath::Pow(FMath::Max(V.X, Epsilon), (1.0f - Alpha)) * FMath::Pow(FMath::Max(W.X, Epsilon), Alpha),
+				FMath::Pow(FMath::Max(V.Y, Epsilon), (1.0f - Alpha)) * FMath::Pow(FMath::Max(W.Y, Epsilon), Alpha),
+				FMath::Pow(FMath::Max(V.Z, Epsilon), (1.0f - Alpha)) * FMath::Pow(FMath::Max(W.Z, Epsilon), Alpha));
+		}
+	}
+
+	static inline FVector VectorExp(const FVector V)
+	{
+		return FVector(
+			FMath::Exp(V.X),
+			FMath::Exp(V.Y),
+			FMath::Exp(V.Z));
+	}
+
+	static inline FVector VectorLogMax(const FVector V, const float Epsilon = UE_SMALL_NUMBER)
+	{
+		return FVector(
+			FMath::Loge(FMath::Max(V.X, Epsilon)),
+			FMath::Loge(FMath::Max(V.Y, Epsilon)),
+			FMath::Loge(FMath::Max(V.Z, Epsilon)));
+	}
+
+	static inline FVector ExtrapolateTranslation(
+		const FVector Translation,
+		const FVector3f Velocity,
+		const FVector3f DecayHalflife,
+		const float Time,
+		const float Epsilon = UE_SMALL_NUMBER)
+	{
+		if (Velocity.SquaredLength() > Epsilon)
+		{
+			const FVector3f C = VectorDivMax(Ln2, DecayHalflife, Epsilon);
+			return Translation + (FVector)(VectorDivMax(Velocity, C, Epsilon) * (FVector3f::OneVector - VectorInvExpApprox(C * Time)));
+		}
+		else
+		{
+			return Translation;
+		}
+	}
+
+	static inline FQuat ExtrapolateRotation(
+		const FQuat Rotation,
+		const FVector3f Velocity,
+		const FVector3f DecayHalflife,
+		const float Time,
+		const float Epsilon = UE_SMALL_NUMBER)
+	{
+		if (Velocity.SquaredLength() > Epsilon)
+		{
+			const FVector3f C = VectorDivMax(Ln2, DecayHalflife, Epsilon);
+			return FQuat::MakeFromRotationVector((FVector)(VectorDivMax(Velocity, C, Epsilon) * (FVector3f::OneVector - VectorInvExpApprox(C * Time)))) * Rotation;
+		}
+		else
+		{
+			return Rotation;
+		}
+	}
+
+	static inline FVector ExtrapolateScale(
+		const FVector Scale,
+		const FVector3f Velocity,
+		const FVector3f DecayHalflife,
+		const float Time,
+		const float Epsilon = UE_SMALL_NUMBER)
+	{
+		if (Velocity.SquaredLength() > Epsilon)
+		{
+			const FVector3f C = VectorDivMax(Ln2, DecayHalflife, Epsilon);
+			return VectorExp((FVector)(VectorDivMax(Velocity, C, Epsilon) * (FVector3f::OneVector - VectorInvExpApprox(C * Time)))) * Scale;
+		}
+		else
+		{
+			return Scale;
+		}
+	}
+
+	static inline float ClipMagnitudeToGreaterThanEpsilon(const float X, const float Epsilon = UE_KINDA_SMALL_NUMBER)
+	{
+		return
+			X >= 0.0f && X < Epsilon ? Epsilon :
+			X <  0.0f && X > -Epsilon ? -Epsilon : X;
+	}
+
+	static inline float ComputeDecayHalfLifeFromDiffAndVelocity(
+		const float SrcDstDiff,
+		const float SrcVelocity,
+		const float HalfLife,
+		const float HalfLifeMin,
+		const float HalfLifeMax,
+		const float Epsilon = UE_KINDA_SMALL_NUMBER)
+	{
+		// Essentially what this function does is compute a half-life based on the ratio between the velocity vector and
+		// the vector from the source to the destination. This is then clamped to some min and max. If the signs are
+		// different (i.e. the velocity and the vector from source to destination are in opposite directions) this will
+		// produce a negative number that will get clamped to HalfLifeMin. If the signs match, this will produce a large
+		// number when the velocity is small and the vector from source to destination is large, and a small number when
+		// the velocity is large and the vector from source to destination is small. This will be clamped either way to 
+		// be in the range given by HalfLifeMin and HalfLifeMax. Finally, since the velocity can be close to zero we 
+		// have to clamp it to always be greater than some given magnitude (preserving the sign).
+
+		return FMath::Clamp(HalfLife * (SrcDstDiff / ClipMagnitudeToGreaterThanEpsilon(SrcVelocity, Epsilon)), HalfLifeMin, HalfLifeMax);
+	}
+
+	static inline FVector3f ComputeDecayHalfLifeFromDiffAndVelocity(
+		const FVector SrcDstDiff,
+		const FVector3f SrcVelocity,
+		const float HalfLife,
+		const float HalfLifeMin,
+		const float HalfLifeMax,
+		const float Epsilon = UE_KINDA_SMALL_NUMBER)
+	{
+		return FVector3f(
+			ComputeDecayHalfLifeFromDiffAndVelocity(SrcDstDiff.X, SrcVelocity.X, HalfLife, HalfLifeMin, HalfLifeMax, Epsilon),
+			ComputeDecayHalfLifeFromDiffAndVelocity(SrcDstDiff.Y, SrcVelocity.Y, HalfLife, HalfLifeMin, HalfLifeMax, Epsilon),
+			ComputeDecayHalfLifeFromDiffAndVelocity(SrcDstDiff.Z, SrcVelocity.Z, HalfLife, HalfLifeMin, HalfLifeMax, Epsilon));
+	}
+#endif
+
+	void Transition(
+		const TArrayView<FQuat4f>& BoneRotationDirections,
+		const FTransformArraySoAView& Source,
+		const TArrayView<FVector3f>& SourceBoneTranslationVelocities,
+		const TArrayView<FVector3f>& SourceBoneRotationVelocities,
+		const TArrayView<FVector3f>& SourceBoneScaleVelocities,
+		const TArrayView<FVector3f>& SourceBoneTranslationDecayHalfLives,
+		const TArrayView<FVector3f>& SourceBoneRotationDecayHalfLives,
+		const TArrayView<FVector3f>& SourceBoneScaleDecayHalfLives,
+		const FTransformArraySoAConstView& Dest,
+		const FTransformArraySoAConstView& Curr,
+		const FTransformArraySoAConstView& Prev,
+		const float DeltaTime,
+		const UE::AnimNext::FDeadBlendTransitionTaskParameters& Parameters)
+	{
+		TRACE_CPUPROFILER_EVENT_SCOPE(AnimNext::DeadBlending::Transition);
+
+		const int32 LODBoneNum = Source.Num();
+		check(LODBoneNum <= Dest.Num());
+		check(LODBoneNum <= Curr.Num());
+		check(LODBoneNum <= Prev.Num());
+
+#if UE_ANIM_NEXT_DEAD_BLENDING_ISPC
+		ispc::AnimNextDeadBlendingTransition(
+			(ispc::FVector4f*)BoneRotationDirections.GetData(),
+			(ispc::FVector*)Source.Translations.GetData(),
+			(ispc::FVector4*)Source.Rotations.GetData(),
+			(ispc::FVector*)Source.Scales3D.GetData(),
+			(ispc::FVector3f*)SourceBoneTranslationVelocities.GetData(),
+			(ispc::FVector3f*)SourceBoneRotationVelocities.GetData(),
+			(ispc::FVector3f*)SourceBoneScaleVelocities.GetData(),
+			(ispc::FVector3f*)SourceBoneTranslationDecayHalfLives.GetData(),
+			(ispc::FVector3f*)SourceBoneRotationDecayHalfLives.GetData(),
+			(ispc::FVector3f*)SourceBoneScaleDecayHalfLives.GetData(),
+			(ispc::FVector*)Dest.Translations.GetData(),
+			(ispc::FVector4*)Dest.Rotations.GetData(),
+			(ispc::FVector*)Dest.Scales3D.GetData(),
+			(ispc::FVector*)Curr.Translations.GetData(),
+			(ispc::FVector4*)Curr.Rotations.GetData(),
+			(ispc::FVector*)Curr.Scales3D.GetData(),
+			(ispc::FVector*)Prev.Translations.GetData(),
+			(ispc::FVector4*)Prev.Rotations.GetData(),
+			(ispc::FVector*)Prev.Scales3D.GetData(),
+			LODBoneNum,
+			DeltaTime,
+			Parameters.ExtrapolationHalfLife,
+			Parameters.ExtrapolationHalfLifeMin,
+			Parameters.ExtrapolationHalfLifeMax,
+			Parameters.MaximumTranslationVelocity,
+			Parameters.MaximumRotationVelocity,
+			Parameters.MaximumScaleVelocity);
+#else
+		for (FBoneIndexType LODBoneIdx = 0; LODBoneIdx < LODBoneNum; LODBoneIdx++)
+		{
+			BoneRotationDirections[LODBoneIdx] = FQuat4f::Identity;
+
+			// Get Source Animation Transform
+
+			const FVector SrcTranslationCurr = Curr.Translations[LODBoneIdx];
+			const FQuat SrcRotationCurr = Curr.Rotations[LODBoneIdx];
+			const FVector SrcScaleCurr = Curr.Scales3D[LODBoneIdx];
+
+			Source.Translations[LODBoneIdx] = SrcTranslationCurr;
+			Source.Rotations[LODBoneIdx] = SrcRotationCurr;
+			Source.Scales3D[LODBoneIdx] = SrcScaleCurr;
+
+			// Get Source Animation Velocity
+
+			const FVector SrcTranslationPrev = Prev.Translations[LODBoneIdx];
+			const FQuat SrcRotationPrev = Prev.Rotations[LODBoneIdx];
+			const FVector SrcScalePrev = Prev.Scales3D[LODBoneIdx];
+
+			const FVector TranslationDiff = SrcTranslationCurr - SrcTranslationPrev;
+
+			FQuat RotationDiff = SrcRotationCurr * SrcRotationPrev.Inverse();
+			RotationDiff.EnforceShortestArcWith(FQuat::Identity);
+
+			const FVector ScaleDiff = Private::VectorDivMax(SrcScaleCurr, SrcScalePrev);
+
+			SourceBoneTranslationVelocities[LODBoneIdx] = (FVector3f)(TranslationDiff / DeltaTime).GetClampedToMaxSize(Parameters.MaximumTranslationVelocity);
+			SourceBoneRotationVelocities[LODBoneIdx] = (FVector3f)(RotationDiff.ToRotationVector() / DeltaTime).GetClampedToMaxSize(Parameters.MaximumRotationVelocity);
+			SourceBoneScaleVelocities[LODBoneIdx] = (FVector3f)(Private::VectorLogMax(ScaleDiff) / DeltaTime).GetClampedToMaxSize(Parameters.MaximumScaleVelocity);
+
+			// Compute Decay HalfLives
+
+			const FVector DstTranslation = Dest.Translations[LODBoneIdx];
+			const FQuat DstRotation = Dest.Rotations[LODBoneIdx];
+			const FVector DstScale = Dest.Scales3D[LODBoneIdx];
+
+			const FVector TranslationSrcDstDiff = DstTranslation - SrcTranslationCurr;
+
+			FQuat RotationSrcDstDiff = DstRotation * SrcRotationCurr.Inverse();
+			RotationSrcDstDiff.EnforceShortestArcWith(FQuat::Identity);
+
+			const FVector ScaleSrcDstDiff = Private::VectorDivMax(DstScale, SrcScaleCurr);
+
+			SourceBoneTranslationDecayHalfLives[LODBoneIdx] = Private::ComputeDecayHalfLifeFromDiffAndVelocity(
+				TranslationSrcDstDiff,
+				SourceBoneTranslationVelocities[LODBoneIdx],
+				Parameters.ExtrapolationHalfLife,
+				Parameters.ExtrapolationHalfLifeMin,
+				Parameters.ExtrapolationHalfLifeMax);
+
+			SourceBoneRotationDecayHalfLives[LODBoneIdx] = Private::ComputeDecayHalfLifeFromDiffAndVelocity(
+				RotationSrcDstDiff.ToRotationVector(),
+				SourceBoneRotationVelocities[LODBoneIdx],
+				Parameters.ExtrapolationHalfLife,
+				Parameters.ExtrapolationHalfLifeMin,
+				Parameters.ExtrapolationHalfLifeMax);
+
+			SourceBoneScaleDecayHalfLives[LODBoneIdx] = Private::ComputeDecayHalfLifeFromDiffAndVelocity(
+				ScaleSrcDstDiff,
+				SourceBoneScaleVelocities[LODBoneIdx],
+				Parameters.ExtrapolationHalfLife,
+				Parameters.ExtrapolationHalfLifeMin,
+				Parameters.ExtrapolationHalfLifeMax);
+		}
+#endif
+	}
+
+
+	void Transition(
+		const TArrayView<FQuat4f>& BoneRotationDirections,
+		const FTransformArraySoAView& Source,
+		const TArrayView<FVector3f>& SourceBoneTranslationVelocities,
+		const TArrayView<FVector3f>& SourceBoneRotationVelocities,
+		const TArrayView<FVector3f>& SourceBoneScaleVelocities,
+		const TArrayView<FVector3f>& SourceBoneTranslationDecayHalfLives,
+		const TArrayView<FVector3f>& SourceBoneRotationDecayHalfLives,
+		const TArrayView<FVector3f>& SourceBoneScaleDecayHalfLives,
+		const FTransformArraySoAConstView& Curr,
+		const UE::AnimNext::FDeadBlendTransitionTaskParameters& Parameters)
+	{
+		TRACE_CPUPROFILER_EVENT_SCOPE(AnimNext::DeadBlending::Transition);
+
+		const int32 LODBoneNum = Source.Num();
+		check(LODBoneNum <= Curr.Num());
+
+#if UE_ANIM_NEXT_DEAD_BLENDING_ISPC
+		ispc::AnimNextDeadBlendingTransitionStatic(
+			(ispc::FVector4f*)BoneRotationDirections.GetData(),
+			(ispc::FVector*)Source.Translations.GetData(),
+			(ispc::FVector4*)Source.Rotations.GetData(),
+			(ispc::FVector*)Source.Scales3D.GetData(),
+			(ispc::FVector3f*)SourceBoneTranslationVelocities.GetData(),
+			(ispc::FVector3f*)SourceBoneRotationVelocities.GetData(),
+			(ispc::FVector3f*)SourceBoneScaleVelocities.GetData(),
+			(ispc::FVector3f*)SourceBoneTranslationDecayHalfLives.GetData(),
+			(ispc::FVector3f*)SourceBoneRotationDecayHalfLives.GetData(),
+			(ispc::FVector3f*)SourceBoneScaleDecayHalfLives.GetData(),
+			(ispc::FVector*)Curr.Translations.GetData(),
+			(ispc::FVector4*)Curr.Rotations.GetData(),
+			(ispc::FVector*)Curr.Scales3D.GetData(),
+			LODBoneNum,
+			Parameters.ExtrapolationHalfLifeMin);
+#else
+		for (FBoneIndexType LODBoneIdx = 0; LODBoneIdx < LODBoneNum; LODBoneIdx++)
+		{
+			BoneRotationDirections[LODBoneIdx] = FQuat4f::Identity;
+
+			Source.Translations[LODBoneIdx] = Curr.Translations[LODBoneIdx];
+			Source.Rotations[LODBoneIdx] = Curr.Rotations[LODBoneIdx];
+			Source.Scales3D[LODBoneIdx] = Curr.Scales3D[LODBoneIdx];
+
+			SourceBoneTranslationVelocities[LODBoneIdx] = FVector3f::ZeroVector;
+			SourceBoneRotationVelocities[LODBoneIdx] = FVector3f::ZeroVector;
+			SourceBoneScaleVelocities[LODBoneIdx] = FVector3f::ZeroVector;
+
+			SourceBoneTranslationDecayHalfLives[LODBoneIdx] = Parameters.ExtrapolationHalfLifeMin * FVector3f::OneVector;
+			SourceBoneRotationDecayHalfLives[LODBoneIdx] = Parameters.ExtrapolationHalfLifeMin * FVector3f::OneVector;
+			SourceBoneScaleDecayHalfLives[LODBoneIdx] = Parameters.ExtrapolationHalfLifeMin * FVector3f::OneVector;
+		}
+#endif
+	}
+
+	void Apply(
+		const FTransformArraySoAView& Dest,
+		const TArrayView<FQuat4f>& BoneRotationDirections,
+		const FTransformArraySoAConstView& Source,
+		const TArrayView<const FVector3f>& SourceBoneTranslationVelocities,
+		const TArrayView<const FVector3f>& SourceBoneRotationVelocities,
+		const TArrayView<const FVector3f>& SourceBoneScaleVelocities,
+		const TArrayView<const FVector3f>& SourceBoneTranslationDecayHalfLives,
+		const TArrayView<const FVector3f>& SourceBoneRotationDecayHalfLives,
+		const TArrayView<const FVector3f>& SourceBoneScaleDecayHalfLives,
+		const float BlendDuration,
+		const float TimeSinceTransition,
+		const EAlphaBlendOption BlendMode,
+		UCurveFloat* CustomBlendCurve)
+	{
+		TRACE_CPUPROFILER_EVENT_SCOPE(AnimNext::DeadBlending::Apply)
+
+		const float Alpha = 1.0f - FAlphaBlend::AlphaToBlendOption(
+			FMath::Clamp(TimeSinceTransition / FMath::Max(BlendDuration, UE_SMALL_NUMBER), 0.0f, 1.0f),
+			BlendMode, CustomBlendCurve);
+
+		if (Alpha == 0.0f)
+		{
+			return;
+		}
+
+		const int32 LODBoneNum = FMath::Min(Dest.Num(), Source.Num());
+
+#if UE_ANIM_NEXT_DEAD_BLENDING_ISPC
+		ispc::AnimNextDeadBlendingApply(
+			(ispc::FVector*)Dest.Translations.GetData(),
+			(ispc::FVector4*)Dest.Rotations.GetData(),
+			(ispc::FVector*)Dest.Scales3D.GetData(),
+			(ispc::FVector4f*)BoneRotationDirections.GetData(),
+			(ispc::FVector*)Source.Translations.GetData(),
+			(ispc::FVector4*)Source.Rotations.GetData(),
+			(ispc::FVector*)Source.Scales3D.GetData(),
+			(ispc::FVector3f*)SourceBoneTranslationVelocities.GetData(),
+			(ispc::FVector3f*)SourceBoneRotationVelocities.GetData(),
+			(ispc::FVector3f*)SourceBoneScaleVelocities.GetData(),
+			(ispc::FVector3f*)SourceBoneTranslationDecayHalfLives.GetData(),
+			(ispc::FVector3f*)SourceBoneRotationDecayHalfLives.GetData(),
+			(ispc::FVector3f*)SourceBoneScaleDecayHalfLives.GetData(),
+			LODBoneNum,
+			Alpha,
+			TimeSinceTransition);
+#else
+		for (FBoneIndexType LODBoneIdx = 0; LODBoneIdx < LODBoneNum; LODBoneIdx++)
+		{
+			// Extrapolate and Blend Translation
+
+			const FVector ExtrapolatedTranslation = DeadBlending::Private::ExtrapolateTranslation(
+				Source.Translations[LODBoneIdx],
+				SourceBoneTranslationVelocities[LODBoneIdx],
+				SourceBoneTranslationDecayHalfLives[LODBoneIdx],
+				TimeSinceTransition);
+
+			Dest.Translations[LODBoneIdx] = FMath::Lerp(Dest.Translations[LODBoneIdx], ExtrapolatedTranslation, Alpha);
+
+			// Extrapolate and Blend Rotation
+
+			const FQuat ExtrapolatedRotation = DeadBlending::Private::ExtrapolateRotation(
+				Source.Rotations[LODBoneIdx],
+				SourceBoneRotationVelocities[LODBoneIdx],
+				SourceBoneRotationDecayHalfLives[LODBoneIdx],
+				TimeSinceTransition);
+
+			// We need to enforce that the blend of the rotation doesn't suddenly "switch sides"
+			// given that the extrapolated rotation can become quite far from the destination
+			// animation. To do this we keep track of the blend "direction" and ensure that the
+			// delta we are applying to the destination animation always remains on the same
+			// side of this rotation.
+
+			FQuat RotationDiff = ExtrapolatedRotation * Dest.Rotations[LODBoneIdx].Inverse();
+			RotationDiff.EnforceShortestArcWith((FQuat)BoneRotationDirections[LODBoneIdx]);
+
+			// Update BoneRotationDirections to match our current path
+			BoneRotationDirections[LODBoneIdx] = (FQuat4f)RotationDiff;
+
+			Dest.Rotations[LODBoneIdx] = FQuat::MakeFromRotationVector(RotationDiff.ToRotationVector() * Alpha) * Dest.Rotations[LODBoneIdx];
+
+			// Extrapolate and Blend Scale
+
+			const FVector ExtrapolatedScale = Private::ExtrapolateScale(
+				Source.Scales3D[LODBoneIdx],
+				SourceBoneScaleVelocities[LODBoneIdx],
+				SourceBoneScaleDecayHalfLives[LODBoneIdx],
+				TimeSinceTransition);
+
+			Dest.Scales3D[LODBoneIdx] = Private::VectorEerp(Dest.Scales3D[LODBoneIdx], ExtrapolatedScale, Alpha);
+		}
+#endif
+	}
+}
+
+FAnimNextDeadBlendingTransitionTask FAnimNextDeadBlendingTransitionTask::Make(
+	UE::AnimNext::FDeadBlendingState* State,
+	const UE::AnimNext::FTransformArraySoAHeap* CurrPose,
+	const UE::AnimNext::FTransformArraySoAHeap* PrevPose,
+	const float DeltaTime,
+	const UE::AnimNext::FDeadBlendTransitionTaskParameters& Parameters)
+{
+	FAnimNextDeadBlendingTransitionTask Task;
+	Task.State = State;
+	Task.CurrPose = CurrPose;
+	Task.PrevPose = PrevPose;
+	Task.DeltaTime = DeltaTime;
+	Task.Parameters = Parameters;
+	return Task;
+}
+
+FAnimNextDeadBlendingTransitionTask FAnimNextDeadBlendingTransitionTask::Make(
+	UE::AnimNext::FDeadBlendingState* State,
+	const UE::AnimNext::FTransformArraySoAHeap* CurrPose,
+	const UE::AnimNext::FDeadBlendTransitionTaskParameters& Parameters)
+{
+	FAnimNextDeadBlendingTransitionTask Task;
+	Task.State = State;
+	Task.CurrPose = CurrPose;
+	Task.Parameters = Parameters;
+	return Task;
+}
+
+void FAnimNextDeadBlendingTransitionTask::Execute(UE::AnimNext::FEvaluationVM& VM) const
+{
+	using namespace UE::AnimNext;
+
+	if (EnumHasAnyFlags(VM.GetFlags(), EEvaluationFlags::Bones))
+	{
+		if (const TUniquePtr<FKeyframeState>* Keyframe = VM.PeekValue<TUniquePtr<FKeyframeState>>(KEYFRAME_STACK_NAME, 0))
+		{
+			if (PrevPose)
+			{
+				// Transition with two previous poses
+
+				const int32 SourceLODBoneNum = FMath::Min((*Keyframe)->Pose.LocalTransformsView.Num(), FMath::Min(CurrPose->Num(), PrevPose->Num()));
+				State->SetNumUninitialized(SourceLODBoneNum);
+
+				DeadBlending::Private::Transition(
+					State->BoneRotationDirections,
+					State->SourcePose.GetView(),
+					State->SourceBoneTranslationVelocities,
+					State->SourceBoneRotationVelocities,
+					State->SourceBoneScaleVelocities,
+					State->SourceBoneTranslationDecayHalfLives,
+					State->SourceBoneRotationDecayHalfLives,
+					State->SourceBoneScaleDecayHalfLives,
+					(*Keyframe)->Pose.LocalTransformsView,
+					CurrPose->GetConstView(),
+					PrevPose->GetConstView(),
+					DeltaTime,
+					Parameters);
+			}
+			else
+			{
+				// Transition with a single previous pose, assuming zero velocity
+
+				const int32 SourceLODBoneNum = CurrPose->Num();
+				State->SetNumUninitialized(SourceLODBoneNum);
+
+				DeadBlending::Private::Transition(
+					State->BoneRotationDirections,
+					State->SourcePose.GetView(),
+					State->SourceBoneTranslationVelocities,
+					State->SourceBoneRotationVelocities,
+					State->SourceBoneScaleVelocities,
+					State->SourceBoneTranslationDecayHalfLives,
+					State->SourceBoneRotationDecayHalfLives,
+					State->SourceBoneScaleDecayHalfLives,
+					CurrPose->GetConstView(),
+					Parameters);
+			}
+		}
+	}
+}
+
+FAnimNextDeadBlendingApplyTask FAnimNextDeadBlendingApplyTask::Make(
+	UE::AnimNext::FDeadBlendingState* State,
+	const float BlendDuration,
+	const float TimeSinceTransition,
+	const EAlphaBlendOption& BlendMode,
+	const TWeakObjectPtr<UCurveFloat> CustomBlendCurve)
+{
+	FAnimNextDeadBlendingApplyTask Task;
+	Task.State = State;
+	Task.BlendDuration = BlendDuration;
+	Task.TimeSinceTransition = TimeSinceTransition;
+	Task.BlendMode = BlendMode;
+	Task.CustomBlendCurve = CustomBlendCurve;
+	return Task;
+}
+
+void FAnimNextDeadBlendingApplyTask::Execute(UE::AnimNext::FEvaluationVM& VM) const
+{
+	using namespace UE::AnimNext;
+
+	if (EnumHasAnyFlags(VM.GetFlags(), EEvaluationFlags::Bones))
+	{
+		if (const TUniquePtr<FKeyframeState>* Keyframe = VM.PeekValue<TUniquePtr<FKeyframeState>>(KEYFRAME_STACK_NAME, 0))
+		{
+			DeadBlending::Private::Apply(
+				(*Keyframe)->Pose.LocalTransforms.GetView(),
+				State->BoneRotationDirections,
+				State->SourcePose.GetConstView(),
+				State->SourceBoneTranslationVelocities,
+				State->SourceBoneRotationVelocities,
+				State->SourceBoneScaleVelocities,
+				State->SourceBoneTranslationDecayHalfLives,
+				State->SourceBoneRotationDecayHalfLives,
+				State->SourceBoneScaleDecayHalfLives,
+				BlendDuration,
+				TimeSinceTransition,
+				BlendMode,
+				CustomBlendCurve.Get());
+		}
+	}
+}
+
+#undef UE_ANIM_NEXT_DEAD_BLENDING_ISPC
\ No newline at end of file
diff --git a/Engine/Plugins/Experimental/AnimNextAnimGraph/Source/AnimNextAnimGraph/Private/EvaluationVM/Tasks/DeadBlending.ispc b/Engine/Plugins/Experimental/AnimNextAnimGraph/Source/AnimNextAnimGraph/Private/EvaluationVM/Tasks/DeadBlending.ispc
new file mode 100644
--- /dev/null
+++ b/Engine/Plugins/Experimental/AnimNextAnimGraph/Source/AnimNextAnimGraph/Private/EvaluationVM/Tasks/DeadBlending.ispc
@@ -0,0 +1,477 @@
+// Copyright Epic Games, Inc. All Rights Reserved.
+
+#include "Math/Vector.isph"
+#include "Math/Quat.isph"
+
+#define UE_LN2 (0.69314718056f)
+#define UE_SMALL_NUMBER (1.e-8f)
+#define UE_KINDA_SMALL_NUMBER (1.e-4f)
+
+inline FVector CastToVector(const FVector3f& V)
+{
+	const FVector Out = { {
+		V.V[0],
+		V.V[1],
+		V.V[2],
+	} };
+
+	return Out;
+}
+
+inline FVector3f CastToVector3f(const FVector& V)
+{
+	const FVector3f Out = { {
+		V.V[0],
+		V.V[1],
+		V.V[2],
+	} };
+    
+	return Out;
+}
+
+inline FVector4f CastToVector4f(const FVector4& V)
+{
+	const FVector4f Out = { {
+		V.V[0],
+		V.V[1],
+		V.V[2],
+		V.V[3],
+	} };
+    
+	return Out;
+}
+
+inline FVector4 CastToVector4(const FVector4f& V)
+{
+	const FVector4 Out = { {
+		V.V[0],
+		V.V[1],
+		V.V[2],
+		V.V[3],
+	} };
+    
+	return Out;
+}
+
+inline FVector3f VectorDivMax(const float S, const FVector3f& W, const uniform float Epsilon = UE_SMALL_NUMBER)
+{
+	const FVector3f Out = { {
+		S / max(W.V[0], Epsilon),
+		S / max(W.V[1], Epsilon),
+		S / max(W.V[2], Epsilon),
+	} };
+    
+	return Out;
+}
+
+inline FVector3f VectorDivMax(const FVector3f& V, const FVector3f& W, const uniform float Epsilon = UE_SMALL_NUMBER)
+{
+	const FVector3f Out = { {
+		V.V[0] / max(W.V[0], Epsilon),
+		V.V[1] / max(W.V[1], Epsilon),
+		V.V[2] / max(W.V[2], Epsilon),
+	} };
+    
+	return Out;
+}
+
+inline FVector VectorDivMax(const FVector& V, const FVector& W, const uniform double Epsilon = UE_SMALL_NUMBER)
+{
+	const FVector Out = { {
+		V.V[0] / max(W.V[0], Epsilon),
+		V.V[1] / max(W.V[1], Epsilon),
+		V.V[2] / max(W.V[2], Epsilon),
+	} };
+    
+	return Out;
+}
+
+inline FVector3f VectorInvExpApprox(const FVector3f& V)
+{
+	const uniform float A = 1.00746054f;
+	const uniform float B = 0.45053901f;
+	const uniform float C = 0.25724632f;
+
+	const FVector3f Out = { {
+		1.0f / (1.0f + A * V.V[0] + B * V.V[0] * V.V[0] + C * V.V[0] * V.V[0] * V.V[0]),
+		1.0f / (1.0f + A * V.V[1] + B * V.V[1] * V.V[1] + C * V.V[1] * V.V[1] * V.V[1]),
+		1.0f / (1.0f + A * V.V[2] + B * V.V[2] * V.V[2] + C * V.V[2] * V.V[2] * V.V[2]),
+	} };
+
+	return Out;
+}
+
+inline FVector VectorEerp(const FVector& V, const FVector& W, const uniform double Alpha, const uniform double Epsilon = UE_SMALL_NUMBER)
+{
+	const FVector Out = { {
+        pow(max(V.V[0], Epsilon), (1.0 - Alpha)) * pow(max(W.V[0], Epsilon), Alpha),
+        pow(max(V.V[1], Epsilon), (1.0 - Alpha)) * pow(max(W.V[0], Epsilon), Alpha),
+        pow(max(V.V[2], Epsilon), (1.0 - Alpha)) * pow(max(W.V[0], Epsilon), Alpha),
+	} };
+
+	return Out;
+}
+
+inline FVector VectorExp(const FVector& V)
+{
+	const FVector Out = { {
+		exp(V.V[0]),
+		exp(V.V[1]),
+		exp(V.V[2]),
+	} };
+
+	return Out;
+}
+
+inline FVector VectorLogMax(const FVector& V, const uniform double Epsilon = UE_SMALL_NUMBER)
+{
+	const FVector Out = { {
+		log(max(V.V[0], Epsilon)),
+		log(max(V.V[1], Epsilon)),
+		log(max(V.V[2], Epsilon)),
+	} };
+    
+	return Out;
+}
+
+inline FVector VectorClampToMaxSize(const FVector V, const uniform double MaxSize)
+{
+	const double L = VectorDot(V, V);
+    
+	return L > (MaxSize * MaxSize) ? V * (MaxSize * rsqrt(L)) : V;
+}
+
+inline FVector4 QuatExp(const FVector& V, const uniform double Epsilon = UE_SMALL_NUMBER)
+{
+	const double HalfAngleSq = VectorDot(V, V);
+
+	cif (HalfAngleSq < Epsilon)
+	{
+		const FVector4 O = { {
+			V.V[0],
+			V.V[1],
+			V.V[2],
+			1.0,
+		} };
+
+		return VectorNormalizeQuaternion(O);
+	}
+	else
+	{
+		const double HalfAngle = sqrt(HalfAngleSq);
+		const double C = cos(HalfAngle);
+		const double S = sin(HalfAngle) / HalfAngle;
+
+		const FVector4 O = { {
+			S * V.V[0],
+			S * V.V[1],
+			S * V.V[2],
+			C,
+		} };
+
+		return O;
+	}
+}
+
+inline FVector QuatLog(const FVector4& Q, const uniform double Epsilon = UE_SMALL_NUMBER)
+{
+	const double LengthSq = VectorDot(Q, Q);
+
+	cif (LengthSq < Epsilon)
+	{
+		const FVector O = { { Q.V[0], Q.V[1], Q.V[2] } };
+		return O;
+	}
+	else
+	{
+		const double Length = sqrt(LengthSq);
+		const double HalfAngle = acos(clamp(Q.V[3], -1.0d, 1.0d));
+
+		const FVector O = { {
+			HalfAngle * (Q.V[0] / Length),
+			HalfAngle * (Q.V[1] / Length),
+			HalfAngle * (Q.V[2] / Length),
+		} };
+
+		return O;
+	}
+}
+
+inline FVector4 QuatFromRotationVector(const FVector& V, const uniform double Epsilon = UE_SMALL_NUMBER)
+{
+	return QuatExp(0.5 * V, Epsilon);
+}
+
+inline FVector QuatToRotationVector(const FVector4& Q, const uniform double Epsilon = UE_SMALL_NUMBER)
+{
+	return 2.0 * QuatLog(Q, Epsilon);
+}
+
+inline FVector4 VectorNegate(const FVector4& Q)
+{
+	const FVector4 O = { { -Q.V[0], -Q.V[1], -Q.V[2], -Q.V[3] } };
+	return O;
+}
+
+inline FVector4 VectorQuaternionEnforceShortestArcWith(const FVector4 Q, const FVector4 P)
+{
+	return VectorDot(Q, P) > 0.0 ? Q : VectorNegate(Q);
+}
+
+inline FVector4 QuatAbs(const FVector4 Q)
+{
+	return Q.V[3] > 0.0 ? Q : VectorNegate(Q);
+}
+
+inline FVector ExtrapolateTranslation(
+	const FVector& Translation,
+	const FVector3f& Velocity,
+	const FVector3f& DecayHalflife,
+	const uniform float Time,
+	const uniform float Epsilon = UE_SMALL_NUMBER)
+{
+	const FVector3f C = VectorDivMax(UE_LN2, DecayHalflife, Epsilon);
+	return Translation + CastToVector(VectorDivMax(Velocity, C, Epsilon) * (FloatOneVector - VectorInvExpApprox(C * Time)));
+}
+
+inline FVector4 ExtrapolateRotation(
+	const FVector4& Rotation,
+	const FVector3f& Velocity,
+	const FVector3f& DecayHalflife,
+	const uniform float Time,
+	const uniform float Epsilon = UE_SMALL_NUMBER)
+{
+	const FVector3f C = VectorDivMax(UE_LN2, DecayHalflife, Epsilon);
+	return VectorQuaternionMultiply2(QuatFromRotationVector(CastToVector(VectorDivMax(Velocity, C, Epsilon) * (FloatOneVector - VectorInvExpApprox(C * Time)))), Rotation);
+}
+
+inline FVector ExtrapolateScale(
+	const FVector& Scale,
+	const FVector3f& Velocity,
+	const FVector3f& DecayHalflife,
+	const uniform float Time,
+	const uniform float Epsilon = UE_SMALL_NUMBER)
+{
+	const FVector3f C = VectorDivMax(UE_LN2, DecayHalflife, Epsilon);
+	return VectorExp(CastToVector(VectorDivMax(Velocity, C, Epsilon) * (FloatOneVector - VectorInvExpApprox(C * Time)))) * Scale;
+}
+
+inline float ClipMagnitudeToGreaterThanEpsilon(const float X, const uniform float Epsilon = UE_KINDA_SMALL_NUMBER)
+{
+	return
+		X >= 0.0f && X < +Epsilon ? +Epsilon :
+		X <  0.0f && X > -Epsilon ? -Epsilon : X;
+}
+
+inline float ComputeDecayHalfLifeFromDiffAndVelocity(
+	const float SrcDstDiff,
+	const float SrcVelocity,
+	const uniform float HalfLife,
+	const uniform float HalfLifeMin,
+	const uniform float HalfLifeMax,
+	const uniform float Epsilon = UE_KINDA_SMALL_NUMBER)
+{
+	return clamp(HalfLife * (SrcDstDiff / ClipMagnitudeToGreaterThanEpsilon(SrcVelocity, Epsilon)), HalfLifeMin, HalfLifeMax);
+}
+
+inline FVector3f ComputeDecayHalfLifeFromDiffAndVelocity(
+	const FVector SrcDstDiff,
+	const FVector3f SrcVelocity,
+	const uniform float HalfLife,
+	const uniform float HalfLifeMin,
+	const uniform float HalfLifeMax,
+	const uniform float Epsilon = UE_KINDA_SMALL_NUMBER)
+{
+	const FVector3f O = { {
+		ComputeDecayHalfLifeFromDiffAndVelocity((float)SrcDstDiff.V[0], SrcVelocity.V[0], HalfLife, HalfLifeMin, HalfLifeMax, Epsilon),
+		ComputeDecayHalfLifeFromDiffAndVelocity((float)SrcDstDiff.V[1], SrcVelocity.V[1], HalfLife, HalfLifeMin, HalfLifeMax, Epsilon),
+		ComputeDecayHalfLifeFromDiffAndVelocity((float)SrcDstDiff.V[2], SrcVelocity.V[2], HalfLife, HalfLifeMin, HalfLifeMax, Epsilon),
+	} };
+
+	return O;
+}
+
+export void AnimNextDeadBlendingTransition(
+	uniform FVector4f BoneRotationDirections[],
+	uniform FVector SourceBoneTranslations[],
+	uniform FVector4 SourceBoneRotations[],
+	uniform FVector SourceBoneScales3D[],
+	uniform FVector3f SourceBoneTranslationVelocities[],
+	uniform FVector3f SourceBoneRotationVelocities[],
+	uniform FVector3f SourceBoneScaleVelocities[],
+	uniform FVector3f SourceBoneTranslationDecayHalfLives[],
+	uniform FVector3f SourceBoneRotationDecayHalfLives[],
+	uniform FVector3f SourceBoneScaleDecayHalfLives[],
+	const uniform FVector DestBoneTranslations[],
+	const uniform FVector4 DestBoneRotations[],
+	const uniform FVector DestBoneScales3D[],
+	const uniform FVector CurrBoneTranslations[],
+	const uniform FVector4 CurrBoneRotations[],
+	const uniform FVector CurrBoneScales3D[],
+	const uniform FVector PrevBoneTranslations[],
+	const uniform FVector4 PrevBoneRotations[],
+	const uniform FVector PrevBoneScales3D[],
+	const uniform int LODBoneNum,
+	const uniform float DeltaTime,
+	const uniform float ExtrapolationHalfLife,
+	const uniform float ExtrapolationHalfLifeMin,
+	const uniform float ExtrapolationHalfLifeMax,
+	const uniform float MaximumTranslationVelocity,
+	const uniform float MaximumRotationVelocity,
+	const uniform float MaximumScaleVelocity)
+{
+	const FVector4f IdentityRotation = Float0001;
+
+	foreach(i = 0 ... LODBoneNum)
+	{
+		const uniform int Index = extract(i, 0);
+
+		VectorStore(&BoneRotationDirections[Index], IdentityRotation);
+
+		const FVector SrcTranslationCurr = VectorLoad(&CurrBoneTranslations[Index]);
+		const FVector4 SrcRotationCurr = VectorLoad(&CurrBoneRotations[Index]);
+		const FVector SrcScaleCurr = VectorLoad(&CurrBoneScales3D[Index]);
+
+		VectorStore(&SourceBoneTranslations[Index], SrcTranslationCurr);
+		VectorStore(&SourceBoneRotations[Index], SrcRotationCurr);
+		VectorStore(&SourceBoneScales3D[Index], SrcScaleCurr);
+
+		const FVector SrcTranslationPrev = VectorLoad(&PrevBoneTranslations[Index]);
+		const FVector4 SrcRotationPrev = VectorLoad(&PrevBoneRotations[Index]);
+		const FVector SrcScalePrev = VectorLoad(&PrevBoneScales3D[Index]);
+
+		const FVector TranslationDiff = SrcTranslationCurr - SrcTranslationPrev;
+		const FVector4 RotationDiff = QuatAbs(VectorQuaternionMultiply2(SrcRotationCurr, QuatInverse(SrcRotationPrev)));
+		const FVector ScaleDiff = VectorDivMax(SrcScaleCurr, SrcScalePrev);
+
+		VectorStore(&SourceBoneTranslationVelocities[Index], CastToVector3f(VectorClampToMaxSize(TranslationDiff / DeltaTime, MaximumTranslationVelocity)));
+		VectorStore(&SourceBoneRotationVelocities[Index], CastToVector3f(VectorClampToMaxSize(QuatToRotationVector(RotationDiff) / DeltaTime, MaximumRotationVelocity)));
+		VectorStore(&SourceBoneScaleVelocities[Index], CastToVector3f(VectorClampToMaxSize(VectorLogMax(ScaleDiff) / DeltaTime, MaximumScaleVelocity)));
+
+		const FVector DstTranslation = VectorLoad(&DestBoneTranslations[Index]);
+		const FVector4 DstRotation = VectorLoad(&DestBoneRotations[Index]);
+		const FVector DstScale = VectorLoad(&DestBoneScales3D[Index]);
+
+		const FVector TranslationSrcDstDiff = DstTranslation - SrcTranslationCurr;
+		const FVector4 RotationSrcDstDiff = QuatAbs(VectorQuaternionMultiply2(DstRotation, QuatInverse(SrcRotationCurr)));
+		const FVector ScaleSrcDstDiff = VectorDivMax(DstScale, SrcScaleCurr);
+
+		VectorStore(&SourceBoneTranslationDecayHalfLives[Index], ComputeDecayHalfLifeFromDiffAndVelocity(
+			TranslationSrcDstDiff,
+			VectorLoad(&SourceBoneTranslationVelocities[Index]),
+			ExtrapolationHalfLife,
+			ExtrapolationHalfLifeMin,
+			ExtrapolationHalfLifeMax));
+
+		VectorStore(&SourceBoneRotationDecayHalfLives[Index], ComputeDecayHalfLifeFromDiffAndVelocity(
+			QuatToRotationVector(RotationSrcDstDiff),
+			VectorLoad(&SourceBoneRotationVelocities[Index]),
+			ExtrapolationHalfLife,
+			ExtrapolationHalfLifeMin,
+			ExtrapolationHalfLifeMax));
+
+		VectorStore(&SourceBoneScaleDecayHalfLives[Index], ComputeDecayHalfLifeFromDiffAndVelocity(
+			ScaleSrcDstDiff,
+			VectorLoad(&SourceBoneScaleVelocities[Index]),
+			ExtrapolationHalfLife,
+			ExtrapolationHalfLifeMin,
+			ExtrapolationHalfLifeMax));
+	}
+}
+
+
+export void AnimNextDeadBlendingTransitionStatic(
+	uniform FVector4f BoneRotationDirections[],
+	uniform FVector SourceBoneTranslations[],
+	uniform FVector4 SourceBoneRotations[],
+	uniform FVector SourceBoneScales3D[],
+	uniform FVector3f SourceBoneTranslationVelocities[],
+	uniform FVector3f SourceBoneRotationVelocities[],
+	uniform FVector3f SourceBoneScaleVelocities[],
+	uniform FVector3f SourceBoneTranslationDecayHalfLives[],
+	uniform FVector3f SourceBoneRotationDecayHalfLives[],
+	uniform FVector3f SourceBoneScaleDecayHalfLives[],
+	const uniform FVector CurrBoneTranslations[],
+	const uniform FVector4 CurrBoneRotations[],
+	const uniform FVector CurrBoneScales[],
+	const uniform int LODBoneNum,
+	const uniform float ExtrapolationHalfLifeMin)
+{
+	const FVector4f IdentityRotation = Float0001;
+	const FVector3f VelocityZeroVector = FloatZeroVector;
+	const FVector3f HalfLifeVector = { { ExtrapolationHalfLifeMin, ExtrapolationHalfLifeMin, ExtrapolationHalfLifeMin } };
+
+	foreach(i = 0 ... LODBoneNum)
+	{
+		const uniform int Index = extract(i, 0);
+
+		VectorStore(&BoneRotationDirections[Index], IdentityRotation);
+		VectorStore(&SourceBoneTranslations[Index], VectorLoad(&CurrBoneTranslations[Index]));
+		VectorStore(&SourceBoneRotations[Index], VectorLoad(&CurrBoneRotations[Index]));
+		VectorStore(&SourceBoneScales3D[Index], VectorLoad(&CurrBoneScales[Index]));
+
+		VectorStore(&SourceBoneTranslationVelocities[Index], VelocityZeroVector);
+		VectorStore(&SourceBoneRotationVelocities[Index], VelocityZeroVector);
+		VectorStore(&SourceBoneScaleVelocities[Index], VelocityZeroVector);
+		VectorStore(&SourceBoneTranslationDecayHalfLives[Index], HalfLifeVector);
+		VectorStore(&SourceBoneRotationDecayHalfLives[Index], HalfLifeVector);
+		VectorStore(&SourceBoneScaleDecayHalfLives[Index], HalfLifeVector);
+	}
+}
+
+export void AnimNextDeadBlendingApply(
+	uniform FVector DestTranslations[],
+	uniform FVector4 DestRotations[],
+	uniform FVector DestScales3D[],
+	uniform FVector4f BoneRotationDirections[],
+	const uniform FVector SourceTranslations[],
+	const uniform FVector4 SourceRotations[],
+	const uniform FVector SourceScales3D[],
+	const uniform FVector3f SourceBoneTranslationVelocities[],
+	const uniform FVector3f SourceBoneRotationVelocities[],
+	const uniform FVector3f SourceBoneScaleVelocities[],
+	const uniform FVector3f SourceBoneTranslationDecayHalfLives[],
+	const uniform FVector3f SourceBoneRotationDecayHalfLives[],
+	const uniform FVector3f SourceBoneScaleDecayHalfLives[],
+	const uniform int LODBoneNum,
+	const uniform float Alpha,
+	const uniform float TimeSinceTransition)
+{
+	foreach (i = 0 ... LODBoneNum)
+	{
+		const uniform int Index = extract(i, 0);
+
+		const FVector ExtrapolatedTranslation = ExtrapolateTranslation(
+			VectorLoad(&SourceTranslations[Index]),
+			VectorLoad(&SourceBoneTranslationVelocities[Index]),
+			VectorLoad(&SourceBoneTranslationDecayHalfLives[Index]),
+			TimeSinceTransition);
+
+		const FVector4 ExtrapolatedRotation = ExtrapolateRotation(
+			VectorLoad(&SourceRotations[Index]),
+			VectorLoad(&SourceBoneRotationVelocities[Index]),
+			VectorLoad(&SourceBoneRotationDecayHalfLives[Index]),
+			TimeSinceTransition);
+
+		const FVector ExtrapolatedScale = ExtrapolateScale(
+			VectorLoad(&SourceScales3D[Index]),
+			VectorLoad(&SourceBoneScaleVelocities[Index]),
+			VectorLoad(&SourceBoneScaleDecayHalfLives[Index]),
+			TimeSinceTransition);
+
+		const FVector4 RotationDiff = VectorQuaternionEnforceShortestArcWith(
+			VectorQuaternionMultiply2(ExtrapolatedRotation, QuatInverse(VectorLoad(&DestRotations[Index]))),
+			CastToVector4(VectorLoad(&BoneRotationDirections[Index])));
+
+		VectorStore(&BoneRotationDirections[Index], CastToVector4f(RotationDiff));
+
+		VectorStore(&DestTranslations[Index], VectorLerp(VectorLoad(&DestTranslations[Index]), ExtrapolatedTranslation, Alpha));
+
+		VectorStore(&DestRotations[Index],
+			VectorQuaternionMultiply2(
+				QuatFromRotationVector(QuatToRotationVector(RotationDiff) * Alpha),
+				VectorLoad(&DestRotations[Index])));
+
+		VectorStore(&DestScales3D[Index], VectorEerp(VectorLoad(&DestScales3D[Index]), ExtrapolatedScale, Alpha));
+	}
+}
\ No newline at end of file
diff --git a/Engine/Plugins/Experimental/AnimNextAnimGraph/Source/AnimNextAnimGraph/Private/TraitInterfaces/IDiscreteBlend.cpp b/Engine/Plugins/Experimental/AnimNextAnimGraph/Source/AnimNextAnimGraph/Private/TraitInterfaces/IDiscreteBlend.cpp
--- a/Engine/Plugins/Experimental/AnimNextAnimGraph/Source/AnimNextAnimGraph/Private/TraitInterfaces/IDiscreteBlend.cpp
+++ b/Engine/Plugins/Experimental/AnimNextAnimGraph/Source/AnimNextAnimGraph/Private/TraitInterfaces/IDiscreteBlend.cpp
@@ -39,7 +39,7 @@ namespace UE::AnimNext
 		return INDEX_NONE;
 	}
 
-	void IDiscreteBlend::OnBlendTransition(const FExecutionContext& Context, const TTraitBinding<IDiscreteBlend>& Binding, int32 OldChildIndex, int32 NewChildIndex) const
+	void IDiscreteBlend::OnBlendTransition(FExecutionContext& Context, const TTraitBinding<IDiscreteBlend>& Binding, int32 OldChildIndex, int32 NewChildIndex) const
 	{
 		TTraitBinding<IDiscreteBlend> SuperBinding;
 		if (Binding.GetStackInterfaceSuper(SuperBinding))
@@ -48,7 +48,7 @@ namespace UE::AnimNext
 		}
 	}
 
-	void IDiscreteBlend::OnBlendInitiated(const FExecutionContext& Context, const TTraitBinding<IDiscreteBlend>& Binding, int32 ChildIndex) const
+	void IDiscreteBlend::OnBlendInitiated(FExecutionContext& Context, const TTraitBinding<IDiscreteBlend>& Binding, int32 ChildIndex) const
 	{
 		TTraitBinding<IDiscreteBlend> SuperBinding;
 		if (Binding.GetStackInterfaceSuper(SuperBinding))
@@ -57,7 +57,7 @@ namespace UE::AnimNext
 		}
 	}
 
-	void IDiscreteBlend::OnBlendTerminated(const FExecutionContext& Context, const TTraitBinding<IDiscreteBlend>& Binding, int32 ChildIndex) const
+	void IDiscreteBlend::OnBlendTerminated(FExecutionContext& Context, const TTraitBinding<IDiscreteBlend>& Binding, int32 ChildIndex) const
 	{
 		TTraitBinding<IDiscreteBlend> SuperBinding;
 		if (Binding.GetStackInterfaceSuper(SuperBinding))
diff --git a/Engine/Plugins/Experimental/AnimNextAnimGraph/Source/AnimNextAnimGraph/Private/Traits/BlendByBool.cpp b/Engine/Plugins/Experimental/AnimNextAnimGraph/Source/AnimNextAnimGraph/Private/Traits/BlendByBool.cpp
--- a/Engine/Plugins/Experimental/AnimNextAnimGraph/Source/AnimNextAnimGraph/Private/Traits/BlendByBool.cpp
+++ b/Engine/Plugins/Experimental/AnimNextAnimGraph/Source/AnimNextAnimGraph/Private/Traits/BlendByBool.cpp
@@ -130,7 +130,7 @@ namespace UE::AnimNext
 		return bCondition ? TRUE_CHILD_INDEX : FALSE_CHILD_INDEX;
 	}
 
-	void FBlendByBoolTrait::OnBlendTransition(const FExecutionContext& Context, const TTraitBinding<IDiscreteBlend>& Binding, int32 OldChildIndex, int32 NewChildIndex) const
+	void FBlendByBoolTrait::OnBlendTransition(FExecutionContext& Context, const TTraitBinding<IDiscreteBlend>& Binding, int32 OldChildIndex, int32 NewChildIndex) const
 	{
 		TTraitBinding<IDiscreteBlend> DiscreteBlendTrait;
 		Binding.GetStackInterface(DiscreteBlendTrait);
@@ -142,7 +142,7 @@ namespace UE::AnimNext
 		DiscreteBlendTrait.OnBlendTerminated(Context, OldChildIndex);
 	}
 
-	void FBlendByBoolTrait::OnBlendInitiated(const FExecutionContext& Context, const TTraitBinding<IDiscreteBlend>& Binding, int32 ChildIndex) const
+	void FBlendByBoolTrait::OnBlendInitiated(FExecutionContext& Context, const TTraitBinding<IDiscreteBlend>& Binding, int32 ChildIndex) const
 	{
 		const FSharedData* SharedData = Binding.GetSharedData<FSharedData>();
 		FInstanceData* InstanceData = Binding.GetInstanceData<FInstanceData>();
@@ -164,7 +164,7 @@ namespace UE::AnimNext
 		}
 	}
 
-	void FBlendByBoolTrait::OnBlendTerminated(const FExecutionContext& Context, const TTraitBinding<IDiscreteBlend>& Binding, int32 ChildIndex) const
+	void FBlendByBoolTrait::OnBlendTerminated(FExecutionContext& Context, const TTraitBinding<IDiscreteBlend>& Binding, int32 ChildIndex) const
 	{
 		FInstanceData* InstanceData = Binding.GetInstanceData<FInstanceData>();
 
diff --git a/Engine/Plugins/Experimental/AnimNextAnimGraph/Source/AnimNextAnimGraph/Private/Traits/BlendByBool.h b/Engine/Plugins/Experimental/AnimNextAnimGraph/Source/AnimNextAnimGraph/Private/Traits/BlendByBool.h
--- a/Engine/Plugins/Experimental/AnimNextAnimGraph/Source/AnimNextAnimGraph/Private/Traits/BlendByBool.h
+++ b/Engine/Plugins/Experimental/AnimNextAnimGraph/Source/AnimNextAnimGraph/Private/Traits/BlendByBool.h
@@ -72,8 +72,8 @@ namespace UE::AnimNext
 		// IDiscreteBlend impl
 		virtual float GetBlendWeight(const FExecutionContext& Context, const TTraitBinding<IDiscreteBlend>& Binding, int32 ChildIndex) const override;
 		virtual int32 GetBlendDestinationChildIndex(const FExecutionContext& Context, const TTraitBinding<IDiscreteBlend>& Binding) const override;
-		virtual void OnBlendTransition(const FExecutionContext& Context, const TTraitBinding<IDiscreteBlend>& Binding, int32 OldChildIndex, int32 NewChildIndex) const override;
-		virtual void OnBlendInitiated(const FExecutionContext& Context, const TTraitBinding<IDiscreteBlend>& Binding, int32 ChildIndex) const override;
-		virtual void OnBlendTerminated(const FExecutionContext& Context, const TTraitBinding<IDiscreteBlend>& Binding, int32 ChildIndex) const override;
+		virtual void OnBlendTransition(FExecutionContext& Context, const TTraitBinding<IDiscreteBlend>& Binding, int32 OldChildIndex, int32 NewChildIndex) const override;
+		virtual void OnBlendInitiated(FExecutionContext& Context, const TTraitBinding<IDiscreteBlend>& Binding, int32 ChildIndex) const override;
+		virtual void OnBlendTerminated(FExecutionContext& Context, const TTraitBinding<IDiscreteBlend>& Binding, int32 ChildIndex) const override;
 	};
 }
diff --git a/Engine/Plugins/Experimental/AnimNextAnimGraph/Source/AnimNextAnimGraph/Private/Traits/BlendInertializer.cpp b/Engine/Plugins/Experimental/AnimNextAnimGraph/Source/AnimNextAnimGraph/Private/Traits/BlendInertializer.cpp
--- a/Engine/Plugins/Experimental/AnimNextAnimGraph/Source/AnimNextAnimGraph/Private/Traits/BlendInertializer.cpp
+++ b/Engine/Plugins/Experimental/AnimNextAnimGraph/Source/AnimNextAnimGraph/Private/Traits/BlendInertializer.cpp
@@ -4,6 +4,9 @@
 
 #include "TraitCore/ExecutionContext.h"
 
+#include "Traits/Inertialization.h"
+
+
 namespace UE::AnimNext
 {
 	AUTO_REGISTER_ANIM_TRAIT(FBlendInertializerTrait)
@@ -16,31 +19,17 @@ namespace UE::AnimNext
 	GENERATE_ANIM_TRAIT_IMPLEMENTATION(FBlendInertializerTrait, TRAIT_INTERFACE_ENUMERATOR, NULL_ANIM_TRAIT_EVENT_ENUMERATOR)
 	#undef TRAIT_INTERFACE_ENUMERATOR
 
-	void FBlendInertializerTrait::OnBlendTransition(const FExecutionContext& Context, const TTraitBinding<IDiscreteBlend>& Binding, int32 OldChildIndex, int32 NewChildIndex) const
+	void FBlendInertializerTrait::OnBlendTransition(FExecutionContext& Context, const TTraitBinding<IDiscreteBlend>& Binding, int32 OldChildIndex, int32 NewChildIndex) const
 	{
 		// Trigger the new transition
 		IDiscreteBlend::OnBlendTransition(Context, Binding, OldChildIndex, NewChildIndex);
 
-		// TODO: Implement the inertialization request API
-#if 0
-		UE::Anim::IInertializationRequester* InertializationRequester = Context.GetMessage<UE::Anim::IInertializationRequester>();
-		if (InertializationRequester)
-		{
-			FInertializationRequest Request;
-			Request.Duration = CurrentBlendTimes[ChildIndex];		// TODO: Get from ISmoothBlend interface using super
-			Request.BlendProfile = CurrentBlendProfile;				// TODO: Get from shared data
-			Request.bUseBlendMode = true;
-			Request.BlendMode = GetBlendType();						// TODO: Add to ISmoothBlend interface
-			Request.CustomBlendCurve = GetCustomBlendCurve();		// TODO: Add to ISmoothBlend interface
-
-			InertializationRequester->RequestInertialization(Request);
-			bRequestedInertializationOnActiveChildIndexChange = true;
-		}
-		else
-		{
-			FAnimNode_Inertialization::LogRequestError(Context, BlendPose[ChildIndex]);
-		}
-#endif
+		const FSharedData* SharedData = Binding.GetSharedData<FSharedData>();
+
+		// Make Request Event
+		TSharedPtr<FAnimNextInertializationRequestEvent> Event = MakeTraitEvent<FAnimNextInertializationRequestEvent>();
+		Event->Request.BlendTime = SharedData->BlendTime;
+		Context.RaiseOutputTraitEvent(Event);
 	}
 
 	float FBlendInertializerTrait::GetBlendTime(const FExecutionContext& Context, const TTraitBinding<ISmoothBlend>& Binding, int32 ChildIndex) const
diff --git a/Engine/Plugins/Experimental/AnimNextAnimGraph/Source/AnimNextAnimGraph/Private/Traits/BlendInertializer.h b/Engine/Plugins/Experimental/AnimNextAnimGraph/Source/AnimNextAnimGraph/Private/Traits/BlendInertializer.h
--- a/Engine/Plugins/Experimental/AnimNextAnimGraph/Source/AnimNextAnimGraph/Private/Traits/BlendInertializer.h
+++ b/Engine/Plugins/Experimental/AnimNextAnimGraph/Source/AnimNextAnimGraph/Private/Traits/BlendInertializer.h
@@ -17,10 +17,9 @@ struct FAnimNextBlendInertializerTraitSharedData : public FAnimNextTraitSharedDa
 {
 	GENERATED_BODY()
 
-	/** Blend profile that configures how fast to blend each bone. */
-	// TODO: Can't show list of blend profiles, we need to find a skeleton to perform the lookup with
-	UPROPERTY(EditAnywhere, Category = "Default", meta = (Inline, UseAsBlendProfile = true))
-	TObjectPtr<UBlendProfile> BlendProfile = nullptr;
+	// Inertialization Blend Time
+	UPROPERTY(EditAnywhere, Category = "Default", meta = (Inline))
+	float BlendTime = 0.2f;
 };
 
 namespace UE::AnimNext
@@ -38,7 +37,7 @@ namespace UE::AnimNext
 		using FInstanceData = FTrait::FInstanceData;
 
 		// IDiscreteBlend impl
-		virtual void OnBlendTransition(const FExecutionContext& Context, const TTraitBinding<IDiscreteBlend>& Binding, int32 OldChildIndex, int32 NewChildIndex) const override;
+		virtual void OnBlendTransition(FExecutionContext& Context, const TTraitBinding<IDiscreteBlend>& Binding, int32 OldChildIndex, int32 NewChildIndex) const override;
 
 		// ISmoothBlend impl
 		virtual float GetBlendTime(const FExecutionContext& Context, const TTraitBinding<ISmoothBlend>& Binding, int32 ChildIndex) const override;
diff --git a/Engine/Plugins/Experimental/AnimNextAnimGraph/Source/AnimNextAnimGraph/Private/Traits/BlendSmoother.cpp b/Engine/Plugins/Experimental/AnimNextAnimGraph/Source/AnimNextAnimGraph/Private/Traits/BlendSmoother.cpp
--- a/Engine/Plugins/Experimental/AnimNextAnimGraph/Source/AnimNextAnimGraph/Private/Traits/BlendSmoother.cpp
+++ b/Engine/Plugins/Experimental/AnimNextAnimGraph/Source/AnimNextAnimGraph/Private/Traits/BlendSmoother.cpp
@@ -170,7 +170,7 @@ namespace UE::AnimNext
 		return InstanceData->PerChildBlendData.IsValidIndex(ChildIndex) ? &InstanceData->PerChildBlendData[ChildIndex].Blend : nullptr;
 	}
 
-	void FBlendSmootherTrait::OnBlendTransition(const FExecutionContext& Context, const TTraitBinding<IDiscreteBlend>& Binding, int32 OldChildIndex, int32 NewChildIndex) const
+	void FBlendSmootherTrait::OnBlendTransition(FExecutionContext& Context, const TTraitBinding<IDiscreteBlend>& Binding, int32 OldChildIndex, int32 NewChildIndex) const
 	{
 		const FSharedData* SharedData = Binding.GetSharedData<FSharedData>();
 		FInstanceData* InstanceData = Binding.GetInstanceData<FInstanceData>();
diff --git a/Engine/Plugins/Experimental/AnimNextAnimGraph/Source/AnimNextAnimGraph/Private/Traits/BlendSmoother.h b/Engine/Plugins/Experimental/AnimNextAnimGraph/Source/AnimNextAnimGraph/Private/Traits/BlendSmoother.h
--- a/Engine/Plugins/Experimental/AnimNextAnimGraph/Source/AnimNextAnimGraph/Private/Traits/BlendSmoother.h
+++ b/Engine/Plugins/Experimental/AnimNextAnimGraph/Source/AnimNextAnimGraph/Private/Traits/BlendSmoother.h
@@ -74,7 +74,7 @@ namespace UE::AnimNext
 		// IDiscreteBlend impl
 		virtual float GetBlendWeight(const FExecutionContext& Context, const TTraitBinding<IDiscreteBlend>& Binding, int32 ChildIndex) const override;
 		virtual const FAlphaBlend* GetBlendState(const FExecutionContext& Context, const TTraitBinding<IDiscreteBlend>& Binding, int32 ChildIndex) const override;
-		virtual void OnBlendTransition(const FExecutionContext& Context, const TTraitBinding<IDiscreteBlend>& Binding, int32 OldChildIndex, int32 NewChildIndex) const override;
+		virtual void OnBlendTransition(FExecutionContext& Context, const TTraitBinding<IDiscreteBlend>& Binding, int32 OldChildIndex, int32 NewChildIndex) const override;
 
 		// ISmoothBlend impl
 		virtual float GetBlendTime(const FExecutionContext& Context, const TTraitBinding<ISmoothBlend>& Binding, int32 ChildIndex) const override;
diff --git a/Engine/Plugins/Experimental/AnimNextAnimGraph/Source/AnimNextAnimGraph/Private/Traits/BlendSmootherPerBone.cpp b/Engine/Plugins/Experimental/AnimNextAnimGraph/Source/AnimNextAnimGraph/Private/Traits/BlendSmootherPerBone.cpp
--- a/Engine/Plugins/Experimental/AnimNextAnimGraph/Source/AnimNextAnimGraph/Private/Traits/BlendSmootherPerBone.cpp
+++ b/Engine/Plugins/Experimental/AnimNextAnimGraph/Source/AnimNextAnimGraph/Private/Traits/BlendSmootherPerBone.cpp
@@ -110,7 +110,7 @@ namespace UE::AnimNext
 		FBlendSampleData::NormalizeDataWeight(InstanceData->PerBoneSampleData);
 	}
 
-	void FBlendSmootherPerBoneTrait::OnBlendTransition(const FExecutionContext& Context, const TTraitBinding<IDiscreteBlend>& Binding, int32 OldChildIndex, int32 NewChildIndex) const
+	void FBlendSmootherPerBoneTrait::OnBlendTransition(FExecutionContext& Context, const TTraitBinding<IDiscreteBlend>& Binding, int32 OldChildIndex, int32 NewChildIndex) const
 	{
 		const FSharedData* SharedData = Binding.GetSharedData<FSharedData>();
 		FInstanceData* InstanceData = Binding.GetInstanceData<FInstanceData>();
diff --git a/Engine/Plugins/Experimental/AnimNextAnimGraph/Source/AnimNextAnimGraph/Private/Traits/BlendSmootherPerBone.h b/Engine/Plugins/Experimental/AnimNextAnimGraph/Source/AnimNextAnimGraph/Private/Traits/BlendSmootherPerBone.h
--- a/Engine/Plugins/Experimental/AnimNextAnimGraph/Source/AnimNextAnimGraph/Private/Traits/BlendSmootherPerBone.h
+++ b/Engine/Plugins/Experimental/AnimNextAnimGraph/Source/AnimNextAnimGraph/Private/Traits/BlendSmootherPerBone.h
@@ -62,7 +62,7 @@ namespace UE::AnimNext
 		virtual void PreUpdate(FUpdateTraversalContext& Context, const TTraitBinding<IUpdate>& Binding, const FTraitUpdateState& TraitState) const override;
 
 		// IDiscreteBlend impl
-		virtual void OnBlendTransition(const FExecutionContext& Context, const TTraitBinding<IDiscreteBlend>& Binding, int32 OldChildIndex, int32 NewChildIndex) const override;
+		virtual void OnBlendTransition(FExecutionContext& Context, const TTraitBinding<IDiscreteBlend>& Binding, int32 OldChildIndex, int32 NewChildIndex) const override;
 
 		// Internal impl
 		static void InitializeInstanceData(const FExecutionContext& Context, const FTraitBinding& Binding, const FSharedData* SharedData, FInstanceData* InstanceData);
diff --git a/Engine/Plugins/Experimental/AnimNextAnimGraph/Source/AnimNextAnimGraph/Private/Traits/InertializationDeadBlend.cpp b/Engine/Plugins/Experimental/AnimNextAnimGraph/Source/AnimNextAnimGraph/Private/Traits/InertializationDeadBlend.cpp
new file mode 100644
--- /dev/null
+++ b/Engine/Plugins/Experimental/AnimNextAnimGraph/Source/AnimNextAnimGraph/Private/Traits/InertializationDeadBlend.cpp
@@ -0,0 +1,187 @@
+// Copyright Epic Games, Inc. All Rights Reserved.
+
+#include "Traits/InertializationDeadBlend.h"
+
+#include "Animation/AnimTypes.h"
+#include "TraitCore/ExecutionContext.h"
+#include "TraitInterfaces/IHierarchy.h"
+#include "EvaluationVM/Tasks/DeadBlending.h"
+#include "EvaluationVM/Tasks/StoreKeyframe.h"
+
+#include "Curves/CurveFloat.h"
+
+namespace UE::AnimNext
+{
+	AUTO_REGISTER_ANIM_TRAIT(FInertializationDeadBlendTrait)
+
+		// Trait implementation boilerplate
+#define TRAIT_INTERFACE_ENUMERATOR(GeneratorMacro) \
+		GeneratorMacro(IUpdate) \
+		GeneratorMacro(IEvaluate) \
+
+		// Trait implementation boilerplate
+#define TRAIT_EVENT_ENUMERATOR(GeneratorMacro) \
+		GeneratorMacro(FInertializationDeadBlendTrait::OnInertializationRequestEvent) \
+
+		GENERATE_ANIM_TRAIT_IMPLEMENTATION(FInertializationDeadBlendTrait, TRAIT_INTERFACE_ENUMERATOR, TRAIT_EVENT_ENUMERATOR)
+#undef TRAIT_INTERFACE_ENUMERATOR
+#undef TRAIT_EVENT_ENUMERATOR
+
+	ETraitStackPropagation FInertializationDeadBlendTrait::OnInertializationRequestEvent(const FExecutionContext& Context, FTraitBinding& Binding, FAnimNextInertializationRequestEvent& Event) const
+	{
+		if (!Event.IsHandled())
+		{
+			FInstanceData* InstanceData = Binding.GetInstanceData<FInstanceData>();
+
+            // If there is no request pending, or this request has a shorter blend
+            // time than the previous request then make this the active request.
+			if (!InstanceData->bRequestPending ||
+				Event.Request.BlendTime < InstanceData->PendingRequest.BlendTime)
+			{
+				InstanceData->bRequestPending = true;
+				InstanceData->PendingRequest = Event.Request;
+			}
+
+            // Always mark inertialization requests as handled
+			Event.MarkHandled();
+		}
+
+		return ETraitStackPropagation::Continue;
+	}
+
+	void FInertializationDeadBlendTrait::PostUpdate(FUpdateTraversalContext& Context, const TTraitBinding<IUpdate>& Binding, const FTraitUpdateState& TraitState) const
+	{
+		IUpdate::PostUpdate(Context, Binding, TraitState);
+
+		FInstanceData* InstanceData = Binding.GetInstanceData<FInstanceData>();
+        
+        // Accumulate time since last evaluation
+		InstanceData->EvaluateDeltaTime += TraitState.GetDeltaTime();
+	}
+
+	void FInertializationDeadBlendTrait::PostEvaluate(FEvaluateTraversalContext& Context, const TTraitBinding<IEvaluate>& Binding) const
+	{
+		IEvaluate::PostEvaluate(Context, Binding);
+
+		const FSharedData* SharedData = Binding.GetSharedData<FSharedData>();
+		FInstanceData* InstanceData = Binding.GetInstanceData<FInstanceData>();
+        
+        if (InstanceData->bRequestPending)
+        {
+            if (InstanceData->CurrPose.IsEmpty() && InstanceData->PrevPose.IsEmpty())
+            {
+                // We've not recorded any poses yet so effectively we have just
+                // become active and either there should not be a discontinuity
+                // or we cannot handle it so we effectively ignore the request
+
+                InstanceData->bRequestPending = false;
+                InstanceData->PendingRequest = UE::AnimNext::FInertializationRequest();
+            }
+            else if (InstanceData->PrevPose.IsEmpty())
+            {
+                // Here we have a single pose recorded so we cannot know the velocity of 
+                // the animation being transitioned from. In this case we assume zero
+                // velocity and construct a transition from the last pose alone.
+
+                check(!InstanceData->CurrPose.IsEmpty());
+
+				UE::AnimNext::FDeadBlendTransitionTaskParameters Parameters;
+				Parameters.ExtrapolationHalfLife = SharedData->ExtrapolationHalfLife;
+				Parameters.ExtrapolationHalfLifeMin = SharedData->ExtrapolationHalfLifeMin;
+				Parameters.ExtrapolationHalfLifeMax = SharedData->ExtrapolationHalfLifeMax;
+				Parameters.MaximumTranslationVelocity = SharedData->MaximumTranslationVelocity;
+				Parameters.MaximumRotationVelocity = FMath::DegreesToRadians(SharedData->MaximumRotationVelocity);
+				Parameters.MaximumScaleVelocity = SharedData->MaximumScaleVelocity;
+
+                Context.AppendTask(FAnimNextDeadBlendingTransitionTask::Make(
+					&InstanceData->State, 
+					&InstanceData->CurrPose, 
+					Parameters));
+
+                InstanceData->bRequestActive = true;
+                InstanceData->ActiveRequest = InstanceData->PendingRequest;
+                InstanceData->bRequestPending = false;
+                InstanceData->PendingRequest = UE::AnimNext::FInertializationRequest();
+                InstanceData->TimeSinceTransition = 0.0f;
+            }
+            else
+            {
+                // Transition as normal
+
+                check(!InstanceData->CurrPose.IsEmpty());
+                check(!InstanceData->PrevPose.IsEmpty());
+
+				UE::AnimNext::FDeadBlendTransitionTaskParameters Parameters;
+				Parameters.ExtrapolationHalfLife = SharedData->ExtrapolationHalfLife;
+				Parameters.ExtrapolationHalfLifeMin = SharedData->ExtrapolationHalfLifeMin;
+				Parameters.ExtrapolationHalfLifeMax = SharedData->ExtrapolationHalfLifeMax;
+				Parameters.MaximumTranslationVelocity = SharedData->MaximumTranslationVelocity;
+				Parameters.MaximumRotationVelocity = FMath::DegreesToRadians(SharedData->MaximumRotationVelocity);
+				Parameters.MaximumScaleVelocity = SharedData->MaximumScaleVelocity;
+
+                Context.AppendTask(FAnimNextDeadBlendingTransitionTask::Make(
+                    &InstanceData->State,
+                    &InstanceData->CurrPose,
+                    &InstanceData->PrevPose,
+                    InstanceData->PoseDeltaTime,
+					Parameters));
+
+                InstanceData->bRequestActive = true;
+                InstanceData->ActiveRequest = InstanceData->PendingRequest;
+                InstanceData->bRequestPending = false;
+                InstanceData->PendingRequest = UE::AnimNext::FInertializationRequest();
+                InstanceData->TimeSinceTransition = 0.0f;
+            }
+        } 
+
+        // If we have a current Inertialization Request active then add task to smooth out transition
+		if (InstanceData->bRequestActive)
+		{
+			// Accumulate time since transition.
+			 
+			// Since the transition pose recorded by FAnimNextDeadBlendingTransitionTask is actually 
+			// from the previous evaluation (i.e. it is the pose stored in CurrPose) even if we have
+			// just transitioned this from it still makes sense here to add the evaluation delta time 
+			// so that we extrapolate this forward to a pose which matches in time what is currently 
+			// on top of the evaluation stack.
+			
+			InstanceData->TimeSinceTransition += InstanceData->EvaluateDeltaTime;
+
+			if (InstanceData->TimeSinceTransition > InstanceData->ActiveRequest.BlendTime)
+			{
+				// Deactivate Dead Blending since transition is complete
+
+				InstanceData->bRequestActive = false;
+				InstanceData->ActiveRequest = UE::AnimNext::FInertializationRequest();
+				InstanceData->TimeSinceTransition = 0.0f;
+				InstanceData->State.Empty();
+			}
+			else
+			{
+				// Add Extrapolation and Blending Task
+
+				Context.AppendTask(FAnimNextDeadBlendingApplyTask::Make(
+					&InstanceData->State,
+					InstanceData->ActiveRequest.BlendTime,
+					InstanceData->TimeSinceTransition,
+					SharedData->DefaultBlendMode,
+					SharedData->DefaultCustomBlendCurve));
+			}
+		}
+
+        // We only need to swap recorded poses if we've been ticked inbetween evaluation with a non trivial DeltaTime
+        if (InstanceData->EvaluateDeltaTime > UE_SMALL_NUMBER)
+        {
+			// Record new DelaTime for pose swap/storage and reset evaluation delta time
+
+			InstanceData->PoseDeltaTime = InstanceData->EvaluateDeltaTime;
+			InstanceData->EvaluateDeltaTime = 0.0f;
+
+            // Add tasks for swapping and storing the new output pose
+
+            Context.AppendTask(FAnimNextSwapTransformsTask::Make(&InstanceData->PrevPose, &InstanceData->CurrPose));
+            Context.AppendTask(FAnimNextStoreKeyframeTransformsTask::Make(&InstanceData->CurrPose));
+        }
+	}
+
+}
diff --git a/Engine/Plugins/Experimental/AnimNextAnimGraph/Source/AnimNextAnimGraph/Private/Traits/InertializationDeadBlend.h b/Engine/Plugins/Experimental/AnimNextAnimGraph/Source/AnimNextAnimGraph/Private/Traits/InertializationDeadBlend.h
new file mode 100644
--- /dev/null
+++ b/Engine/Plugins/Experimental/AnimNextAnimGraph/Source/AnimNextAnimGraph/Private/Traits/InertializationDeadBlend.h
@@ -0,0 +1,141 @@
+// Copyright Epic Games, Inc. All Rights Reserved.
+
+#pragma once
+
+#include "CoreMinimal.h"
+
+#include "AlphaBlend.h"
+#include "TraitCore/Trait.h"
+#include "TraitInterfaces/IEvaluate.h"
+#include "TraitInterfaces/IUpdate.h"
+#include "TransformArray.h"
+
+#include "Traits/Inertialization.h"
+
+#include "InertializationDeadBlend.generated.h"
+
+class UCurveFloat;
+
+USTRUCT(meta = (DisplayName = "Dead Blending"))
+struct FAnimNextInertializationDeadBlendTraitSharedData : public FAnimNextTraitSharedData
+{
+	GENERATED_BODY()
+
+	// The default blend duration to use when "Always Use Default Blend Settings" is set to true.
+	UPROPERTY(EditAnywhere, Category = Blending, meta = (Inline))
+	float DefaultBlendDuration = 0.25f;
+
+	// Default blend mode to use when no blend mode is supplied with the inertialization request.
+	UPROPERTY(EditAnywhere, Category = Blending, meta = (Inline))
+	EAlphaBlendOption DefaultBlendMode = EAlphaBlendOption::Linear;
+
+	// Default custom blend curve to use along with the default blend mode.
+	UPROPERTY(EditAnywhere, Category = Blending, meta = (Inline))
+	TObjectPtr<UCurveFloat> DefaultCustomBlendCurve = nullptr;
+
+	/**
+	 * The average half-life of decay in seconds to use when extrapolating the animation. To get the final half-life of
+	 * decay, this value will be scaled by the amount to which the velocities of the animation being transitioned from
+	 * are moving toward the animation being transitioned to.
+	 */
+	UPROPERTY(EditAnywhere, Category = Extrapolation, meta = (Inline, Min = "0.0", UIMin = "0.0"), DisplayName = "Extrapolation Half Life")
+	float ExtrapolationHalfLife = 1.0f;
+
+	/**
+	 * The minimum half-life of decay in seconds to use when extrapolating the animation. This will be used when the
+	 * velocities of the animation being transitioned from are very small or moving away from the animation being
+	 * transitioned to.
+	 */
+	UPROPERTY(EditAnywhere, Category = Extrapolation, meta = (Inline, Min = "0.0", UIMin = "0.0"), DisplayName = "Minimum Extrapolation Half Life")
+	float ExtrapolationHalfLifeMin = 0.05f;
+
+	/**
+	 * The maximum half-life of decay in seconds to use when extrapolating the animation. This will dictate the longest
+	 * decay duration possible when velocities of the animation being transitioned from are small and moving towards the
+	 * animation being transitioned to.
+	 */
+	UPROPERTY(EditAnywhere, Category = Extrapolation, meta = (Inline, Min = "0.0", UIMin = "0.0"), DisplayName = "Maximum Extrapolation Half Life")
+	float ExtrapolationHalfLifeMax = 1.0f;
+
+	/**
+	 * The maximum velocity to allow for extrapolation of bone translations in centimeters per second. Smaller values
+	 * may help prevent the pose breaking during blending but too small values can make the blend less smooth.
+	 */
+	UPROPERTY(EditAnywhere, Category = Extrapolation, meta = (Inline, Min = "0.0", UIMin = "0.0"))
+	float MaximumTranslationVelocity = 500.0f;
+
+	/**
+	 * The maximum velocity to allow for extrapolation of bone rotations in degrees per second. Smaller values
+	 * may help prevent the pose breaking during blending but too small values can make the blend less smooth.
+	 */
+	UPROPERTY(EditAnywhere, Category = Extrapolation, meta = (Inline, Min = "0.0", UIMin = "0.0"))
+	float MaximumRotationVelocity = 360.0f;
+
+	/**
+	 * The maximum velocity to allow for extrapolation of bone scales. Smaller values may help prevent the pose
+	 * breaking during blending but too small values can make the blend less smooth.
+	 */
+	UPROPERTY(EditAnywhere, Category = Extrapolation, meta = (Inline, Min = "0.0", UIMin = "0.0"))
+	float MaximumScaleVelocity = 4.0f;
+
+	/**
+	 * The maximum velocity to allow for extrapolation of curves. Smaller values may help prevent extreme curve values
+	 * during blending but too small values can make the blending of curves less smooth.
+	 */
+	UPROPERTY(EditAnywhere, Category = Extrapolation, meta = (Inline, Min = "0.0", UIMin = "0.0"))
+	float MaximumCurveVelocity = 100.0f;
+};
+
+namespace UE::AnimNext
+{
+	/**
+	 * FInertializationDeadBlendTrait
+	 *
+	 * A trait that inertializes animation by extrapolating from the point of transition.
+	 */
+	struct FInertializationDeadBlendTrait : FAdditiveTrait, IUpdate, IEvaluate
+	{
+		DECLARE_ANIM_TRAIT(FInertializationDeadBlendTrait, 0xf11f2b26, FAdditiveTrait)
+
+		using FSharedData = FAnimNextInertializationDeadBlendTraitSharedData;
+
+		struct FInstanceData : FTrait::FInstanceData
+		{
+			// Pending request details
+			bool bRequestPending = false;
+			FInertializationRequest PendingRequest;
+
+			// Current active request details
+			bool bRequestActive = false;
+			FInertializationRequest ActiveRequest;
+
+			// Time since the last transition
+			float TimeSinceTransition = 0.0f;
+
+			// Delta Time since the last evaluation
+			float EvaluateDeltaTime = 0.0f;
+
+			// Delta Time between the storage of Curr and Prev Poses
+			float PoseDeltaTime = 0.0f;
+
+			// Recorded Current Pose State
+			FTransformArraySoAHeap CurrPose;
+
+			// Recorded Previous Pose State
+			FTransformArraySoAHeap PrevPose;
+
+			// Extrapolation State Data
+			FDeadBlendingState State;
+		};
+
+		// InertializationRequestTrait handler
+		ETraitStackPropagation OnInertializationRequestEvent(const FExecutionContext& Context, FTraitBinding& Binding, FAnimNextInertializationRequestEvent& Event) const;
+
+		// IUpdate impl
+		virtual void PostUpdate(FUpdateTraversalContext& Context, const TTraitBinding<IUpdate>& Binding, const FTraitUpdateState& TraitState) const override;
+
+		// IEvaluate impl
+		virtual void PostEvaluate(FEvaluateTraversalContext& Context, const TTraitBinding<IEvaluate>& Binding) const override;
+
+	};
+}
diff --git a/Engine/Plugins/Experimental/AnimNextAnimGraph/Source/AnimNextAnimGraph/Private/Traits/Passthrough.cpp b/Engine/Plugins/Experimental/AnimNextAnimGraph/Source/AnimNextAnimGraph/Private/Traits/Passthrough.cpp
new file mode 100644
--- /dev/null
+++ b/Engine/Plugins/Experimental/AnimNextAnimGraph/Source/AnimNextAnimGraph/Private/Traits/Passthrough.cpp
@@ -0,0 +1,43 @@
+// Copyright Epic Games, Inc. All Rights Reserved.
+
+#include "Traits/Passthrough.h"
+
+#include "TraitCore/ExecutionContext.h"
+
+namespace UE::AnimNext
+{
+	AUTO_REGISTER_ANIM_TRAIT(FPassthroughTrait)
+
+		// Trait implementation boilerplate
+#define TRAIT_INTERFACE_ENUMERATOR(GeneratorMacro) \
+		GeneratorMacro(IHierarchy) \
+
+		GENERATE_ANIM_TRAIT_IMPLEMENTATION(FPassthroughTrait, TRAIT_INTERFACE_ENUMERATOR, NULL_ANIM_TRAIT_EVENT_ENUMERATOR)
+#undef TRAIT_INTERFACE_ENUMERATOR
+
+	void FPassthroughTrait::FInstanceData::Construct(const FExecutionContext& Context, const FTraitBinding& Binding)
+	{
+		FTrait::FInstanceData::Construct(Context, Binding);
+
+		const FSharedData* SharedData = Binding.GetSharedData<FSharedData>();
+		FInstanceData* InstanceData = Binding.GetInstanceData<FInstanceData>();
+
+		if (!InstanceData->Input.IsValid())
+		{
+			InstanceData->Input = Context.AllocateNodeInstance(Binding, SharedData->Input);
+		}
+	}
+
+	uint32 FPassthroughTrait::GetNumChildren(const FExecutionContext& Context, const TTraitBinding<IHierarchy>& Binding) const
+	{
+		return 1;
+	}
+
+	void FPassthroughTrait::GetChildren(const FExecutionContext& Context, const TTraitBinding<IHierarchy>& Binding, FChildrenArray& Children) const
+	{
+		const FInstanceData* InstanceData = Binding.GetInstanceData<FInstanceData>();
+
+		// Add the child, even if the handle is empty
+		Children.Add(InstanceData->Input);
+	}
+}
diff --git a/Engine/Plugins/Experimental/AnimNextAnimGraph/Source/AnimNextAnimGraph/Private/Traits/Passthrough.h b/Engine/Plugins/Experimental/AnimNextAnimGraph/Source/AnimNextAnimGraph/Private/Traits/Passthrough.h
new file mode 100644
--- /dev/null
+++ b/Engine/Plugins/Experimental/AnimNextAnimGraph/Source/AnimNextAnimGraph/Private/Traits/Passthrough.h
@@ -0,0 +1,47 @@
+// Copyright Epic Games, Inc. All Rights Reserved.
+
+#pragma once
+
+#include "CoreMinimal.h"
+
+#include "TraitCore/Trait.h"
+#include "TraitInterfaces/IHierarchy.h"
+#include "TraitInterfaces/IUpdate.h"
+
+#include "Passthrough.generated.h"
+
+USTRUCT(meta = (DisplayName = "Passthough"))
+struct FAnimNextPassthroughSharedData : public FAnimNextTraitSharedData
+{
+	GENERATED_BODY()
+
+	/** Input to pass to output */
+	UPROPERTY()
+	FAnimNextTraitHandle Input;
+};
+
+namespace UE::AnimNext
+{
+	/**
+	 * FPassthroughTrait
+	 *
+	 * A trait that passes through the input without modification.
+	 */
+	struct FPassthroughTrait : FBaseTrait, IHierarchy
+	{
+		DECLARE_ANIM_TRAIT(FPassthroughTrait, 0xa4982ed3, FBaseTrait)
+
+		using FSharedData = FAnimNextPassthroughSharedData;
+
+		struct FInstanceData : FTrait::FInstanceData
+		{
+			FTraitPtr Input;
+
+			void Construct(const FExecutionContext& Context, const FTraitBinding& Binding);
+		};
+
+		// IHierarchy impl
+		virtual uint32 GetNumChildren(const FExecutionContext& Context, const TTraitBinding<IHierarchy>& Binding) const override;
+		virtual void GetChildren(const FExecutionContext& Context, const TTraitBinding<IHierarchy>& Binding, FChildrenArray& Children) const override;
+	};
+}
diff --git a/Engine/Plugins/Experimental/AnimNextAnimGraph/Source/AnimNextAnimGraph/Private/Traits/SubGraphHost.cpp b/Engine/Plugins/Experimental/AnimNextAnimGraph/Source/AnimNextAnimGraph/Private/Traits/SubGraphHost.cpp
--- a/Engine/Plugins/Experimental/AnimNextAnimGraph/Source/AnimNextAnimGraph/Private/Traits/SubGraphHost.cpp
+++ b/Engine/Plugins/Experimental/AnimNextAnimGraph/Source/AnimNextAnimGraph/Private/Traits/SubGraphHost.cpp
@@ -187,7 +187,7 @@ namespace UE::AnimNext
 		return InstanceData->CurrentlyActiveSubGraphIndex;
 	}
 
-	void FSubGraphHostTrait::OnBlendTransition(const FExecutionContext& Context, const TTraitBinding<IDiscreteBlend>& Binding, int32 OldChildIndex, int32 NewChildIndex) const
+	void FSubGraphHostTrait::OnBlendTransition(FExecutionContext& Context, const TTraitBinding<IDiscreteBlend>& Binding, int32 OldChildIndex, int32 NewChildIndex) const
 	{
 		TTraitBinding<IDiscreteBlend> DiscreteBlendTrait;
 		Binding.GetStackInterface(DiscreteBlendTrait);
@@ -199,7 +199,7 @@ namespace UE::AnimNext
 		DiscreteBlendTrait.OnBlendTerminated(Context, OldChildIndex);
 	}
 
-	void FSubGraphHostTrait::OnBlendInitiated(const FExecutionContext& Context, const TTraitBinding<IDiscreteBlend>& Binding, int32 ChildIndex) const
+	void FSubGraphHostTrait::OnBlendInitiated(FExecutionContext& Context, const TTraitBinding<IDiscreteBlend>& Binding, int32 ChildIndex) const
 	{
 		FInstanceData* InstanceData = Binding.GetInstanceData<FInstanceData>();
 
@@ -215,7 +215,7 @@ namespace UE::AnimNext
 		}
 	}
 
-	void FSubGraphHostTrait::OnBlendTerminated(const FExecutionContext& Context, const TTraitBinding<IDiscreteBlend>& Binding, int32 ChildIndex) const
+	void FSubGraphHostTrait::OnBlendTerminated(FExecutionContext& Context, const TTraitBinding<IDiscreteBlend>& Binding, int32 ChildIndex) const
 	{
 		FInstanceData* InstanceData = Binding.GetInstanceData<FInstanceData>();
 
diff --git a/Engine/Plugins/Experimental/AnimNextAnimGraph/Source/AnimNextAnimGraph/Private/Traits/SubGraphHost.h b/Engine/Plugins/Experimental/AnimNextAnimGraph/Source/AnimNextAnimGraph/Private/Traits/SubGraphHost.h
--- a/Engine/Plugins/Experimental/AnimNextAnimGraph/Source/AnimNextAnimGraph/Private/Traits/SubGraphHost.h
+++ b/Engine/Plugins/Experimental/AnimNextAnimGraph/Source/AnimNextAnimGraph/Private/Traits/SubGraphHost.h
@@ -107,9 +107,9 @@ namespace UE::AnimNext
 		// IDiscreteBlend impl
 		virtual float GetBlendWeight(const FExecutionContext& Context, const TTraitBinding<IDiscreteBlend>& Binding, int32 ChildIndex) const override;
 		virtual int32 GetBlendDestinationChildIndex(const FExecutionContext& Context, const TTraitBinding<IDiscreteBlend>& Binding) const override;
-		virtual void OnBlendTransition(const FExecutionContext& Context, const TTraitBinding<IDiscreteBlend>& Binding, int32 OldChildIndex, int32 NewChildIndex) const override;
-		virtual void OnBlendInitiated(const FExecutionContext& Context, const TTraitBinding<IDiscreteBlend>& Binding, int32 ChildIndex) const override;
-		virtual void OnBlendTerminated(const FExecutionContext& Context, const TTraitBinding<IDiscreteBlend>& Binding, int32 ChildIndex) const override;
+		virtual void OnBlendTransition(FExecutionContext& Context, const TTraitBinding<IDiscreteBlend>& Binding, int32 OldChildIndex, int32 NewChildIndex) const override;
+		virtual void OnBlendInitiated(FExecutionContext& Context, const TTraitBinding<IDiscreteBlend>& Binding, int32 ChildIndex) const override;
+		virtual void OnBlendTerminated(FExecutionContext& Context, const TTraitBinding<IDiscreteBlend>& Binding, int32 ChildIndex) const override;
 
 		// IGarbageCollection impl
 		virtual void AddReferencedObjects(const FExecutionContext& Context, const TTraitBinding<IGarbageCollection>& Binding, FReferenceCollector& Collector) const override;
diff --git a/Engine/Plugins/Experimental/AnimNextAnimGraph/Source/AnimNextAnimGraph/Public/EvaluationVM/Tasks/DeadBlending.h b/Engine/Plugins/Experimental/AnimNextAnimGraph/Source/AnimNextAnimGraph/Public/EvaluationVM/Tasks/DeadBlending.h
new file mode 100644
--- /dev/null
+++ b/Engine/Plugins/Experimental/AnimNextAnimGraph/Source/AnimNextAnimGraph/Public/EvaluationVM/Tasks/DeadBlending.h
@@ -0,0 +1,121 @@
+// Copyright Epic Games, Inc. All Rights Reserved.
+
+#pragma once
+
+#include "CoreMinimal.h"
+#include "EvaluationVM/EvaluationTask.h"
+#include "TransformArray.h"
+#include "TransformArrayView.h"
+#include "AlphaBlend.h"
+
+#include "DeadBlending.generated.h"
+
+namespace UE::AnimNext
+{
+	struct FDeadBlendingState;
+
+	/**
+	 * FDeadBlendTransitionTaskParameters
+	 *
+	 * Parameters for the Dead Blending Transition Task. This should be filled according to the 
+	 * values provided by the Dead Blending Trait shared data.
+	 */
+	struct FDeadBlendTransitionTaskParameters
+	{
+		float ExtrapolationHalfLife = 0.0f;
+		float ExtrapolationHalfLifeMin = 0.0f;
+		float ExtrapolationHalfLifeMax = 0.0f;
+		float MaximumTranslationVelocity = 0.0f;
+		float MaximumRotationVelocity = 0.0f;
+		float MaximumScaleVelocity = 0.0f;
+	};
+}
+
+class UCurveFloat;
+
+/**
+ * FAnimNextDeadBlendingTransitionTask
+ * 
+ * Task for performing a Dead Blending transition. Stores the current pose state, computes the 
+ * velocity using finite difference, and fits extrapolation half-lives based on the state being
+ * transitioned to.
+ */
+USTRUCT()
+struct ANIMNEXTANIMGRAPH_API FAnimNextDeadBlendingTransitionTask : public FAnimNextEvaluationTask
+{
+	GENERATED_BODY()
+
+	DECLARE_ANIM_EVALUATION_TASK(FAnimNextDeadBlendingTransitionTask)
+
+	// Make a DeadBlendingTransitionTask using two previous poses
+	static FAnimNextDeadBlendingTransitionTask Make(
+		UE::AnimNext::FDeadBlendingState* State,
+		const UE::AnimNext::FTransformArraySoAHeap* CurrPose,
+		const UE::AnimNext::FTransformArraySoAHeap* PrevPose,
+		const float DeltaTime,
+		const UE::AnimNext::FDeadBlendTransitionTaskParameters& Parameters);
+
+	// Make a DeadBlendingTransitionTask using only a single previous poses
+	static FAnimNextDeadBlendingTransitionTask Make(
+		UE::AnimNext::FDeadBlendingState* State,
+		const UE::AnimNext::FTransformArraySoAHeap* CurrPose,
+		const UE::AnimNext::FDeadBlendTransitionTaskParameters& Parameters);
+
+	// Task entry point
+	virtual void Execute(UE::AnimNext::FEvaluationVM& VM) const override;
+
+	// Output
+
+	UE::AnimNext::FDeadBlendingState* State = nullptr;
+
+	// Inputs
+
+	const UE::AnimNext::FTransformArraySoAHeap* CurrPose = nullptr;
+	const UE::AnimNext::FTransformArraySoAHeap* PrevPose = nullptr;
+
+	// Parameters
+
+	float DeltaTime = 0.0f;
+	UE::AnimNext::FDeadBlendTransitionTaskParameters Parameters;
+};
+
+/**
+ * FAnimNextDeadBlendingApplyTask
+ *
+ * Task for applying dead blending to remove a discontinuity. Will extrapolate the FDeadBlendingState forward
+ * in time and then blend it with the Keyframe which is on top of the stack according to the BlendDuration. Part
+ * of the FDeadBlendingState will be updated to ensure the blend does not "flip sides" part way through.
+ */
+USTRUCT()
+struct ANIMNEXTANIMGRAPH_API FAnimNextDeadBlendingApplyTask : public FAnimNextEvaluationTask
+{
+	GENERATED_BODY()
+
+	DECLARE_ANIM_EVALUATION_TASK(FAnimNextDeadBlendingApplyTask)
+
+	// Make a FAnimNextDeadBlendingApplyTask task
+	static FAnimNextDeadBlendingApplyTask Make(
+		UE::AnimNext::FDeadBlendingState* State,
+		const float BlendDuration,
+		const float TimeSinceTransition,
+		const EAlphaBlendOption& BlendMode,
+		const TWeakObjectPtr<UCurveFloat> CustomBlendCurve);
+
+	// Task entry point
+	virtual void Execute(UE::AnimNext::FEvaluationVM& VM) const override;
+
+	// Input / Output
+
+	UE::AnimNext::FDeadBlendingState* State = nullptr;
+
+	// Parameters
+
+	float BlendDuration = 0.0f;
+	float TimeSinceTransition = 0.0f;
+	EAlphaBlendOption BlendMode = EAlphaBlendOption::Linear;
+	
+	UPROPERTY()
+	TWeakObjectPtr<UCurveFloat> CustomBlendCurve = nullptr;
+};
+
+
diff --git a/Engine/Plugins/Experimental/AnimNextAnimGraph/Source/AnimNextAnimGraph/Public/TraitInterfaces/IDiscreteBlend.h b/Engine/Plugins/Experimental/AnimNextAnimGraph/Source/AnimNextAnimGraph/Public/TraitInterfaces/IDiscreteBlend.h
--- a/Engine/Plugins/Experimental/AnimNextAnimGraph/Source/AnimNextAnimGraph/Public/TraitInterfaces/IDiscreteBlend.h
+++ b/Engine/Plugins/Experimental/AnimNextAnimGraph/Source/AnimNextAnimGraph/Public/TraitInterfaces/IDiscreteBlend.h
@@ -38,13 +38,13 @@ namespace UE::AnimNext
 		// NewChildIndex can be larger than the current number of known children to support a dynamic number of children at runtime
 		// When this occurs, the number of children increments by one
 		// The number of children never shrinks
-		virtual void OnBlendTransition(const FExecutionContext& Context, const TTraitBinding<IDiscreteBlend>& Binding, int32 OldChildIndex, int32 NewChildIndex) const;
+		virtual void OnBlendTransition(FExecutionContext& Context, const TTraitBinding<IDiscreteBlend>& Binding, int32 OldChildIndex, int32 NewChildIndex) const;
 
 		// Called when the blend for specified child is initiated
-		virtual void OnBlendInitiated(const FExecutionContext& Context, const TTraitBinding<IDiscreteBlend>& Binding, int32 ChildIndex) const;
+		virtual void OnBlendInitiated(FExecutionContext& Context, const TTraitBinding<IDiscreteBlend>& Binding, int32 ChildIndex) const;
 
 		// Called when the blend for specified child terminates
-		virtual void OnBlendTerminated(const FExecutionContext& Context, const TTraitBinding<IDiscreteBlend>& Binding, int32 ChildIndex) const;
+		virtual void OnBlendTerminated(FExecutionContext& Context, const TTraitBinding<IDiscreteBlend>& Binding, int32 ChildIndex) const;
 	};
 
 	/**
@@ -72,19 +72,19 @@ namespace UE::AnimNext
 		}
 
 		// @see IDiscreteBlend::OnBlendTransition
-		void OnBlendTransition(const FExecutionContext& Context, int32 OldChildIndex, int32 NewChildIndex) const
+		void OnBlendTransition(FExecutionContext& Context, int32 OldChildIndex, int32 NewChildIndex) const
 		{
 			GetInterface()->OnBlendTransition(Context, *this, OldChildIndex, NewChildIndex);
 		}
 
 		// @see IDiscreteBlend::OnBlendInitiated
-		void OnBlendInitiated(const FExecutionContext& Context, int32 ChildIndex) const
+		void OnBlendInitiated(FExecutionContext& Context, int32 ChildIndex) const
 		{
 			GetInterface()->OnBlendInitiated(Context, *this, ChildIndex);
 		}
 
 		// @see IDiscreteBlend::OnBlendTerminated
-		void OnBlendTerminated(const FExecutionContext& Context, int32 ChildIndex) const
+		void OnBlendTerminated(FExecutionContext& Context, int32 ChildIndex) const
 		{
 			GetInterface()->OnBlendTerminated(Context, *this, ChildIndex);
 		}
diff --git a/Engine/Plugins/Experimental/AnimNextAnimGraph/Source/AnimNextAnimGraph/Public/Traits/Inertialization.h b/Engine/Plugins/Experimental/AnimNextAnimGraph/Source/AnimNextAnimGraph/Public/Traits/Inertialization.h
new file mode 100644
--- /dev/null
+++ b/Engine/Plugins/Experimental/AnimNextAnimGraph/Source/AnimNextAnimGraph/Public/Traits/Inertialization.h
@@ -0,0 +1,80 @@
+// Copyright Epic Games, Inc. All Rights Reserved.
+
+#pragma once
+
+#include "CoreMinimal.h"
+
+#include "TraitCore/TraitUID.h"
+#include "TraitCore/TraitEvent.h"
+#include "TransformArray.h"
+
+#include "Inertialization.generated.h"
+
+namespace UE::AnimNext
+{
+	/**
+	 * FInertializationRequest
+	 *
+	 * Data contained in an InertializationRequest
+	 */
+	struct ANIMNEXTANIMGRAPH_API FInertializationRequest
+	{
+		// Request Blend Time
+		float BlendTime = 0.0f;
+	};
+
+	/**
+	 * FDeadBlendingState
+	 *
+	 * State data required by the DeadBlending node when extrapolating the pose.
+	 */
+	struct ANIMNEXTANIMGRAPH_API FDeadBlendingState
+	{
+		inline void Empty()
+		{
+			BoneRotationDirections.Empty();
+			SourcePose.Empty();
+			SourceBoneTranslationVelocities.Empty();
+			SourceBoneRotationVelocities.Empty();
+			SourceBoneScaleVelocities.Empty();
+			SourceBoneTranslationDecayHalfLives.Empty();
+			SourceBoneRotationDecayHalfLives.Empty();
+			SourceBoneScaleDecayHalfLives.Empty();
+		}
+
+		inline void SetNumUninitialized(const int32 NewBoneNum)
+		{
+			BoneRotationDirections.SetNumUninitialized(NewBoneNum);
+			SourcePose.SetNumUninitialized(NewBoneNum);
+			SourceBoneTranslationVelocities.SetNumUninitialized(NewBoneNum);
+			SourceBoneRotationVelocities.SetNumUninitialized(NewBoneNum);
+			SourceBoneScaleVelocities.SetNumUninitialized(NewBoneNum);
+			SourceBoneTranslationDecayHalfLives.SetNumUninitialized(NewBoneNum);
+			SourceBoneRotationDecayHalfLives.SetNumUninitialized(NewBoneNum);
+			SourceBoneScaleDecayHalfLives.SetNumUninitialized(NewBoneNum);
+		}
+
+		TArray<FQuat4f> BoneRotationDirections;
+		FTransformArraySoAHeap SourcePose;
+		TArray<FVector3f> SourceBoneTranslationVelocities;
+		TArray<FVector3f> SourceBoneRotationVelocities;
+		TArray<FVector3f> SourceBoneScaleVelocities;
+		TArray<FVector3f> SourceBoneTranslationDecayHalfLives;
+		TArray<FVector3f> SourceBoneRotationDecayHalfLives;
+		TArray<FVector3f> SourceBoneScaleDecayHalfLives;
+	};
+}
+
+/**
+ * FAnimNextInertializationRequestEvent
+ *
+ * Inertialization Request Event Object
+ */
+USTRUCT()
+struct ANIMNEXTANIMGRAPH_API FAnimNextInertializationRequestEvent : public FAnimNextTraitEvent
+{
+	GENERATED_BODY()
+	DECLARE_ANIM_TRAIT_EVENT(FAnimNextInertializationRequestEvent)
+
+	UE::AnimNext::FInertializationRequest Request;
+};
\ No newline at end of file
