Added AnimNext Traits Editor, enabling add, removal and reordering of AnimNext Traits in a Trait Stack (supports clicking and drag and drop).
Also allows displaying the implemented and required interfaces of the traits added to a stack (visuals are WIP).
Added helper functions to AnimNextGraph_Controller for Add, Remove and Swap
Currently implemented as a Tab, it gets the required information from the TraitStack node customization, when we click on a node (sets the data) or when we click away (clears the data).
#rnx

diff --git a/Engine/Plugins/Experimental/AnimNext/Source/AnimNext/Internal/TraitCore/TraitRegistry.h b/Engine/Plugins/Experimental/AnimNext/Source/AnimNext/Internal/TraitCore/TraitRegistry.h
--- a/Engine/Plugins/Experimental/AnimNext/Source/AnimNext/Internal/TraitCore/TraitRegistry.h
+++ b/Engine/Plugins/Experimental/AnimNext/Source/AnimNext/Internal/TraitCore/TraitRegistry.h
@@ -37,6 +37,10 @@ namespace UE::AnimNext
 		// If the matching trait is not registered, nullptr is returned.
 		const FTrait* Find(const UScriptStruct* TraitSharedDataStruct) const;
 
+		// Finds and returns the trait UID associated with the provided trait name.
+		// If the matching trait is not registered, nullptr is returned.
+		const FTrait* Find(FName TraitTypeName) const;
+
 		// Registers a trait dynamically
 		void Register(FTrait* Trait);
 
@@ -98,7 +102,8 @@ namespace UE::AnimNext
 		TArray<uintptr_t>				DynamicTraits;
 		int32							DynamicTraitFreeIndexHead = INDEX_NONE;
 
-		TMap<uint32, FRegistryEntry>	TraitUIDToEntryMap;
+		TMap<FTraitUIDRaw, FRegistryEntry>	TraitUIDToEntryMap;
+		TMap<FName, FTraitUIDRaw>			TraitNameToUIDMap;
 
 		friend class FModule;
 		friend struct FTraitStaticInitHook;
diff --git a/Engine/Plugins/Experimental/AnimNext/Source/AnimNext/Private/TraitCore/TraitRegistry.cpp b/Engine/Plugins/Experimental/AnimNext/Source/AnimNext/Private/TraitCore/TraitRegistry.cpp
--- a/Engine/Plugins/Experimental/AnimNext/Source/AnimNext/Private/TraitCore/TraitRegistry.cpp
+++ b/Engine/Plugins/Experimental/AnimNext/Source/AnimNext/Private/TraitCore/TraitRegistry.cpp
@@ -120,7 +120,7 @@ namespace UE::AnimNext
 
 		const FTraitUID TraitUID = Trait->GetTraitUID();
 
-		if (ensure(!TraitUIDToEntryMap.Contains(TraitUID.GetUID())))
+		if (ensure(!TraitUIDToEntryMap.Contains(TraitUID.GetUID())) && ensure(!TraitNameToUIDMap.Contains(*Trait->GetTraitName())))
 		{
 			// This is a new trait, we'll keep it
 			if (!bFitsInStaticBuffer)
@@ -143,6 +143,7 @@ namespace UE::AnimNext
 			}
 
 			TraitUIDToEntryMap.Add(TraitUID.GetUID(), FRegistryEntry{ Trait, TraitConstructor, TraitHandle });
+			TraitNameToUIDMap.Add(*Trait->GetTraitName(), TraitUID.GetUID());
 		}
 		else
 		{
@@ -174,6 +175,9 @@ namespace UE::AnimNext
 			{
 				check(Entry.TraitHandle.IsValid());
 
+				// Remove name from map before we destroy the Trait
+				TraitNameToUIDMap.Remove(*Entry.Trait->GetTraitName());
+
 				// Destroy and release our trait
 				// We always own auto-registered trait instances
 				Entry.Trait->~FTrait();
@@ -200,6 +204,7 @@ namespace UE::AnimNext
 				}
 
 				TraitUIDToEntryMap.Remove(It.Key());
+				
 
 				break;
 			}
@@ -266,6 +271,25 @@ namespace UE::AnimNext
 		return nullptr;
 	}
 
+	const FTrait* FTraitRegistry::Find(FName TraitTypeName) const
+	{
+		if (TraitTypeName == NAME_None)
+		{
+			return nullptr;
+		}
+
+		const FTraitUIDRaw* TraitGUID = TraitNameToUIDMap.Find(TraitTypeName);
+		if (TraitGUID != nullptr && FTraitUID(*TraitGUID).IsValid())
+		{
+			if (const FRegistryEntry* Entry = TraitUIDToEntryMap.Find(*TraitGUID))
+			{
+				return Entry->Trait;
+			}
+		}
+
+		return nullptr;
+	}
+
 	void FTraitRegistry::Register(FTrait* Trait)
 	{
 		if (Trait == nullptr)
@@ -275,7 +299,7 @@ namespace UE::AnimNext
 
 		const FTraitUID TraitUID = Trait->GetTraitUID();
 
-		if (ensure(!TraitUIDToEntryMap.Contains(TraitUID.GetUID())))
+		if (ensure(!TraitUIDToEntryMap.Contains(TraitUID.GetUID())) && ensure(!TraitNameToUIDMap.Contains(*Trait->GetTraitName())))
 		{
 			// This is a new trait, we'll keep it
 			// Find our dynamic trait index
@@ -297,6 +321,7 @@ namespace UE::AnimNext
 			FTraitRegistryHandle TraitHandle = FTraitRegistryHandle::MakeDynamic(TraitIndex);
 
 			TraitUIDToEntryMap.Add(TraitUID.GetUID(), FRegistryEntry{ Trait, nullptr, TraitHandle });
+			TraitNameToUIDMap.Add(*Trait->GetTraitName(), TraitUID.GetUID());
 		}
 	}
 
diff --git a/Engine/Plugins/Experimental/AnimNext/Source/AnimNext/Public/TraitCore/IScopedTraitInterface.h b/Engine/Plugins/Experimental/AnimNext/Source/AnimNext/Public/TraitCore/IScopedTraitInterface.h
--- a/Engine/Plugins/Experimental/AnimNext/Source/AnimNext/Public/TraitCore/IScopedTraitInterface.h
+++ b/Engine/Plugins/Experimental/AnimNext/Source/AnimNext/Public/TraitCore/IScopedTraitInterface.h
@@ -14,6 +14,6 @@ namespace UE::AnimNext
 	 */
 	struct ANIMNEXT_API IScopedTraitInterface : ITraitInterface
 	{
-		DECLARE_ANIM_TRAIT_INTERFACE(IScopedTraitInterface, 0x621b01eb)
+		DECLARE_ANIM_TRAIT_INTERFACE(IScopedTraitInterface, STI, 0x621b01eb)
 	};
 }
diff --git a/Engine/Plugins/Experimental/AnimNext/Source/AnimNext/Public/TraitCore/ITraitInterface.h b/Engine/Plugins/Experimental/AnimNext/Source/AnimNext/Public/TraitCore/ITraitInterface.h
--- a/Engine/Plugins/Experimental/AnimNext/Source/AnimNext/Public/TraitCore/ITraitInterface.h
+++ b/Engine/Plugins/Experimental/AnimNext/Source/AnimNext/Public/TraitCore/ITraitInterface.h
@@ -6,9 +6,9 @@
 #include "TraitCore/TraitInterfaceUID.h"
 
 // Helper macros
-#define DECLARE_ANIM_TRAIT_INTERFACE(InterfaceName, InterfaceNameHash) \
+#define DECLARE_ANIM_TRAIT_INTERFACE(InterfaceName, InterfaceShortName, InterfaceNameHash) \
 	/* Globally unique UID for this interface */ \
-	static constexpr UE::AnimNext::FTraitInterfaceUID InterfaceUID = UE::AnimNext::FTraitInterfaceUID(InterfaceNameHash, TEXT(#InterfaceName)); \
+	static constexpr UE::AnimNext::FTraitInterfaceUID InterfaceUID = UE::AnimNext::FTraitInterfaceUID(InterfaceNameHash, TEXT(#InterfaceName), TEXT(#InterfaceShortName)); \
 	virtual UE::AnimNext::FTraitInterfaceUID GetInterfaceUID() const override { return InterfaceUID; }
 
 namespace UE::AnimNext
diff --git a/Engine/Plugins/Experimental/AnimNext/Source/AnimNext/Public/TraitCore/Trait.h b/Engine/Plugins/Experimental/AnimNext/Source/AnimNext/Public/TraitCore/Trait.h
--- a/Engine/Plugins/Experimental/AnimNext/Source/AnimNext/Public/TraitCore/Trait.h
+++ b/Engine/Plugins/Experimental/AnimNext/Source/AnimNext/Public/TraitCore/Trait.h
@@ -33,6 +33,7 @@ class FArchive;
 	virtual void DestructTraitInstance(const UE::AnimNext::FExecutionContext& Context, const UE::AnimNext::FTraitBinding& Binding) const override; \
 	virtual const UE::AnimNext::ITraitInterface* GetTraitInterface(UE::AnimNext::FTraitInterfaceUID InterfaceUID) const override; \
 	virtual TConstArrayView<FTraitInterfaceUID> GetTraitInterfaces() const override; \
+	virtual TConstArrayView<FTraitInterfaceUID> GetTraitRequiredInterfaces() const override; \
 	virtual UE::AnimNext::ETraitStackPropagation OnTraitEvent(UE::AnimNext::FExecutionContext& Context, UE::AnimNext::FTraitBinding& Binding, FAnimNextTraitEvent& Event) const override; \
 	virtual TConstArrayView<FTraitEventUID> GetTraitEvents() const override; \
 	virtual uint32 GetNumLatentTraitProperties() const override { return -FSharedData::GetLatentPropertyIndex(~(size_t)0); } \
@@ -117,6 +118,19 @@ class FArchive;
 		return CachedInterfaceList; \
 	}
 
+// Implements GetTraitRequiredInterfaces()
+#define ANIM_NEXT_IMPL_DEFINE_ANIM_TRAIT_GET_REQUIRED_INTERFACES(TraitName, InterfaceEnumeratorMacro) \
+	TConstArrayView<UE::AnimNext::FTraitInterfaceUID> TraitName::GetTraitRequiredInterfaces() const \
+	{ \
+		/* Thread safe cache initialization */ \
+		static TArray<UE::AnimNext::FTraitInterfaceUID> CachedInterfaceList = FTrait::BuildTraitInterfaceList( \
+			TraitSuper::GetTraitInterfaces(), \
+			{ \
+				InterfaceEnumeratorMacro(ANIM_NEXT_IMPL_GET_INTERFACES_IMPL_FOR_INTERFACE) \
+			}); \
+		return CachedInterfaceList; \
+	}
+
 namespace UE::AnimNext::Private
 {
 	// Helper to grab the event type from an event handler function signature
@@ -191,10 +205,11 @@ namespace UE::AnimNext::Private
 /**
   * This macro defines the necessary boilerplate for implementing FTrait. See above for usage example.
   */
-#define GENERATE_ANIM_TRAIT_IMPLEMENTATION(TraitName, InterfaceEnumeratorMacro, EventEnumeratorMacro) \
+#define GENERATE_ANIM_TRAIT_IMPLEMENTATION(TraitName, InterfaceEnumeratorMacro, RequiredInterfaceEnumeratorMacro, EventEnumeratorMacro) \
 	ANIM_NEXT_IMPL_DEFINE_ANIM_TRAIT(TraitName) \
 	ANIM_NEXT_IMPL_DEFINE_ANIM_TRAIT_GET_INTERFACE(TraitName, InterfaceEnumeratorMacro) \
 	ANIM_NEXT_IMPL_DEFINE_ANIM_TRAIT_GET_INTERFACES(TraitName, InterfaceEnumeratorMacro) \
+	ANIM_NEXT_IMPL_DEFINE_ANIM_TRAIT_GET_REQUIRED_INTERFACES(TraitName, RequiredInterfaceEnumeratorMacro) \
 	ANIM_NEXT_IMPL_DEFINE_ANIM_TRAIT_ON_TRAIT_EVENT(TraitName, EventEnumeratorMacro) \
 	ANIM_NEXT_IMPL_DEFINE_ANIM_TRAIT_GET_TRAIT_EVENTS(TraitName, EventEnumeratorMacro)
 
@@ -344,6 +359,9 @@ namespace UE::AnimNext
 		// Returns a list of interfaces that this trait supports
 		virtual TConstArrayView<FTraitInterfaceUID> GetTraitInterfaces() const { return TConstArrayView<FTraitInterfaceUID>(); }
 
+		// Returns a list of interfaces that this trait reqquires
+		virtual TConstArrayView<FTraitInterfaceUID> GetTraitRequiredInterfaces() const { return TConstArrayView<FTraitInterfaceUID>(); }
+
 		// Called when an event reaches an instance of this trait
 		virtual ETraitStackPropagation OnTraitEvent(FExecutionContext& Context, FTraitBinding& Binding, FAnimNextTraitEvent& Event) const { return ETraitStackPropagation::Continue; }
 
diff --git a/Engine/Plugins/Experimental/AnimNext/Source/AnimNext/Public/TraitCore/TraitInterfaceUID.h b/Engine/Plugins/Experimental/AnimNext/Source/AnimNext/Public/TraitCore/TraitInterfaceUID.h
--- a/Engine/Plugins/Experimental/AnimNext/Source/AnimNext/Public/TraitCore/TraitInterfaceUID.h
+++ b/Engine/Plugins/Experimental/AnimNext/Source/AnimNext/Public/TraitCore/TraitInterfaceUID.h
@@ -25,14 +25,16 @@ namespace UE::AnimNext
 			: UID(INVALID_UID)
 #if !UE_BUILD_SHIPPING
 			, InterfaceName(TEXT("<Invalid Interface UID>"))
+			, InterfaceShortName(TEXT("???"))
 #endif
 		{}
 
 		// Constructs an interface UID
-		explicit constexpr FTraitInterfaceUID(FTraitInterfaceUIDRaw InUID, const TCHAR* InInterfaceName = TEXT("<Unknown Trait Interface Name>"))
+		explicit constexpr FTraitInterfaceUID(FTraitInterfaceUIDRaw InUID, const TCHAR* InInterfaceName = TEXT("<Unknown Trait Interface Name>"), const TCHAR* InInterfaceShortName = TEXT("???"))
 			: UID(InUID)
 #if !UE_BUILD_SHIPPING
 			, InterfaceName(InInterfaceName)
+			, InterfaceShortName(InInterfaceShortName)
 #endif
 		{
 		}
@@ -40,6 +42,7 @@ namespace UE::AnimNext
 #if !UE_BUILD_SHIPPING
 		// Returns a literal string to the interface name
 		constexpr const TCHAR* GetInterfaceName() const { return InterfaceName; }
+		constexpr const TCHAR* GetInterfaceShortName() const { return InterfaceShortName; }
 #endif
 
 		// Returns the interface global UID
@@ -54,7 +57,8 @@ namespace UE::AnimNext
 		FTraitInterfaceUIDRaw		UID;
 
 #if !UE_BUILD_SHIPPING
-		const TCHAR*				InterfaceName;
+		const TCHAR*				InterfaceName = nullptr;
+		const TCHAR*				InterfaceShortName = nullptr;
 #endif
 	};
 
diff --git a/Engine/Plugins/Experimental/AnimNext/Source/AnimNext/Public/TraitCore/TraitMode.h b/Engine/Plugins/Experimental/AnimNext/Source/AnimNext/Public/TraitCore/TraitMode.h
--- a/Engine/Plugins/Experimental/AnimNext/Source/AnimNext/Public/TraitCore/TraitMode.h
+++ b/Engine/Plugins/Experimental/AnimNext/Source/AnimNext/Public/TraitCore/TraitMode.h
@@ -13,6 +13,11 @@ namespace UE::AnimNext
 	 */
 	enum class ETraitMode
 	{
+		/**
+		 * Invalid value is used for default variable initialization
+		 */
+		Invalid = 0,
+
 		/**
 		 * Base traits can live on their own in an animation node and have no 'Super'.
 		 * As a result, calls to GetInterface() do not forward to other traits below
@@ -28,5 +33,10 @@ namespace UE::AnimNext
 		 * until a base trait is found.
 		 */
 		Additive,
+
+		/**
+		 * Max value is used for iteration / static checks
+		 */
+		Num = Additive
 	};
 }
diff --git a/Engine/Plugins/Experimental/AnimNext/Source/AnimNext/Public/TraitInterfaces/IGarbageCollection.h b/Engine/Plugins/Experimental/AnimNext/Source/AnimNext/Public/TraitInterfaces/IGarbageCollection.h
--- a/Engine/Plugins/Experimental/AnimNext/Source/AnimNext/Public/TraitInterfaces/IGarbageCollection.h
+++ b/Engine/Plugins/Experimental/AnimNext/Source/AnimNext/Public/TraitInterfaces/IGarbageCollection.h
@@ -17,7 +17,7 @@ namespace UE::AnimNext
 	 */
 	struct ANIMNEXT_API IGarbageCollection : ITraitInterface
 	{
-		DECLARE_ANIM_TRAIT_INTERFACE(IGarbageCollection, 0x231a2017)
+		DECLARE_ANIM_TRAIT_INTERFACE(IGarbageCollection, GC, 0x231a2017)
 
 		// Registers the provided binding for GC callback
 		// Once registered, AddReferencedObjects is called during GC to collect references
diff --git a/Engine/Plugins/Experimental/AnimNext/Source/AnimNext/Public/TraitInterfaces/IHierarchy.h b/Engine/Plugins/Experimental/AnimNext/Source/AnimNext/Public/TraitInterfaces/IHierarchy.h
--- a/Engine/Plugins/Experimental/AnimNext/Source/AnimNext/Public/TraitInterfaces/IHierarchy.h
+++ b/Engine/Plugins/Experimental/AnimNext/Source/AnimNext/Public/TraitInterfaces/IHierarchy.h
@@ -20,7 +20,7 @@ namespace UE::AnimNext
 	 */
 	struct ANIMNEXT_API IHierarchy : ITraitInterface
 	{
-		DECLARE_ANIM_TRAIT_INTERFACE(IHierarchy, 0x846d8a37)
+		DECLARE_ANIM_TRAIT_INTERFACE(IHierarchy, HIE, 0x846d8a37)
 
 		// Returns the number of children of the trait implementation (not the whole stack)
 		// Includes inactive children
diff --git a/Engine/Plugins/Experimental/AnimNext/Source/AnimNextEditor/Private/AnimNextEditorModule.cpp b/Engine/Plugins/Experimental/AnimNext/Source/AnimNextEditor/Private/AnimNextEditorModule.cpp
--- a/Engine/Plugins/Experimental/AnimNext/Source/AnimNextEditor/Private/AnimNextEditorModule.cpp
+++ b/Engine/Plugins/Experimental/AnimNext/Source/AnimNextEditor/Private/AnimNextEditorModule.cpp
@@ -17,6 +17,9 @@
 #include "Graph/AnimNextGraphPanelNodeFactory.h"
 #include "Graph/AnimNextGraph_EdGraphNodeCustomization.h"
 #include "Graph/AnimNextGraph_EditorData.h"
+#include "Graph/AnimNextGraph_EdGraphNode.h"
+#include "Graph/TraitEditorTabSummoner.h"
+#include "Graph/AnimNextCompilerResultsTabSummoner.h"
 #include "HAL/PlatformApplicationMisc.h"
 #include "Param/ParameterCustomization.h"
 #include "Param/ParameterPickerArgs.h"
@@ -42,6 +45,7 @@
 #include "AnimNextRigVMWorkspaceAssetUserData.h"
 #include "Graph/AnimNextGraph_OutlinerItemDetails.h"
 #include "Param/AnimNextActorLocatorEditor.h"
+#include "IWorkspaceEditor.h"
 
 #define LOCTEXT_NAMESPACE "AnimNextEditorModule"
 
@@ -81,9 +85,6 @@ void FModule::StartupModule()
 	PropertyModule.RegisterCustomClassLayout("AnimNextGraph_Parameter", 
 		FOnGetDetailCustomizationInstance::CreateLambda([] { return MakeShared<FParameterCustomization>(); }));
 
-	PropertyModule.RegisterCustomClassLayout("AnimNextGraph_EdGraphNode",
-		FOnGetDetailCustomizationInstance::CreateLambda([] { return MakeShared<FAnimNextGraph_EdGraphNodeCustomization>(); }));
-
 	AnimNextGraphPanelNodeFactory = MakeShared<FAnimNextGraphPanelNodeFactory>();
 	FEdGraphUtilities::RegisterVisualNodeFactory(AnimNextGraphPanelNodeFactory);
 
@@ -92,10 +93,27 @@ void FModule::StartupModule()
 
 	Workspace::IWorkspaceEditorModule& WorkspaceEditorModule = FModuleManager::Get().LoadModuleChecked<Workspace::IWorkspaceEditorModule>("WorkspaceEditor");
 
+	WorkspaceEditorModule.OnRegisterTabsForEditor().AddLambda([](FWorkflowAllowedTabSet& TabFactories, const TSharedRef<FTabManager>& InTabManager, TSharedPtr<UE::Workspace::IWorkspaceEditor> EditorPtr)
+		{
+			TabFactories.RegisterFactory(MakeShared<UE::AnimNext::Editor::FTraitEditorTabSummoner>(EditorPtr));
+
+			TabFactories.RegisterFactory(MakeShared<UE::AnimNext::Editor::FAnimNextCompilerResultsTabSummoner>(EditorPtr));
+
+			for (auto FactoryIt = TabFactories.CreateIterator(); FactoryIt; ++FactoryIt)
+			{
+				FactoryIt.Value()->RegisterTabSpawner(InTabManager, nullptr);
+			}
+		});
+
 	RegisterWorkspaceDocumentTypes(WorkspaceEditorModule);
 
-	WorkspaceEditorModule.OnRegisterWorkspaceDetailsCustomization().AddLambda([](TSharedPtr<IDetailsView>& InDetailsView)
+	WorkspaceEditorModule.OnRegisterWorkspaceDetailsCustomization().AddLambda([](const TWeakPtr<Workspace::IWorkspaceEditor>& InWorkspaceEditor, TSharedPtr<IDetailsView>& InDetailsView)
 		{
+			InDetailsView->RegisterInstancedCustomPropertyLayout(UAnimNextGraph_EdGraphNode::StaticClass(), FOnGetDetailCustomizationInstance::CreateLambda([InWorkspaceEditor]()
+				{
+					return MakeShared<FAnimNextGraph_EdGraphNodeCustomization>(InWorkspaceEditor);
+				}));
+
 			TArray<UScriptStruct*> StructsToCustomize = {
 				TBaseStructure<FVector>::Get(),
 				TBaseStructure<FVector2D>::Get(),
@@ -105,7 +123,6 @@ void FModule::StartupModule()
 				TBaseStructure<FTransform>::Get(),
 				TBaseStructure<FEulerTransform>::Get(),
 			};
-
 			for (UScriptStruct* StructToCustomize : StructsToCustomize)
 			{
 				InDetailsView->RegisterInstancedCustomPropertyTypeLayout(StructToCustomize->GetFName(), FOnGetPropertyTypeCustomizationInstance::CreateLambda([]()
@@ -318,18 +335,6 @@ void FModule::RegisterWorkspaceDocumentTypes(Workspace::IWorkspaceEditorModule&
 
 				TWeakPtr<Workspace::IWorkspaceEditor> WeakWorkspaceEditor = InContext.WorkspaceEditor;
 
-				EditorData->RigVMGraphModifiedEvent.RemoveAll(&InContext.WorkspaceEditor.Get());
-				EditorData->RigVMGraphModifiedEvent.AddSPLambda(&InContext.WorkspaceEditor.Get(), [WeakWorkspaceEditor](ERigVMGraphNotifType InType, URigVMGraph* InGraph, UObject* InSubject)
-				{
-					if(TSharedPtr<Workspace::IWorkspaceEditor> WorkspaceEditor = WeakWorkspaceEditor.Pin())
-					{
-						if (InType == ERigVMGraphNotifType::InteractionBracketClosed)
-						{
-							WorkspaceEditor->RefreshDetails();
-						}
-					}
-				});
-
 				return SNew(SRigVMAssetView, EditorData)
 					.OnSelectionChanged_Lambda([WeakWorkspaceEditor](const TArray<UObject*>& InEntries)
 					{
@@ -707,6 +712,34 @@ void FModule::RegisterWorkspaceDocumentTypes(Workspace::IWorkspaceEditorModule&
 	});
 
 	Workspace::FObjectDocumentArgs GraphDocumentArgs = WorkspaceEditorModule.CreateGraphDocumentArgs(GraphArgs);
+	Workspace::FOnMakeDocumentWidget WorkspaceMakeDocumentWidgetDelegate = GraphDocumentArgs.OnMakeDocumentWidget;
+	GraphDocumentArgs.OnMakeDocumentWidget = Workspace::FOnMakeDocumentWidget::CreateLambda([WorkspaceMakeDocumentWidgetDelegate](const Workspace::FWorkspaceEditorContext& InContext)
+	{
+		TWeakPtr<Workspace::IWorkspaceEditor> WeakWorkspaceEditor = InContext.WorkspaceEditor;
+
+		if (UAnimNextGraph_EdGraph* EdGraph = Cast<UAnimNextGraph_EdGraph>(InContext.Object))
+		{
+			UAnimNextRigVMAssetEditorData* EditorData = EdGraph->GetTypedOuter<UAnimNextRigVMAssetEditorData>();
+			check(EditorData);
+
+			EditorData->InteractionBracketFinished.RemoveAll(&InContext.WorkspaceEditor.Get());
+			EditorData->InteractionBracketFinished.AddSPLambda(&InContext.WorkspaceEditor.Get(), [WeakWorkspaceEditor](UAnimNextRigVMAssetEditorData* InEditorData)
+				{
+					if (TSharedPtr<Workspace::IWorkspaceEditor> WorkspaceEditor = WeakWorkspaceEditor.Pin())
+					{
+						WorkspaceEditor->RefreshDetails();
+					}
+				});
+		}
+
+		if (WorkspaceMakeDocumentWidgetDelegate.IsBound())
+		{
+			return WorkspaceMakeDocumentWidgetDelegate.Execute(InContext);
+		}
+
+		return SNullWidget::NullWidget;
+	});
+
 	GraphDocumentArgs.OnGetDocumentBreadcrumbTrail = Workspace::FOnGetDocumentBreadcrumbTrail::CreateLambda([](const Workspace::FWorkspaceEditorContext& InContext, TArray<TSharedPtr<Workspace::FWorkspaceBreadcrumb>>& OutBreadcrumbs)
 	{
 		if (const URigVMEdGraph* RigVMEdGraph = Cast<URigVMEdGraph>(InContext.Object))
diff --git a/Engine/Plugins/Experimental/AnimNext/Source/AnimNextEditor/Private/Common/SCategoryTableRow.h b/Engine/Plugins/Experimental/AnimNext/Source/AnimNextEditor/Private/Common/SCategoryTableRow.h
new file mode 100644
--- /dev/null
+++ b/Engine/Plugins/Experimental/AnimNext/Source/AnimNextEditor/Private/Common/SCategoryTableRow.h
@@ -0,0 +1,104 @@
+// Copyright Epic Games, Inc. All Rights Reserved.
+
+#pragma once
+
+#include "Layout/Margin.h"
+#include "Styling/AppStyle.h"
+#include "Widgets/SBoxPanel.h"
+#include "Widgets/Layout/SBorder.h"
+#include "Widgets/Views/SExpanderArrow.h"
+#include "Widgets/Views/STableRow.h"
+
+namespace UE::AnimNext::Editor
+{
+
+// An entry category displayed in a tree view
+template<typename ItemType>
+class SCategoryHeaderTableRow : public STableRow<ItemType>
+{
+public:
+	SLATE_BEGIN_ARGS(SCategoryHeaderTableRow)
+		: _Padding(FMargin(0.0f, 2.0f, .0f, 0.0f))
+		{}
+		SLATE_DEFAULT_SLOT(typename SCategoryHeaderTableRow::FArguments, Content)
+		SLATE_ATTRIBUTE(FMargin, Padding)
+	SLATE_END_ARGS()
+
+	void Construct(const FArguments& InArgs, const TSharedRef<STableViewBase>& InOwnerTableView)
+	{
+		STableRow<ItemType>::ChildSlot
+			.Padding(InArgs._Padding)
+			[
+				SAssignNew(ContentBorder, SBorder)
+					.BorderImage(this, &SCategoryHeaderTableRow::GetBackgroundImage)
+					.Padding(FMargin(3.0f, 5.0f))
+					[
+						SNew(SHorizontalBox)
+							+ SHorizontalBox::Slot()
+							.VAlign(VAlign_Center)
+							.Padding(5.0f)
+							.AutoWidth()
+							[
+								SNew(SExpanderArrow, STableRow< ItemType >::SharedThis(this))
+							]
+							+ SHorizontalBox::Slot()
+							.VAlign(VAlign_Center)
+							.AutoWidth()
+							[
+								InArgs._Content.Widget
+							]
+					]
+			];
+
+		STableRow < ItemType >::ConstructInternal(
+			typename STableRow< ItemType >::FArguments()
+			.Style(FAppStyle::Get(), "DetailsView.TreeView.TableRow")
+			.ShowSelection(false),
+			InOwnerTableView
+		);
+	}
+
+	const FSlateBrush* GetBackgroundImage() const
+	{
+		if (STableRow<ItemType>::IsHovered())
+		{
+			return FAppStyle::Get().GetBrush("Brushes.Secondary");
+		}
+		else
+		{
+			return FAppStyle::Get().GetBrush("Brushes.Header");
+		}
+	}
+
+	virtual void SetContent(TSharedRef< SWidget > InContent) override
+	{
+		ContentBorder->SetContent(InContent);
+	}
+
+	virtual void SetRowContent(TSharedRef< SWidget > InContent) override
+	{
+		ContentBorder->SetContent(InContent);
+	}
+
+	virtual const FSlateBrush* GetBorder() const
+	{
+		return nullptr;
+	}
+
+	FReply OnMouseButtonDown(const FGeometry& MyGeometry, const FPointerEvent& MouseEvent)
+	{
+		if (MouseEvent.GetEffectingButton() == EKeys::LeftMouseButton)
+		{
+			STableRow<ItemType>::ToggleExpansion();
+			return FReply::Handled();
+		}
+		else
+		{
+			return FReply::Unhandled();
+		}
+	}
+private:
+	TSharedPtr<SBorder> ContentBorder;
+};
+
+} // end namespace UE::AnimNext::Editor
diff --git a/Engine/Plugins/Experimental/AnimNext/Source/AnimNextEditor/Private/Common/SRigVMAssetView.cpp b/Engine/Plugins/Experimental/AnimNext/Source/AnimNextEditor/Private/Common/SRigVMAssetView.cpp
--- a/Engine/Plugins/Experimental/AnimNext/Source/AnimNextEditor/Private/Common/SRigVMAssetView.cpp
+++ b/Engine/Plugins/Experimental/AnimNext/Source/AnimNextEditor/Private/Common/SRigVMAssetView.cpp
@@ -5,6 +5,7 @@
 #include "AnimNextRigVMAsset.h"
 #include "AnimNextRigVMAssetEditorData.h"
 #include "AnimNextRigVMAssetEntry.h"
+#include "Common/SCategoryTableRow.h"
 #include "Widgets/Layout/SBox.h"
 #include "Widgets/Input/SSearchBox.h"
 #include "DetailLayoutBuilder.h"
@@ -52,93 +53,6 @@ static const FName NAME_Category("Category");
 
 TMap<FName, SRigVMAssetView::FCategoryWidgetFactoryFunction> SRigVMAssetView::CategoryFactories;
 
-// An entry category displayed in the parameters table
-template<typename ItemType>
-class SCategoryHeaderTableRow : public STableRow<ItemType>
-{
-public:
-	SLATE_BEGIN_ARGS(SCategoryHeaderTableRow)
-		{}
-		SLATE_DEFAULT_SLOT(typename SCategoryHeaderTableRow::FArguments, Content)
-	SLATE_END_ARGS()
-
-	void Construct(const FArguments& InArgs, const TSharedRef<STableViewBase>& InOwnerTableView)
-	{
-		STableRow<ItemType>::ChildSlot
-			.Padding(0.0f, 2.0f, .0f, 0.0f)
-			[
-				SAssignNew(ContentBorder, SBorder)
-					.BorderImage(this, &SCategoryHeaderTableRow::GetBackgroundImage)
-					.Padding(FMargin(3.0f, 5.0f))
-					[
-						SNew(SHorizontalBox)
-							+ SHorizontalBox::Slot()
-							.VAlign(VAlign_Center)
-							.Padding(5.0f)
-							.AutoWidth()
-							[
-								SNew(SExpanderArrow, STableRow< ItemType >::SharedThis(this))
-							]
-							+ SHorizontalBox::Slot()
-							.VAlign(VAlign_Center)
-							.AutoWidth()
-							[
-								InArgs._Content.Widget
-							]
-					]
-			];
-
-		STableRow < ItemType >::ConstructInternal(
-			typename STableRow< ItemType >::FArguments()
-			.Style(FAppStyle::Get(), "DetailsView.TreeView.TableRow")
-			.ShowSelection(false),
-			InOwnerTableView
-		);
-	}
-
-	const FSlateBrush* GetBackgroundImage() const
-	{
-		if (STableRow<ItemType>::IsHovered())
-		{
-			return FAppStyle::Get().GetBrush("Brushes.Secondary");
-		}
-		else
-		{
-			return FAppStyle::Get().GetBrush("Brushes.Header");
-		}
-	}
-
-	virtual void SetContent(TSharedRef< SWidget > InContent) override
-	{
-		ContentBorder->SetContent(InContent);
-	}
-
-	virtual void SetRowContent(TSharedRef< SWidget > InContent) override
-	{
-		ContentBorder->SetContent(InContent);
-	}
-
-	virtual const FSlateBrush* GetBorder() const
-	{
-		return nullptr;
-	}
-
-	FReply OnMouseButtonDown(const FGeometry& MyGeometry, const FPointerEvent& MouseEvent)
-	{
-		if (MouseEvent.GetEffectingButton() == EKeys::LeftMouseButton)
-		{
-			STableRow<ItemType>::ToggleExpansion();
-			return FReply::Handled();
-		}
-		else
-		{
-			return FReply::Unhandled();
-		}
-	}
-private:
-	TSharedPtr<SBorder> ContentBorder;
-};
-
 // An entry displayed in the parameters view
 struct FRigVMAssetViewEntry
 {
diff --git a/Engine/Plugins/Experimental/AnimNext/Source/AnimNextEditor/Private/Graph/AnimNextCompilerResultsTabSummoner.cpp b/Engine/Plugins/Experimental/AnimNext/Source/AnimNextEditor/Private/Graph/AnimNextCompilerResultsTabSummoner.cpp
new file mode 100644
--- /dev/null
+++ b/Engine/Plugins/Experimental/AnimNext/Source/AnimNextEditor/Private/Graph/AnimNextCompilerResultsTabSummoner.cpp
@@ -0,0 +1,81 @@
+// Copyright Epic Games, Inc. All Rights Reserved.
+
+#include "Graph/AnimNextCompilerResultsTabSummoner.h"
+#include "IWorkspaceEditor.h"
+#include "MessageLogModule.h"
+
+#define LOCTEXT_NAMESPACE "WorkspaceTabSummoner"
+
+namespace UE::AnimNext::Editor
+{
+
+const FName CompilerResultsTabName("CompilerResultsTab");
+
+// ***************************************************************************
+
+void SAnimNextCompilerResultsWidget::Construct(const FArguments& InArgs, const TWeakPtr<UE::Workspace::IWorkspaceEditor>& InWorkspaceEditorWeak)
+{
+	CreateMessageLog(InWorkspaceEditorWeak);
+
+	ChildSlot
+	[
+		SNew(SVerticalBox)
+
+		+SVerticalBox::Slot()
+		.FillHeight(1.0)
+		.Padding(10.f, 10.f, 10.f, 10.f)
+		[
+			CompilerResults.ToSharedRef()
+		]
+	];
+}
+
+void SAnimNextCompilerResultsWidget::CreateMessageLog(const TWeakPtr<UE::Workspace::IWorkspaceEditor>& InWorkspaceEditorWeak)
+{
+	if (const TSharedPtr<UE::Workspace::IWorkspaceEditor> WorkspaceEditorShared = InWorkspaceEditorWeak.Pin())
+	{
+		FMessageLogModule& MessageLogModule = FModuleManager::LoadModuleChecked<FMessageLogModule>("MessageLog");
+		const FName LogName = *FString::Printf(TEXT("%s_CompilerResultsLog"), WorkspaceEditorShared.IsValid() ? *WorkspaceEditorShared->GetEditorName().ToString() : TEXT("TraitEditor"));
+		// Reuse any existing log, or create a new one (that is not held onto bey the message log system)
+		if (MessageLogModule.IsRegisteredLogListing(LogName))
+		{
+			CompilerResultsListing = MessageLogModule.GetLogListing(LogName);
+		}
+		else
+		{
+			FMessageLogInitializationOptions LogInitOptions;
+			LogInitOptions.bShowInLogWindow = false;
+			CompilerResultsListing = MessageLogModule.CreateLogListing(LogName, LogInitOptions);
+		}
+
+		CompilerResults = MessageLogModule.CreateLogListingWidget(CompilerResultsListing.ToSharedRef());
+	}
+}
+
+// ***************************************************************************
+
+FAnimNextCompilerResultsTabSummoner::FAnimNextCompilerResultsTabSummoner(TSharedPtr<UE::Workspace::IWorkspaceEditor> InHostingApp)
+	: FWorkflowTabFactory(CompilerResultsTabName, StaticCastSharedPtr<FAssetEditorToolkit>(InHostingApp))
+{
+	TabLabel = LOCTEXT("AnimNExtCompilerResultsTabLabel", "Compiler Results");
+	TabIcon = FSlateIcon("EditorStyle", "LevelEditor.Tabs.Outliner");
+	ViewMenuDescription = LOCTEXT("AnimNExtCompilerResultsTabMenuDescription", "Compiler Results");
+	ViewMenuTooltip = LOCTEXT("AnimNExtCompilerResultsTabToolTip", "Shows the Compiler Results tab.");
+	bIsSingleton = true;
+
+	AnimNextCompilerResultsWidget = SNew(SAnimNextCompilerResultsWidget, InHostingApp.ToWeakPtr());
+}
+
+TSharedRef<SWidget> FAnimNextCompilerResultsTabSummoner::CreateTabBody(const FWorkflowTabSpawnInfo& Info) const
+{
+	return AnimNextCompilerResultsWidget.ToSharedRef();
+}
+
+FText FAnimNextCompilerResultsTabSummoner::GetTabToolTipText(const FWorkflowTabSpawnInfo& Info) const
+{
+	return ViewMenuTooltip;
+}
+
+} // end namespace UE::AnimNext::Editor
+
+#undef LOCTEXT_NAMESPACE
diff --git a/Engine/Plugins/Experimental/AnimNext/Source/AnimNextEditor/Private/Graph/AnimNextCompilerResultsTabSummoner.h b/Engine/Plugins/Experimental/AnimNext/Source/AnimNextEditor/Private/Graph/AnimNextCompilerResultsTabSummoner.h
new file mode 100644
--- /dev/null
+++ b/Engine/Plugins/Experimental/AnimNext/Source/AnimNextEditor/Private/Graph/AnimNextCompilerResultsTabSummoner.h
@@ -0,0 +1,56 @@
+// Copyright Epic Games, Inc. All Rights Reserved.
+
+#pragma once
+
+#include "WorkflowOrientedApp/WorkflowTabFactory.h"
+#include "IMessageLogListing.h"
+
+namespace UE::Workspace
+{
+	class IWorkspaceEditor;
+}
+
+namespace UE::AnimNext::Editor
+{
+
+extern const FName CompilerResultsTabName;
+
+class SAnimNextCompilerResultsWidget : public SCompoundWidget
+{
+public:
+	SLATE_BEGIN_ARGS(SAnimNextCompilerResultsWidget) {}
+
+	SLATE_END_ARGS()
+
+	void Construct(const FArguments& InArgs, const TWeakPtr<UE::Workspace::IWorkspaceEditor>& InWorkspaceEditorWeak);
+
+	inline TSharedPtr<class IMessageLogListing> GetCompilerResultsListing() const
+	{
+		return CompilerResultsListing;
+	}
+
+private:
+	void CreateMessageLog(const TWeakPtr<UE::Workspace::IWorkspaceEditor>& InWorkspaceEditorWeak);
+
+	TSharedPtr<class IMessageLogListing> CompilerResultsListing;
+	TSharedPtr<class SWidget> CompilerResults;
+};
+
+
+
+
+struct FAnimNextCompilerResultsTabSummoner : public FWorkflowTabFactory
+{
+public:
+	FAnimNextCompilerResultsTabSummoner(TSharedPtr<UE::Workspace::IWorkspaceEditor> InHostingApp);
+
+private:
+	// FWorkflowTabFactory interface
+	virtual TSharedRef<SWidget> CreateTabBody(const FWorkflowTabSpawnInfo& Info) const override;
+	virtual FText GetTabToolTipText(const FWorkflowTabSpawnInfo& Info) const override;
+
+	// The widget this tab spawner wraps
+	TSharedPtr<class SAnimNextCompilerResultsWidget> AnimNextCompilerResultsWidget;
+};
+
+} // end namespace UE::AnimNext::Editor
diff --git a/Engine/Plugins/Experimental/AnimNext/Source/AnimNextEditor/Private/Graph/AnimNextGraph_EdGraphNodeCustomization.cpp b/Engine/Plugins/Experimental/AnimNext/Source/AnimNextEditor/Private/Graph/AnimNextGraph_EdGraphNodeCustomization.cpp
--- a/Engine/Plugins/Experimental/AnimNext/Source/AnimNextEditor/Private/Graph/AnimNextGraph_EdGraphNodeCustomization.cpp
+++ b/Engine/Plugins/Experimental/AnimNext/Source/AnimNextEditor/Private/Graph/AnimNextGraph_EdGraphNodeCustomization.cpp
@@ -6,15 +6,40 @@
 #include "DetailCategoryBuilder.h"
 #include "Graph/AnimNextGraph_EdGraphNode.h"
 #include "Graph/RigDecorator_AnimNextCppTrait.h"
+#include "Graph/TraitEditorTabSummoner.h"
 #include "Logging/LogScopedVerbosityOverride.h"
 #include "InstancedPropertyBagStructureDataProvider.h"
 #include "RigVMModel/RigVMController.h"
+#include "IWorkspaceEditor.h"
+#include "STraitEditorView.h"
+#include "Widgets/Docking/SDockTab.h"
 
 #define LOCTEXT_NAMESPACE "EdGraphNodeCustomization"
 
 namespace UE::AnimNext::Editor
 {
 
+FAnimNextGraph_EdGraphNodeCustomization::FAnimNextGraph_EdGraphNodeCustomization(const TWeakPtr<UE::Workspace::IWorkspaceEditor>& InWorkspaceEditorWeak)
+	: WorkspaceEditorWeak(InWorkspaceEditorWeak)
+{
+}
+
+void FAnimNextGraph_EdGraphNodeCustomization::PendingDelete()
+{
+	CategoryDetailsData.Reset();
+
+	if (TSharedPtr<UE::Workspace::IWorkspaceEditor> WorkspaceEditor = WorkspaceEditorWeak.Pin())
+	{
+		if (TSharedPtr<SDockTab> DockTab = WorkspaceEditor->GetTabManager()->FindExistingLiveTab(UE::AnimNext::Editor::TraitEditorTabName))
+		{
+			if (TSharedPtr<STraitEditorView> TraitEditorView = StaticCastSharedPtr<STraitEditorView>(DockTab->GetContent().ToSharedPtr()))
+			{
+				TraitEditorView->SetTraitData(FTraitStackData(nullptr));
+			}
+		}
+	}
+}
+
 void FAnimNextGraph_EdGraphNodeCustomization::CustomizeDetails(IDetailLayoutBuilder& DetailBuilder)
 {
 	TArray<TWeakObjectPtr<UObject>> Objects;
@@ -35,18 +60,13 @@ void FAnimNextGraph_EdGraphNodeCustomization::CustomizeObjects(IDetailLayoutBuil
 	{
 		if (UAnimNextGraph_EdGraphNode* EdGraphNode = Cast<UAnimNextGraph_EdGraphNode>(Objects[i].Get()))
 		{
-			if (URigVMNode* ModelNode = EdGraphNode->GetModelNode())
+			if (EdGraphNode->IsTraitStack())
 			{
-				// If there are any traits in the node, treat it as a trait stack
-				const TArray<URigVMPin*> TraitPins = ModelNode->GetDecoratorPins();
-				if (TraitPins.Num() > 0)
-				{
-					GenerateTraitData(EdGraphNode, CategoryDetailsData);
-				}
-				else
-				{
-					GenerateRigVMData(EdGraphNode, CategoryDetailsData);
-				}
+				GenerateTraitData(EdGraphNode, CategoryDetailsData);
+			}
+			else
+			{
+				GenerateRigVMData(EdGraphNode, CategoryDetailsData);
 			}
 		}
 	}
@@ -55,13 +75,31 @@ void FAnimNextGraph_EdGraphNodeCustomization::CustomizeObjects(IDetailLayoutBuil
 	{
 		PopulateCategory(DetailBuilder, DetailsData);
 	}
+
+	// Pass the TraitStack EdGraphNode to the TraitEditor (or nullptr if we select other node type, like RigVM)
+	if (Objects.Num() == 1)
+	{
+		if (UAnimNextGraph_EdGraphNode* EdGraphNode = Cast<UAnimNextGraph_EdGraphNode>(Objects[0].Get()))
+		{
+			if (TSharedPtr<UE::Workspace::IWorkspaceEditor> WorkspaceEditor = WorkspaceEditorWeak.Pin())
+			{
+				if (TSharedPtr<SDockTab> DockTab = WorkspaceEditor->GetTabManager()->FindExistingLiveTab(UE::AnimNext::Editor::TraitEditorTabName))
+				{
+					if (TSharedPtr<STraitEditorView> TraitEditorView = StaticCastSharedPtr<STraitEditorView>(DockTab->GetContent().ToSharedPtr()))
+					{
+						TraitEditorView->SetTraitData(EdGraphNode->IsTraitStack() ? FTraitStackData(EdGraphNode) : FTraitStackData(nullptr));
+					}
+				}
+			}
+		}
+	}
 }
 
 void FAnimNextGraph_EdGraphNodeCustomization::GenerateTraitData(UAnimNextGraph_EdGraphNode* EdGraphNode, TArray<TSharedPtr<FCategoryDetailsData>>& CategoryDetailsData)
 {
 	if (URigVMNode* ModelNode = EdGraphNode->GetModelNode())
 	{
-		// If there are any traits in the node, treat it as a trait stack
+		// Obtain the pins from the stack
 		const TArray<URigVMPin*> TraitPins = ModelNode->GetDecoratorPins();
 		if (TraitPins.Num() > 0)
 		{
diff --git a/Engine/Plugins/Experimental/AnimNext/Source/AnimNextEditor/Private/Graph/AnimNextGraph_EdGraphNodeCustomization.h b/Engine/Plugins/Experimental/AnimNext/Source/AnimNextEditor/Private/Graph/AnimNextGraph_EdGraphNodeCustomization.h
--- a/Engine/Plugins/Experimental/AnimNext/Source/AnimNextEditor/Private/Graph/AnimNextGraph_EdGraphNodeCustomization.h
+++ b/Engine/Plugins/Experimental/AnimNext/Source/AnimNextEditor/Private/Graph/AnimNextGraph_EdGraphNodeCustomization.h
@@ -11,16 +11,28 @@ class URigVMPin;
 class IDetailCategoryBuilder;
 class UAnimNextGraph_EdGraphNode;
 
+namespace UE::Workspace
+{
+	class IWorkspaceEditor;
+}
+
 namespace UE::AnimNext::Editor
 {
 
 class FAnimNextGraph_EdGraphNodeCustomization : public IDetailCustomization
 {
+public:
+	FAnimNextGraph_EdGraphNodeCustomization() = default;
+	explicit FAnimNextGraph_EdGraphNodeCustomization(const TWeakPtr<UE::Workspace::IWorkspaceEditor>& InWorkspaceEditorWeak);
+
 protected:
 
 	// --- IDetailCustomization Begin ---
 	/** Called when details should be customized */
 	virtual void CustomizeDetails(IDetailLayoutBuilder& DetailBuilder) override;
+	/** Called when no longer used and will be deleted */
+	virtual void PendingDelete() override;
+
 	// --- IDetailCustomization End ---
 
 	struct FCategoryDetailsData
@@ -91,6 +103,7 @@ protected:
 
 	static void GenerateMemoryStorage(const TArray<TWeakObjectPtr<URigVMPin>> & ModelPinsToDisplay, FRigVMMemoryStorageStruct& MemoryStorage);
 
+	TWeakPtr<UE::Workspace::IWorkspaceEditor> WorkspaceEditorWeak;
 	TArray<TSharedPtr<FCategoryDetailsData>> CategoryDetailsData;
 };
 
diff --git a/Engine/Plugins/Experimental/AnimNext/Source/AnimNextEditor/Private/Graph/STraitEditorView.cpp b/Engine/Plugins/Experimental/AnimNext/Source/AnimNextEditor/Private/Graph/STraitEditorView.cpp
new file mode 100644
--- /dev/null
+++ b/Engine/Plugins/Experimental/AnimNext/Source/AnimNextEditor/Private/Graph/STraitEditorView.cpp
@@ -0,0 +1,679 @@
+// Copyright Epic Games, Inc. All Rights Reserved.
+
+#include "STraitEditorView.h"
+#include "ContentBrowserModule.h"
+#include "IContentBrowserSingleton.h"
+#include "SAssetDropTarget.h"
+#include "ScopedTransaction.h"
+#include "WorkspaceSchema.h"
+#include "Framework/Commands/GenericCommands.h"
+#include "TraitCore/TraitRegistry.h"
+#include "ObjectEditorUtils.h"
+#include "Graph/AnimNextGraph_EdGraphNode.h"
+#include "STraitListView.h"
+#include "STraitStackView.h"
+#include "MessageLogModule.h"
+#include "IWorkspaceEditor.h"
+#include "IMessageLogListing.h"
+#include "Graph/AnimNextGraph_Controller.h"
+#include "Graph/AnimNextCompilerResultsTabSummoner.h"
+#include "Widgets/Docking/SDockTab.h"
+#include "Widgets/Input/SComboButton.h"
+#include "Widgets/Input/SSearchBox.h"
+#include "Widgets/Text/SInlineEditableTextBlock.h"
+
+#define LOCTEXT_NAMESPACE "TraitListEditor"
+
+namespace UE::AnimNext::Editor
+{
+
+STraitEditorView::STraitEditorView()
+	: TraitEditorSharedData(MakeShared<FTraitEditorSharedData>())
+{
+	TraitEditorSharedData->CurrentTraitsDataShared = MakeShared<TArray<TSharedPtr<FTraitDataEditorDef>>>();
+}
+
+void STraitEditorView::Construct(const FArguments& InArgs, TWeakPtr<UE::Workspace::IWorkspaceEditor> InWorkspaceEditorWeak)
+{
+	WorkspaceEditorWeak = MoveTemp(InWorkspaceEditorWeak);
+
+	TWeakObjectPtr<UAnimNextGraph_EdGraphNode> EdGraphNodeWeak = TraitEditorSharedData->EdGraphNodeWeak;
+
+	ChildSlot
+	[
+		SNew(SVerticalBox)
+
+		+ SVerticalBox::Slot()
+		.FillHeight(1.0)
+		[
+			SNew(SBorder)
+			.Padding(4.0f)
+			.BorderImage(FAppStyle::GetBrush("ToolPanel.GroupBorder"))
+			[
+				SNew(SVerticalBox)
+
+				+ SVerticalBox::Slot()
+				.AutoHeight()
+				[
+					GetOptionsMenuWidget()
+				]
+
+				+ SVerticalBox::Slot()
+				.FillHeight(1.0)
+				.Padding(0.f, 10.f, 0.f, 0.f)
+				[
+					SNew(SHorizontalBox)
+					+ SHorizontalBox::Slot()
+					.FillWidth(1.0)
+					[
+						SAssignNew(TraitListWidget, STraitListView, TraitEditorSharedData)
+						.OnTraitClicked(STraitListView::FOnTraitClicked::CreateSP(this, &STraitEditorView::OnTraitClicked))
+						.OnGetSelectedTraitData(STraitListView::FOnGetSelectedTraitData::CreateSP(this, &STraitEditorView::OnGetSelectedTraitData))
+					]
+					+ SHorizontalBox::Slot()
+					.FillWidth(1.0)
+					[
+						SAssignNew(TraitStackWidget, STraitStackView, TraitEditorSharedData)
+						.OnTraitDeleteRequest(STraitStackView::FOnStackTraitDeleteRequest::CreateSP(this, &STraitEditorView::OnTraitDeleteRequest))
+						.OnStatckTraitSelectionChanged(STraitStackView::FOnStatckTraitSelectionChanged::CreateSP(this, &STraitEditorView::OnStatckTraitSelectionChanged))
+						.OnStackTraitDragAccepted(STraitStackView::FOnStackTraitDragAccepted::CreateSP(this, &STraitEditorView::OnStackTraitDragAccepted))
+					]
+				]
+			]
+		]
+	];
+}
+
+void STraitEditorView::SetTraitData(const FTraitStackData& InTraitStackData)
+{
+ 	TraitEditorSharedData->EdGraphNodeWeak = InTraitStackData.EdGraphNodeWeak;
+
+	GenerateTraitStackData(TraitEditorSharedData->EdGraphNodeWeak, TraitEditorSharedData);
+
+	Refresh();
+}
+
+FReply STraitEditorView::OnTraitClicked(const FTraitUID InClickedTraitUID)
+{
+	FTraitRegistry& TraitRegistry = FTraitRegistry::Get();
+	if (const FTrait* Trait = TraitRegistry.Find(InClickedTraitUID))
+	{
+		if (UAnimNextGraph_EdGraphNode* EdGraphNode = Cast<UAnimNextGraph_EdGraphNode>(TraitEditorSharedData->EdGraphNodeWeak))
+		{
+			if (UAnimNextGraph_Controller* Controller = Cast<UAnimNextGraph_Controller>(EdGraphNode->GetController()))
+			{
+				int32 TraitIndex = INDEX_NONE;
+				TSharedPtr<FTraitDataEditorDef> SwapTraitData = FTraitEditorUtils::FindTraitInCurrentStackData(SelectedTraitUID, TraitEditorSharedData->CurrentTraitsDataShared, &TraitIndex);
+
+				const bool bIsAddingMissingBaseTrait = Trait->GetTraitMode() == ETraitMode::Base && TraitIndex == 0 && (*TraitEditorSharedData->CurrentTraitsDataShared)[0]->TraitUID == FTraitUID();
+				const int32 PinIndex = GetTraitPinIndex(EdGraphNode, SwapTraitData, TraitIndex);
+
+				const FName NewTraitTypeName = *Trait->GetTraitName();
+				if (TraitIndex == INDEX_NONE || bIsAddingMissingBaseTrait)
+				{
+					Controller->AddTrait(EdGraphNode->GetFName(), NewTraitTypeName, PinIndex);
+				}
+				else
+				{
+					Controller->SwapTrait(EdGraphNode->GetFName(), SwapTraitData->TraitName, PinIndex, NewTraitTypeName);
+				}
+			}
+		}
+
+		return FReply::Handled();
+	}
+
+	return FReply::Unhandled();
+}
+
+TWeakPtr<FTraitDataEditorDef> STraitEditorView::OnGetSelectedTraitData() const
+{
+	return StackSelectedTrait.ToWeakPtr();
+}
+
+FReply STraitEditorView::OnTraitDeleteRequest(const FTraitUID InTraitUIDToDelete)
+{
+	if (TArray<TSharedPtr<FTraitDataEditorDef>>* CurrentTraitsData = TraitEditorSharedData->CurrentTraitsDataShared.Get())
+	{
+		const int32 Num = CurrentTraitsData->Num();
+		for (int32 i = Num - 1; i >= 0; --i)
+		{
+			TSharedPtr<FTraitDataEditorDef>& TraitDataEditorDef = (*CurrentTraitsData)[i];
+			if (TraitDataEditorDef->TraitUID == InTraitUIDToDelete)
+			{
+				if (TraitEditorSharedData->EdGraphNodeWeak.IsValid())
+				{
+					if (UAnimNextGraph_EdGraphNode* EdGraphNode = Cast<UAnimNextGraph_EdGraphNode>(TraitEditorSharedData->EdGraphNodeWeak.Get()))
+					{
+						if (UAnimNextGraph_Controller* Controller = Cast<UAnimNextGraph_Controller>(EdGraphNode->GetController()))
+						{
+							Controller->RemoveTrait(EdGraphNode->GetFName(), TraitDataEditorDef->TraitName);
+						}
+					}
+				}
+
+				return FReply::Handled();
+			}
+		}
+	}
+
+	return FReply::Unhandled();
+}
+
+FReply STraitEditorView::OnStackTraitDragAccepted(const FTraitUID DraggedTraitUID, const FTraitUID TargetTraitUID, EItemDropZone DropZone)
+{
+	GEditor->GetTimerManager()->SetTimerForNextTick(FTimerDelegate::CreateSPLambda(this, [this, DraggedTraitUID, TargetTraitUID, DropZone]()
+	{
+		ExecuteTraitDrag(DraggedTraitUID, TargetTraitUID, DropZone);
+	}));
+
+	return FReply::Handled();
+}
+
+void STraitEditorView::ExecuteTraitDrag(const FTraitUID DraggedTraitUID, const FTraitUID TargetTraitUID, EItemDropZone DropZone)
+{
+	FTraitRegistry& TraitRegistry = FTraitRegistry::Get();
+
+	const FTrait* DraggedTrait = TraitRegistry.Find(DraggedTraitUID);
+	
+	if (DraggedTrait != nullptr)
+	{
+		if (UAnimNextGraph_EdGraphNode* EdGraphNode = Cast<UAnimNextGraph_EdGraphNode>(TraitEditorSharedData->EdGraphNodeWeak))
+		{
+			if (UAnimNextGraph_Controller* Controller = Cast<UAnimNextGraph_Controller>(EdGraphNode->GetController()))
+			{
+				int32 TargetTraitIndex = INDEX_NONE;
+				TSharedPtr<FTraitDataEditorDef> SwapTraitData = FTraitEditorUtils::FindTraitInCurrentStackData(TargetTraitUID, TraitEditorSharedData->CurrentTraitsDataShared, &TargetTraitIndex);
+				const int32 PinIndex = GetTraitPinIndex(EdGraphNode, SwapTraitData, TargetTraitIndex);
+				const bool bIsAddingMissingBaseTrait = DraggedTrait->GetTraitMode() == ETraitMode::Base && TargetTraitIndex == 0 && (*TraitEditorSharedData->CurrentTraitsDataShared)[0]->TraitUID == FTraitUID();
+
+				if (TargetTraitIndex != INDEX_NONE || bIsAddingMissingBaseTrait)
+				{
+					const FName DraggedTraitTypeName = *DraggedTrait->GetTraitName();
+
+					if (DropZone == EItemDropZone::OntoItem)
+					{
+						if(bIsAddingMissingBaseTrait)
+						{
+							Controller->AddTrait(EdGraphNode->GetFName(), DraggedTraitTypeName, PinIndex);
+						}
+						else
+						{
+							Controller->SwapTrait(EdGraphNode->GetFName(), SwapTraitData->TraitName, PinIndex, DraggedTraitTypeName);
+						}
+					}
+					else
+					{
+						int32 DraggedTraitIndex = INDEX_NONE;
+						TSharedPtr<FTraitDataEditorDef> DraggedTraitData = FTraitEditorUtils::FindTraitInCurrentStackData(DraggedTraitUID, TraitEditorSharedData->CurrentTraitsDataShared, &DraggedTraitIndex);
+
+						if (DraggedTraitIndex != INDEX_NONE)
+						{
+							Controller->SetTraitPinIndex(EdGraphNode->GetFName(), DraggedTraitData->TraitName, PinIndex + 1);
+						}
+						else
+						{
+							Controller->AddTrait(EdGraphNode->GetFName(), DraggedTraitTypeName, PinIndex + 1);
+						}
+					}
+				}
+			}
+		}
+	}
+}
+
+void STraitEditorView::OnStatckTraitSelectionChanged(const FTraitUID InTraitSelected)
+{
+	SelectedTraitUID = InTraitSelected;
+
+	StackSelectedTrait = InTraitSelected != FTraitUID() ? FTraitEditorUtils::FindTraitInCurrentStackData(InTraitSelected, TraitEditorSharedData->CurrentTraitsDataShared) : nullptr;
+}
+
+void STraitEditorView::Refresh()
+{
+	RefreshTraitStackTraitsStatus();
+	RefreshWidgets();
+}
+
+void STraitEditorView::RefreshWidgets()
+{
+	TraitListWidget->RefreshList();
+	TraitStackWidget->RefreshList();
+}
+
+void STraitEditorView::RefreshTraitStack()
+{
+	GenerateTraitStackData(TraitEditorSharedData->EdGraphNodeWeak, TraitEditorSharedData);
+	RefreshTraitStackTraitsStatus();
+
+	TraitStackWidget->RefreshList();
+}
+
+void STraitEditorView::OnRequestRefresh()
+{
+	Refresh();
+}
+
+void STraitEditorView::RefreshTraitStackTraitsStatus()
+{
+	if (TraitEditorSharedData->CurrentTraitsDataShared.IsValid())
+	{
+		if (TraitEditorSharedData->EdGraphNodeWeak.IsValid())
+		{
+			TraitEditorSharedData->bStackContainsErrors = false;
+
+			TArray<TSharedPtr<FTraitDataEditorDef>>& CurrentTraitsData = *(TraitEditorSharedData->CurrentTraitsDataShared.Get());
+
+			const int32 NumTraits = CurrentTraitsData.Num();
+			for (int32 i = NumTraits - 1; i >= 0; i--)
+			{
+				TSharedPtr<FTraitDataEditorDef>& TraitData = CurrentTraitsData[i];
+				UpdateTraitStatusInStack(CurrentTraitsData, i, TraitData);
+				
+				for (FTraitInterfaceUID TtaitInterfaceUID : TraitData->StackStatus.MissingInterfaces)
+				{
+					TraitEditorSharedData->StackMissingInterfaces.AddUnique(TtaitInterfaceUID);
+					
+					const int32 StackUsedInterfacesIndex = TraitEditorSharedData->StackUsedInterfaces.Find(TtaitInterfaceUID);
+					ensure(StackUsedInterfacesIndex != INDEX_NONE);
+					TraitEditorSharedData->StackUsedInteraceMissingIndexes.AddUnique(StackUsedInterfacesIndex);
+				}
+
+				if (TraitData->StackStatus.StatusMessages.Num() > 0)
+				{
+					TraitEditorSharedData->bStackContainsErrors = true;
+
+					if (TSharedPtr<IMessageLogListing> CompilerResultsListing = GetMessageLogListing().Pin())
+					{
+						for (const FTraitStackTraitStatus::FStatusMessage& StatusMessage : TraitData->StackStatus.StatusMessages)
+						{
+							const EMessageSeverity::Type Severity = (StatusMessage.Status == FTraitStackTraitStatus::EStackStatus::Warning) ? EMessageSeverity::Warning : EMessageSeverity::Error;
+							const FText& MessageText = StatusMessage.MessageText;
+
+							TSharedRef<FTokenizedMessage> Message = FTokenizedMessage::Create(Severity, FText::Format(LOCTEXT("TraitEditorLogTraitNameErrorFormat", "{0} : {1}"), TraitData->TraitDisplayName, MessageText));
+							CompilerResultsListing->AddMessage(Message, false);
+						}
+					}
+				}
+			}
+		}
+		else
+		{
+			TraitEditorSharedData->CurrentTraitsDataShared->Reset();
+		}
+	}
+}
+
+void STraitEditorView::UpdateTraitStatusInStack(const TArray<TSharedPtr<FTraitDataEditorDef>>& CurrentTraitsData, int32 TraitIndex, TSharedPtr<FTraitDataEditorDef>& TraitData)
+{
+	switch (TraitData->TraitMode)
+	{
+		case ETraitMode::Base:
+		{
+			if (TraitData->TraitUID == FTraitUID())
+			{
+				TraitData->StackStatus.StatusMessages.Add(FTraitStackTraitStatus::FStatusMessage(FTraitStackTraitStatus::EStackStatus::Error, FText(LOCTEXT("TraitStatusInStack_InvalidBaseTrait", "Base Trait Data is Invalid. Please, select a new Base Trait."))));
+				return;
+			}
+			if (CurrentTraitsData[0]->TraitUID != TraitData->TraitUID)
+			{
+				TraitData->StackStatus.StatusMessages.Add(FTraitStackTraitStatus::FStatusMessage(FTraitStackTraitStatus::EStackStatus::Error, FText(LOCTEXT("TraitStatusInStack_BaseNotAtTop", "Base Traits should be the first Trait in the Stack."))));
+				return;
+			}
+			break;
+		}
+		case ETraitMode::Additive:
+		{
+			if (TraitData->TraitUID == FTraitUID())
+			{
+				TraitData->StackStatus.StatusMessages.Add(FTraitStackTraitStatus::FStatusMessage(FTraitStackTraitStatus::EStackStatus::Error, FText(LOCTEXT("TraitStatusInStack_InvalidAdditiveTrait", "Additive Trait Data is Invalid. Please, fix the Stack."))));
+				return;
+			}
+			if (CurrentTraitsData[0]->TraitUID == TraitData->TraitUID)
+			{
+				TraitData->StackStatus.StatusMessages.Add(FTraitStackTraitStatus::FStatusMessage(FTraitStackTraitStatus::EStackStatus::Error, FText(LOCTEXT("TraitStatusInStack_AdditiveAtTop", "Additive Traits can not be at the Top of the Stack."))));
+				return;
+			}
+			break;
+		}
+		default:
+		{
+			TraitData->StackStatus.StatusMessages.Add(FTraitStackTraitStatus::FStatusMessage(FTraitStackTraitStatus::EStackStatus::Error, FText(LOCTEXT("TraitStatusInStack_InvalidTraitData", "Trait data is invalid, please correct the Stack."))));
+			break;
+		}
+	}
+
+	if (TraitData->TraitMode == ETraitMode::Additive)
+	{
+		const TArray<FTraitInterfaceUID>& RequiredInterfaces = TraitData->RequiredInterfaces;
+
+		const int32 NumRequired = RequiredInterfaces.Num();
+		if (NumRequired > 0)
+		{
+			for (int32 RequiredIndex = 0; RequiredIndex < NumRequired; RequiredIndex++)
+			{
+				const FTraitInterfaceUID& RequiredInterface = RequiredInterfaces[RequiredIndex];
+
+				bool bFound = false;
+				for (int32 ParentIndex = TraitIndex - 1; ParentIndex >= 0 && !bFound; ParentIndex--)
+				{
+					const TSharedPtr<FTraitDataEditorDef>& ParentTraitData = CurrentTraitsData[ParentIndex];
+
+					const TArray<FTraitInterfaceUID>& ParentImplementedInterfaces = ParentTraitData->ImplementedInterfaces;
+
+					for (const FTraitInterfaceUID& ParentImplementedInterface : ParentImplementedInterfaces)
+					{
+						if (ParentImplementedInterface == RequiredInterface)
+						{
+							bFound = true;
+							break;
+						}
+					}
+				}
+				if (!bFound)
+				{
+					TraitData->StackStatus.MissingInterfaces.Add(RequiredInterface);
+				}
+			}
+		}
+
+		for (const FTraitInterfaceUID& MissingInterface : TraitData->StackStatus.MissingInterfaces)
+		{
+			const FString InterfaceName(MissingInterface.GetInterfaceName());
+			const FText MissingError = FText::Format(LOCTEXT("TraitStatusInStack_MissingInterface", "Trait {0} requires a parent implementing interface {1}"), TraitData->TraitDisplayName, FText::FromString(InterfaceName));
+			TraitData->StackStatus.StatusMessages.Add(FTraitStackTraitStatus::FStatusMessage(FTraitStackTraitStatus::EStackStatus::Warning, MissingError));
+		}
+	}
+
+	TraitData->StackStatus.TraitStatus = TraitData->StackStatus.HasErrors()
+											? FTraitStackTraitStatus::EStackStatus::Error 
+											: TraitData->StackStatus.HasWarnings() 
+												? FTraitStackTraitStatus::EStackStatus::Warning
+												: FTraitStackTraitStatus::EStackStatus::Ok;
+}
+
+TSharedRef<SWidget> STraitEditorView::GetOptionsMenuWidget()
+{
+	TSharedPtr<SImage> FilterImage = SNew(SImage)
+		.Image(FAppStyle::Get().GetBrush("DetailsView.ViewOptions"))
+		.ColorAndOpacity(FSlateColor::UseForeground());
+
+	FMenuBuilder DetailViewOptions(true, nullptr);
+
+	TSharedPtr<FTraitEditorSharedData>& TraitEditorSharedDataLocal = TraitEditorSharedData;
+
+	DetailViewOptions.AddMenuEntry(
+		LOCTEXT("TraitEditor_ShowTraitInterfaces", "Show Trait Interfaces"),
+		LOCTEXT("TraitEditor_ShowTraitInterfaces_ToolTip", "Displays Trait Implemented and Required interfaces"),
+		FSlateIcon(),
+		FUIAction(
+			FExecuteAction::CreateLambda([this, TraitEditorSharedDataLocal]()
+			{
+				if (TraitEditorSharedDataLocal.IsValid())
+				{
+					// TODO zzz : Store bShowTraitInterfaces in a serialized settings class somewhere
+					TraitEditorSharedDataLocal->bShowTraitInterfaces = !TraitEditorSharedDataLocal->bShowTraitInterfaces;
+					TraitEditorSharedDataLocal->bShowTraitInterfacesIfWarningsOrErrors = false;
+					RefreshWidgets();
+				}
+			}),
+			FCanExecuteAction(),
+			FIsActionChecked::CreateLambda([TraitEditorSharedDataLocal]()
+			{
+				return (TraitEditorSharedDataLocal.IsValid()) ? TraitEditorSharedDataLocal->bShowTraitInterfaces : false;
+			})
+		),
+		NAME_None,
+		EUserInterfaceActionType::Check
+	);
+	DetailViewOptions.AddMenuEntry(
+		LOCTEXT("TraitEditor_ShowTraitInterfaces", "Show Trait Interfaces If Warnings / Errors"),
+		LOCTEXT("TraitEditor_ShowTraitInterfaces_ToolTip", "Displays Trait Implemented and Required interfaces if is there any Warning or Error on the Stack"),
+		FSlateIcon(),
+		FUIAction(
+			FExecuteAction::CreateLambda([this, TraitEditorSharedDataLocal]()
+				{
+					if (TraitEditorSharedDataLocal.IsValid())
+					{
+						// TODO zzz : Store bShowTraitInterfacesIfWarningsOrErrors in a serialized settings class somewhere
+						TraitEditorSharedDataLocal->bShowTraitInterfacesIfWarningsOrErrors = !TraitEditorSharedDataLocal->bShowTraitInterfacesIfWarningsOrErrors;
+						TraitEditorSharedDataLocal->bShowTraitInterfaces = false;
+						RefreshWidgets(); //RefreshTraitStack();
+					}
+				}),
+			FCanExecuteAction(),
+			FIsActionChecked::CreateLambda([TraitEditorSharedDataLocal]()
+				{
+					return (TraitEditorSharedDataLocal.IsValid()) ? TraitEditorSharedDataLocal->bShowTraitInterfacesIfWarningsOrErrors : false;
+				})
+		),
+		NAME_None,
+		EUserInterfaceActionType::Check
+	);
+
+
+	return SNew(SHorizontalBox)
+		+ SHorizontalBox::Slot()
+		.FillWidth(1.0)
+		.HAlign(HAlign_Right)
+		[
+			SNew(SComboButton)
+			.HasDownArrow(false)
+			.ContentPadding(0.0f)
+			.ForegroundColor(FSlateColor::UseForeground())
+			.ButtonStyle(FAppStyle::Get(), "SimpleButton")
+			.AddMetaData<FTagMetaData>(FTagMetaData(TEXT("ViewOptions")))
+			.MenuContent()
+			[
+				DetailViewOptions.MakeWidget()
+			]
+			.ButtonContent()
+			[
+				FilterImage.ToSharedRef()
+			]
+		];
+}
+
+TWeakPtr<IMessageLogListing>& STraitEditorView::GetMessageLogListing()
+{
+	if (!CompilerResultsListingWeak.IsValid())
+	{
+		if (TSharedPtr<UE::Workspace::IWorkspaceEditor> WorkspaceEditor = WorkspaceEditorWeak.Pin())
+		{
+			TSharedPtr<SDockTab> DockTab = WorkspaceEditor->GetTabManager()->FindExistingLiveTab(UE::AnimNext::Editor::CompilerResultsTabName);
+			if (!DockTab.IsValid())
+			{
+				DockTab = WorkspaceEditor->GetTabManager()->TryInvokeTab(UE::AnimNext::Editor::CompilerResultsTabName);
+			}
+
+			if (DockTab.IsValid())
+			{
+				if (TSharedPtr<SAnimNextCompilerResultsWidget> AnimNextCompilerResultsWidget = StaticCastSharedPtr<SAnimNextCompilerResultsWidget>(DockTab->GetContent().ToSharedPtr()))
+				{
+					CompilerResultsListingWeak = AnimNextCompilerResultsWidget->GetCompilerResultsListing();
+				}
+			}
+		}
+	}
+
+	return CompilerResultsListingWeak;
+}
+
+int32 STraitEditorView::GetTraitPinIndex(UAnimNextGraph_EdGraphNode* InEdGraphNode, const TSharedPtr<FTraitDataEditorDef>& InTraitData, int32 InTraitIndex)
+{
+	if (InEdGraphNode == nullptr || !InTraitData.IsValid())
+	{
+		return INDEX_NONE;
+	}
+
+	if (URigVMNode* ModelNode = InEdGraphNode->GetModelNode())
+	{
+		// Special case for a stack that has no Base Trait but has Additive traits (we have an empty trait at index 0)
+		if (InTraitIndex == 0 && TraitEditorSharedData->CurrentTraitsDataShared.IsValid() && TraitEditorSharedData->CurrentTraitsDataShared->Num() > 0 && (*TraitEditorSharedData->CurrentTraitsDataShared)[0]->TraitUID == FTraitUID())
+		{
+			// Get first index that is a decorator pin, skipping non decorator pins and execute pins
+			const TArray<URigVMPin*> Pins = ModelNode->GetPins();
+			int32 Index = 0;
+			for (const URigVMPin* Pin : Pins)
+			{
+				if (!Pin->IsDecoratorPin() || Pin->IsExecuteContext())
+				{
+					Index++;
+					continue;
+				}
+
+				break;
+			}
+			return Index;
+		}
+		else
+		{
+			// Obtain the pins from the stack
+			const TArray<URigVMPin*> TraitPins = ModelNode->GetDecoratorPins();
+			for (const URigVMPin* TraitPin : TraitPins)
+			{
+				if (TraitPin->IsExecuteContext())
+				{
+					continue;
+				}
+
+				if (TraitPin->GetFName() == InTraitData->TraitName)
+				{
+					return TraitPin->GetPinIndex();
+				}
+			}
+		}
+	}
+
+	return INDEX_NONE;
+}
+
+/*static*/ void STraitEditorView::GenerateTraitStackData(const TWeakObjectPtr<UAnimNextGraph_EdGraphNode>& EdGraphNodeWeak, TSharedPtr<FTraitEditorSharedData>& InTraitEditorSharedData)
+{
+	if (!ensure(InTraitEditorSharedData.IsValid() && InTraitEditorSharedData->CurrentTraitsDataShared.IsValid()))
+	{
+		return;
+	}
+
+	TArray<TSharedPtr<FTraitDataEditorDef>>& TraitsData = *InTraitEditorSharedData->CurrentTraitsDataShared.Get();
+	TArray<FTraitInterfaceUID>& StackUsedInterfaces = InTraitEditorSharedData->StackUsedInterfaces;
+	TArray<FTraitInterfaceUID>& StackMissingInterfaces = InTraitEditorSharedData->StackMissingInterfaces;
+	
+	TraitsData.Reset();
+	StackUsedInterfaces.Reset();
+	StackMissingInterfaces.Reset();
+
+	if (UAnimNextGraph_EdGraphNode* EdGraphNode = Cast<UAnimNextGraph_EdGraphNode>(EdGraphNodeWeak.Get()))
+	{
+		if (URigVMNode* ModelNode = EdGraphNode->GetModelNode())
+		{
+			int32 NumBaseTraits = 0;
+
+			// Obtain the pins from the stack
+			const TArray<URigVMPin*> TraitPins = ModelNode->GetDecoratorPins();
+			const int32 NumPins = TraitPins.Num();
+			if (NumPins > 0)
+			{
+				FTraitRegistry& TraitRegistry = FTraitRegistry::Get();
+
+				TArray<int32> StackUsedInterfacesNumHits;
+				StackUsedInterfacesNumHits.Reset(NumPins);
+
+				// For each Trait (represented as a pin in the node)
+				for (const URigVMPin* TraitPin : TraitPins)
+				{
+					if (TraitPin->IsExecuteContext())
+					{
+						continue;
+					}
+
+					// Create a temporary trait instance, in order to get the correct TraitSharedDataStruct
+					if (TSharedPtr<FStructOnScope> ScopedTrait = ModelNode->GetDecoratorInstance(TraitPin->GetFName()))
+					{
+						if (const FRigVMDecorator* RigVMTrait = (FRigVMDecorator*)ScopedTrait->GetStructMemory())
+						{
+							if (TWeakObjectPtr<UScriptStruct> TraitStruct = RigVMTrait->GetDecoratorSharedDataStruct(); TraitStruct.IsValid())
+							{
+								if (const FTrait* Trait = TraitRegistry.Find(TraitStruct.Get()))
+								{
+									TConstArrayView<FTraitInterfaceUID> TraitImplementedInterfaces = Trait->GetTraitInterfaces();
+									TConstArrayView<FTraitInterfaceUID> TraitRequiredInterfaces = Trait->GetTraitRequiredInterfaces();
+									
+									TSharedPtr<FTraitDataEditorDef>& TraitData = TraitsData.Add_GetRef(MakeShared<FTraitDataEditorDef>(RigVMTrait->GetName()
+																																		, TraitStruct->GetDisplayNameText()
+																																		, Trait->GetTraitUID()
+																																		, Trait->GetTraitMode()
+																																		, TArray<FTraitInterfaceUID>(TraitImplementedInterfaces)
+																																		, TArray<FTraitInterfaceUID>(TraitRequiredInterfaces)));
+
+									for (const FTraitInterfaceUID& TraitInterface : TraitImplementedInterfaces)
+									{
+										if (!FTraitEditorUtils::IsInternalTraitInteface(TraitInterface))
+										{
+											if (int32 Index = StackUsedInterfaces.Find(TraitInterface) != INDEX_NONE)
+											{
+												StackUsedInterfacesNumHits[Index]++;
+											}
+											else
+											{
+												Index = StackUsedInterfaces.Add(TraitInterface);
+												StackUsedInterfacesNumHits.Add(1);
+											}
+										}
+									}
+
+									for (const FTraitInterfaceUID& TraitInterface : TraitRequiredInterfaces)
+									{
+										if (!FTraitEditorUtils::IsInternalTraitInteface(TraitInterface))
+										{
+											if (int32 Index = StackUsedInterfaces.Find(TraitInterface) != INDEX_NONE)
+											{
+												StackUsedInterfacesNumHits[Index]++;
+											}
+											else
+											{
+												Index = StackUsedInterfaces.Add(TraitInterface);
+												StackUsedInterfacesNumHits.Add(1);
+											}
+										}
+									}
+								}
+							}
+						}
+					}
+				}
+
+				const int32 NumUsedInterfaces = StackUsedInterfaces.Num();
+				if (NumUsedInterfaces > 0)
+				{
+					// Generate the stack interface used indexes for each Trait
+					for (TSharedPtr<FTraitDataEditorDef>& TraitData : TraitsData)
+					{
+						FTraitEditorUtils::GenerateStackInterfacesUsedIndexes(TraitData, InTraitEditorSharedData);
+					}
+				}
+
+				for (const TSharedPtr<FTraitDataEditorDef>& TraitData : TraitsData)
+				{
+					if (TraitData->TraitMode == ETraitMode::Base)
+					{
+						NumBaseTraits++;
+					}
+				}
+			}
+
+			if (NumBaseTraits == 0)
+			{
+				TraitsData.Insert(MakeShared<FTraitDataEditorDef>(NAME_None, LOCTEXT("BaseTraitUnset", "<Base Trait Unset>"), FTraitUID(), ETraitMode::Base, TArray<FTraitInterfaceUID>(), TArray<FTraitInterfaceUID>()), 0);
+			}
+		}
+	}
+	else
+	{
+		TraitsData.Reset();
+	}
+}
+
+} // end namespace UE::Workspace
+
+#undef LOCTEXT_NAMESPACE
diff --git a/Engine/Plugins/Experimental/AnimNext/Source/AnimNextEditor/Private/Graph/STraitEditorView.h b/Engine/Plugins/Experimental/AnimNext/Source/AnimNextEditor/Private/Graph/STraitEditorView.h
new file mode 100644
--- /dev/null
+++ b/Engine/Plugins/Experimental/AnimNext/Source/AnimNextEditor/Private/Graph/STraitEditorView.h
@@ -0,0 +1,94 @@
+// Copyright Epic Games, Inc. All Rights Reserved.
+
+#pragma once
+
+#include "ContentBrowserDelegates.h"
+#include "Widgets/SCompoundWidget.h"
+#include "Widgets/Views/STableRow.h"
+
+#include "Graph/TraitEditorDefs.h"
+
+
+class UAnimNextGraph_Controller;
+class UAnimNextGraph_EdGraphNode;
+class IMessageLogListing;
+
+namespace UE::Workspace
+{
+
+class IWorkspaceEditor;
+
+}
+
+namespace UE::AnimNext::Editor
+{
+
+class STraitListView;
+class STraitStackView;
+
+struct FTraitEditorViewEntry;
+
+}
+
+
+namespace UE::AnimNext::Editor
+{
+
+struct FTraitStackData
+{
+	TWeakObjectPtr<UAnimNextGraph_EdGraphNode> EdGraphNodeWeak = nullptr;
+};
+
+class STraitEditorView : public SCompoundWidget
+{
+public:
+	STraitEditorView();
+
+	SLATE_BEGIN_ARGS(STraitEditorView) {}
+
+	SLATE_END_ARGS()
+
+	void Construct(const FArguments& InArgs, TWeakPtr<UE::Workspace::IWorkspaceEditor> InWorkspaceEditorWeak);
+
+	void SetTraitData(const FTraitStackData& InTraitStackData);
+
+private:
+	FReply OnTraitClicked(const FTraitUID InTraitClicked);
+	FReply OnTraitDeleteRequest(const FTraitUID InTraitUIDToDelete);
+	FReply OnStackTraitDragAccepted(const FTraitUID DraggedTraitUID, const FTraitUID TargetTraitUID, EItemDropZone DropZone);
+	void ExecuteTraitDrag(const FTraitUID DraggedTraitUID, const FTraitUID TargetTraitUID, EItemDropZone DropZone);
+	void OnStatckTraitSelectionChanged(const FTraitUID InTraitSelected);
+	TWeakPtr<FTraitDataEditorDef> OnGetSelectedTraitData() const;
+
+	void Refresh();
+	void RefreshWidgets();
+	void RefreshTraitStack();
+
+	void OnRequestRefresh();
+	void RefreshTraitStackTraitsStatus();
+	void UpdateTraitStatusInStack(const TArray<TSharedPtr<FTraitDataEditorDef>>& CurrentTraitsData, int32 TraitIndex, TSharedPtr<FTraitDataEditorDef>& TraitData);
+
+	TSharedRef<SWidget> GetOptionsMenuWidget();
+
+	TWeakPtr<IMessageLogListing>& GetMessageLogListing();
+
+	int32 GetTraitPinIndex(UAnimNextGraph_EdGraphNode* InEdGraphNode, const TSharedPtr<FTraitDataEditorDef>& InTraitData, int32 TraitIndex = INDEX_NONE);
+
+	static void GenerateTraitStackData(const TWeakObjectPtr<UAnimNextGraph_EdGraphNode>& EdGraphNodeWeak, TSharedPtr<FTraitEditorSharedData>& InTraitEditorSharedData);
+
+	TWeakPtr<UE::Workspace::IWorkspaceEditor> WorkspaceEditorWeak;
+	TSharedPtr<FTraitEditorSharedData> TraitEditorSharedData;
+
+	TSharedPtr<STraitListView> TraitListWidget;
+	TSharedPtr<STraitStackView> TraitStackWidget;
+
+	TSharedPtr<FTraitDataEditorDef> StackSelectedTrait;
+
+	TWeakPtr<IMessageLogListing> CompilerResultsListingWeak;
+
+	FTraitUID SelectedTraitUID;
+	bool bStackContainsErrors = false;
+	bool bShowTraitInterfaces = false;
+};
+
+}
diff --git a/Engine/Plugins/Experimental/AnimNext/Source/AnimNextEditor/Private/Graph/STraitListView.cpp b/Engine/Plugins/Experimental/AnimNext/Source/AnimNextEditor/Private/Graph/STraitListView.cpp
new file mode 100644
--- /dev/null
+++ b/Engine/Plugins/Experimental/AnimNext/Source/AnimNextEditor/Private/Graph/STraitListView.cpp
@@ -0,0 +1,834 @@
+// Copyright Epic Games, Inc. All Rights Reserved.
+
+#include "STraitListView.h"
+#include "ContentBrowserModule.h"
+#include "IContentBrowserSingleton.h"
+#include "SAssetDropTarget.h"
+#include "ScopedTransaction.h"
+#include "WorkspaceSchema.h"
+#include "Framework/Commands/GenericCommands.h"
+#include "Widgets/Input/SSearchBox.h"
+#include "ObjectEditorUtils.h"
+#include "TraitCore/TraitMode.h"
+#include "TraitCore/TraitUID.h"
+#include "TraitCore/TraitRegistry.h"
+#include "Common/SCategoryTableRow.h"
+#include "Widgets/Input/SButton.h"
+
+#define LOCTEXT_NAMESPACE "TraitListEditor"
+
+namespace
+{
+
+static const FName NAME_Category("Category");
+static const FName NAME_DefaultCategory("Default");
+static const FName Column_Trait(TEXT("Traits"));
+static const FName BaseCategoryName(TEXT("Base"));
+static const FName AdditiveCategoryName(TEXT("Additive"));
+static const FText BaseCategoryText(LOCTEXT("TraitBaseCategoryName", "Base"));
+static const FText AdditiveCategoryText(LOCTEXT("AdditiveBaseCategoryName", "Additive"));
+
+}
+
+namespace UE::AnimNext::Editor
+{
+
+// An entry displayed in the parameters view
+struct FTraitListEntry
+{
+	FTraitListEntry() = default;
+	virtual ~FTraitListEntry() = default;
+
+	virtual bool IsCategory() const
+	{
+		return false;
+	}
+
+	virtual bool IsTrait() const
+	{
+		return false;
+	}
+
+	virtual uint8 GetDepthLevel() const
+	{
+		return DepthLevel;
+	}
+
+	virtual bool IsRoot() const
+	{
+		return DepthLevel == 0;
+	}
+
+	virtual bool HasChildren() const
+	{
+		return false;
+	}
+
+	virtual void GenerateChildren(const FString& InFilterText, uint8 CurrentDepth = 0)
+	{
+	}
+
+	virtual void GetChildren(TArray<TSharedRef<FTraitListEntry>>& OutChildren) const
+	{
+	}
+
+	virtual void GetChildrenRecursive(TArray<TSharedRef<FTraitListEntry>>& OutChildren)
+	{
+	}
+
+	virtual const FName& GetCategory() const
+	{
+		static const FName None = NAME_None;
+		return None;
+	}
+
+	virtual const FText& GetCategoryText() const
+	{
+		return FText::GetEmpty();
+	}
+
+	virtual const FText& GetTraitNameText() const
+	{
+		return FText::GetEmpty();
+	}
+
+	virtual const TSharedPtr<FTraitDataEditorDef>& GetTraitData() const
+	{
+		static TSharedPtr<FTraitDataEditorDef> Empty;
+		return Empty;
+	}
+
+	virtual FTraitUID GetTraitUID() const
+	{
+		return FTraitUID();
+	}
+
+	virtual ETraitMode GetTraitMode() const
+	{
+		return ETraitMode::Invalid;
+	}
+
+	uint8 DepthLevel = 0;
+};
+
+struct FTraitEntry : FTraitListEntry
+{
+	FTraitEntry() = default;
+
+	explicit FTraitEntry(const TSharedPtr<FTraitDataEditorDef>& InTraitDataEditorDef)
+		: TraitDataEditorDef(InTraitDataEditorDef)
+	{}
+
+	virtual bool IsTrait() const override
+	{
+		return true;
+	}
+
+	virtual FTraitUID GetTraitUID() const override
+	{
+		return TraitDataEditorDef->TraitUID;
+	}
+
+	virtual ETraitMode GetTraitMode() const override
+	{
+		return TraitDataEditorDef->TraitMode;
+	}
+
+	virtual const FText& GetTraitNameText() const override
+	{
+		return TraitDataEditorDef->TraitDisplayName;
+	}
+
+	virtual const TSharedPtr<FTraitDataEditorDef>& GetTraitData() const override
+	{
+		return TraitDataEditorDef;
+	}
+
+	TSharedPtr<FTraitDataEditorDef> TraitDataEditorDef;
+};
+
+struct FTraitListCategoryEntry : FTraitListEntry
+{
+	FTraitListCategoryEntry() = default;
+
+	explicit FTraitListCategoryEntry(const FTraitCategoryData& InTraitCategoryData)
+		: Category(InTraitCategoryData.Category)
+		, CategoryText(InTraitCategoryData.CategoryText)
+		, TraitList(InTraitCategoryData.TraitList)
+	{
+	}
+
+	FTraitListCategoryEntry(FName InCategory, FText InCategoryText)
+		: Category(InCategory)
+		, CategoryText(InCategoryText)
+	{
+	}
+
+	virtual bool IsCategory() const override
+	{
+		return true;
+	}
+
+	virtual const FName& GetCategory() const override
+	{
+		return Category;
+	}
+
+	virtual const FText& GetCategoryText() const override
+	{
+		return CategoryText;
+	}
+
+	virtual bool HasChildren() const override
+	{
+		return Children.Num() > 0;
+	}
+
+	virtual void GetChildren(TArray<TSharedRef<FTraitListEntry>>& OutChildren) const
+	{
+		OutChildren = Children;
+	}
+
+	void GetChildrenRecursive(TArray< TSharedRef<FTraitListEntry> >& OutChildren)
+	{
+		for (TSharedRef<FTraitListEntry>& Entry : Children)
+		{
+			OutChildren.Add(Entry);
+			Entry->GetChildrenRecursive(OutChildren);
+		}
+	}
+
+	virtual void GenerateChildren(const FString& InFilterText, uint8 CurrentDepth = 0) override
+	{
+		DepthLevel = CurrentDepth;
+
+		const int32 NumChildren = Children.Num();
+		for (int32 i = NumChildren - 1; i >= 0; --i)
+		{
+			TSharedRef<FTraitListEntry>& Entry = Children[i];
+			if (Entry->IsCategory())
+			{
+				Entry->GenerateChildren(InFilterText, DepthLevel + 1);
+			}
+			else
+			{
+				Children.RemoveAt(i);
+			}
+		}
+
+		if (TraitList.Num() > 0)
+		{
+			for (const TSharedPtr<FTraitDataEditorDef>& TraitDataDef : TraitList)
+			{
+				if (InFilterText.IsEmpty() 
+					|| TraitDataDef->TraitDisplayName.ToString().Contains(InFilterText))
+				{
+					Children.Add(MakeShared<FTraitEntry>(TraitDataDef));
+				}
+			}
+		}
+	}
+
+	void AddEntry(const FName& InCategory, const FText& InCategoryText, const TSharedPtr<FTraitDataEditorDef>& InTraitDataDef)
+	{
+		TSharedRef<FTraitListCategoryEntry> CategoryEntry = FindOrCreateSubCategory(InCategory, InCategoryText);
+		CategoryEntry->Children.Add(MakeShared<FTraitEntry>(InTraitDataDef));
+	}
+
+	TSharedRef<FTraitListCategoryEntry> FindOrCreateSubCategory(const FName& InCategory, const FText& InCategoryText)
+	{
+		for (TSharedRef<FTraitListEntry>& Entry : Children)
+		{
+			if (Entry->GetCategory() == InCategory)
+			{
+				return StaticCastSharedRef<FTraitListCategoryEntry>(Entry);
+			}
+		}
+
+		TSharedRef<FTraitListCategoryEntry> Entry = StaticCastSharedRef<FTraitListCategoryEntry>(Children.Add_GetRef(MakeShared<FTraitListCategoryEntry>()));
+		Entry->Category = InCategory;
+		Entry->CategoryText = InCategoryText;
+
+		return Entry;
+	}
+
+	FName Category;
+	FText CategoryText;
+	TArray<TSharedPtr<FTraitDataEditorDef>> TraitList;
+	TArray<TSharedRef<FTraitListCategoryEntry>> SubCategories;
+
+	TArray<TSharedRef<FTraitListEntry>> Children;
+};
+
+
+void STraitListView::Construct(const FArguments& InArgs, TSharedPtr<FTraitEditorSharedData>& InTraitEditorSharedData)
+{
+	check(InTraitEditorSharedData.IsValid());
+	TraitEditorSharedData = InTraitEditorSharedData.ToSharedRef();
+
+	//UICommandList = MakeShared<FUICommandList>();
+
+	OnTraitClicked = InArgs._OnTraitClicked;
+	OnGetSelectedTraitData = InArgs._OnGetSelectedTraitData;
+
+	SAssignNew(TraitListFilterBox, SSearchBox)
+	.OnTextChanged(this, &STraitListView::OnFilterTextChanged);
+
+	ChildSlot
+	[
+		SNew(SHorizontalBox)
+		+ SHorizontalBox::Slot()
+		.FillWidth(1.0)
+		[
+			SNew(SBorder)
+			.Padding(4.0f)
+			.BorderImage(FAppStyle::GetBrush("ToolPanel.GroupBorder"))
+			[
+				SNew(SVerticalBox)
+				+ SVerticalBox::Slot()
+				.AutoHeight()
+				[
+					SNew(SHorizontalBox)
+					+ SHorizontalBox::Slot()
+					.FillWidth(1.0)
+					.VAlign(VAlign_Center)
+					[
+						TraitListFilterBox.ToSharedRef()
+					]
+				]
+
+				+ SVerticalBox::Slot()
+				.FillHeight(1.0)
+				.Padding(0.f, 4.f, 0.f, 0.f)
+				[
+					SAssignNew(EntriesList, STreeView<TSharedRef<FTraitListEntry>>)
+					.TreeItemsSource(&FilteredEntries)
+					.OnGenerateRow(this, &STraitListView::HandleGenerateRow)
+					.OnGetChildren(this, &STraitListView::HandleGetChildren)
+					.OnItemScrolledIntoView(this, &STraitListView::HandleItemScrolledIntoView)
+					.OnSelectionChanged(this, &STraitListView::HandleSelectionChanged)
+					.OnExpansionChanged_Lambda([this](TSharedRef<FTraitListEntry> Entry, bool bExpanded)
+						{
+							if (FilterText.IsEmpty())
+							{
+								StoreExpansionState();
+							}
+						})
+					.HeaderRow(
+						SNew(SHeaderRow)
+						+SHeaderRow::Column(Column_Trait)
+						.DefaultLabel(LOCTEXT("TraitListColumnHeader", "Traits"))
+						.HAlignHeader(HAlign_Center)
+						.ToolTipText(LOCTEXT("TraitListColumnHeaderTooltip", "The list of available Traits"))
+						.FillWidth(10.0f)
+					)
+				]
+			]
+		]
+	];
+
+	FilterText = FText();
+}
+
+void STraitListView::RefreshList()
+{
+	GenerateTraitList();
+
+	StoreExpansionState();
+
+	RefreshEntries();
+
+	TArray<TSharedRef<FTraitListEntry>> AllEntries;
+	GetAllEntries(AllEntries);
+	ExpandAllCategories(AllEntries);
+
+	FilterText = FText();
+	RefreshFilter();
+}
+
+FReply STraitListView::OnKeyDown(const FGeometry& MyGeometry, const FKeyEvent& InKeyEvent)
+{
+	if (UICommandList.IsValid() && UICommandList->ProcessCommandBindings(InKeyEvent))
+	{
+		return FReply::Handled();
+	}
+	return FReply::Unhandled();
+}
+
+void STraitListView::OnFilterTextChanged(const FText& InFilterText)
+{
+	FilterText = InFilterText;
+	RefreshFilter();
+}
+
+bool STraitListView::HasValidEditorData() const
+{
+	return TraitEditorSharedData.IsValid() && TraitEditorSharedData->EdGraphNodeWeak.IsValid();
+}
+
+TSharedPtr<FTraitDataEditorDef> STraitListView::GetSelectedTraitData() const
+{
+	TSharedPtr<FTraitDataEditorDef> StackSelectedTrait;
+	if (OnGetSelectedTraitData.IsBound())
+	{
+		StackSelectedTrait = OnGetSelectedTraitData.Execute().Pin();
+	}
+
+	return StackSelectedTrait;
+}
+
+void STraitListView::StoreExpansionState()
+{
+	if (!bIgnoreExpansion									// Do not store state while expanding all items or restoring expansion state
+		&& HasValidEditorData()								// Do not store state if no data is set (or has been cleared)
+		&& EntriesList->GetNumItemsBeingObserved() > 0)		// Do not store state if the entry list is empty (set after being previously cleared)
+	{
+		OldExpansionState.Reset();
+		EntriesList->GetExpandedItems(OldExpansionState);
+	}
+}
+
+template<typename ItemType, typename ComparisonType>
+static void RestoreExpansionStateT(TSharedPtr< STreeView<ItemType> > InTree, const TArray<ItemType>& ItemSource, const TSet<ItemType>& OldExpansionState, ComparisonType ComparisonFunction)
+{
+	check(InTree.IsValid());
+
+	// Iterate over new tree items
+	for (int32 ItemIdx = 0; ItemIdx < ItemSource.Num(); ItemIdx++)
+	{
+		ItemType NewItem = ItemSource[ItemIdx];
+
+		bool bFound = false;
+		// Look through old expansion state
+		for (typename TSet<ItemType>::TConstIterator OldExpansionIter(OldExpansionState); OldExpansionIter; ++OldExpansionIter)
+		{
+			const ItemType OldItem = *OldExpansionIter;
+			// See if this matches this new item
+			if (ComparisonFunction(OldItem, NewItem))
+			{
+				// It does, so expand it
+				InTree->SetItemExpansion(NewItem, true);
+				bFound = true;
+			}
+		}
+		if (!bFound)
+		{
+			InTree->SetItemExpansion(NewItem, false);
+		}
+	}
+}
+
+void STraitListView::RestoreExpansionState(TArray<TSharedRef<FTraitListEntry>>& AllEntries)
+{
+	if (!OldExpansionState.IsEmpty() && AllEntries.Num() > 0)
+	{
+		TGuardValue<bool> GuardThreadId(bIgnoreExpansion, true);
+
+		// Restore the expanded items
+		RestoreExpansionStateT<TSharedRef<FTraitListEntry>>(EntriesList, AllEntries, OldExpansionState, [](TSharedRef<FTraitListEntry> A, TSharedRef<FTraitListEntry> B) -> bool
+			{
+				if (!A->GetCategory().IsNone() && A->GetCategory() == B->GetCategory())
+				{
+					return true;
+				}
+
+				return false;
+			});
+	}
+}
+
+void STraitListView::GetAllEntries(TArray<TSharedRef<FTraitListEntry>>& AllEntries) const
+{
+	for (const TSharedRef<FTraitListEntry>& Entry : Categories)
+	{
+		AllEntries.Add(Entry);
+		Entry->GetChildrenRecursive(AllEntries);
+	}
+}
+
+void STraitListView::ExpandAllCategories(const TArray<TSharedRef<FTraitListEntry>>& AllEntries)
+{
+	TGuardValue<bool> GuardThreadId(bIgnoreExpansion, true);
+
+	for (const TSharedRef<FTraitListEntry>& Entry : AllEntries)
+	{
+		if (Entry->IsCategory())
+		{
+			EntriesList->SetItemExpansion(Entry, true);
+		}
+	}
+}
+
+void STraitListView::RefreshEntries()
+{
+	Categories.Reset(2);
+
+	if (HasValidEditorData())
+	{
+		CreateTraitCategories(BaseCategoryName, BaseCategoryText, BaseTraitCategories);
+		CreateTraitCategories(AdditiveCategoryName, AdditiveCategoryText, AdditiveTraitCategories);
+	}
+}
+
+void STraitListView::CreateTraitCategories(const FName& InCategoryName, const FText& InCategoryText, TMap<FName, FTraitCategoryData>& InCategoriesMap)
+{
+	TSharedRef<FTraitListCategoryEntry> Category = StaticCastSharedRef<FTraitListCategoryEntry>(Categories.Add_GetRef(MakeShared<FTraitListCategoryEntry>(InCategoryName, InCategoryText)));
+
+	const int32 NumBaseTraitCategories = InCategoriesMap.Num();
+	for (const TPair<FName, FTraitCategoryData>& CategoryPair : InCategoriesMap)
+	{
+		const FName& CategoryName = CategoryPair.Key;
+		const FTraitCategoryData& TraitCategoryData = CategoryPair.Value;
+		check(CategoryName == TraitCategoryData.Category);
+
+		if (CategoryName == NAME_DefaultCategory && NumBaseTraitCategories == 1)
+		{
+			Category->TraitList = TraitCategoryData.TraitList;
+		}
+		else
+		{
+			TSharedRef<FTraitListCategoryEntry> CategoryEntry = Category->FindOrCreateSubCategory(CategoryName, TraitCategoryData.CategoryText);
+			CategoryEntry->TraitList = TraitCategoryData.TraitList;
+		}
+	}
+}
+
+void STraitListView::RefreshFilter()
+{
+	FilteredEntries = Categories;
+
+	const FString FilterTextAsString = FilterText.ToString();
+	for (TSharedRef<FTraitListEntry>& Category : Categories)
+	{
+		Category->GenerateChildren(FilterTextAsString);
+	}
+
+	TArray<TSharedRef<FTraitListEntry>> AllEntries;
+	GetAllEntries(AllEntries);
+	ExpandAllCategories(AllEntries);
+
+	if (FilterText.IsEmpty())
+	{
+		RestoreExpansionState(AllEntries);
+	}
+
+	EntriesList->RequestListRefresh();
+}
+
+static bool IsTraitAvailable(const TWeakPtr<TArray<TSharedPtr<FTraitDataEditorDef>>> CurrentTraitsDataWeak, const FTraitUID TraitUID, const ETraitMode TraitMode, const TSharedPtr<FTraitDataEditorDef>& InStackSelectedTrait)
+{
+	// Disable Buttons of Traits that already exist in the list
+	if (TSharedPtr<TArray<TSharedPtr<FTraitDataEditorDef>>> CurrentTraitsDataShared = CurrentTraitsDataWeak.Pin())
+	{
+		// Only a base Trait is allowed if the list is empty
+		if (CurrentTraitsDataShared->IsEmpty())
+		{
+			return TraitMode == ETraitMode::Base;
+		}
+
+		// If a Trait is selected, only allow same type traits
+		if (InStackSelectedTrait.IsValid())
+		{
+			if (InStackSelectedTrait->TraitMode != TraitMode)
+			{
+				return false;
+			}
+		}
+
+		// Only additive Traits are allowed if the list is not empty, unless the first element is invalid
+		if (TraitMode == ETraitMode::Base)
+		{
+			// Only one base trait allowed
+			if (CurrentTraitsDataShared->Top()->TraitUID != FTraitUID() && !InStackSelectedTrait.IsValid())
+			{
+				for (const TSharedPtr<FTraitDataEditorDef>& Trait : *CurrentTraitsDataShared.Get())
+				{
+					if (Trait->TraitMode == ETraitMode::Base && Trait->TraitUID != FTraitUID())
+					{
+						return false;
+					}
+				}
+			}
+		}
+
+		// No duplicated Traits allowed in the list
+		for (const TSharedPtr<FTraitDataEditorDef>& Trait : *CurrentTraitsDataShared.Get())
+		{
+			if (Trait->TraitUID == TraitUID)
+			{
+				return false;
+			}
+		}
+		return true;
+	}
+	return false;
+}
+
+TSharedRef<ITableRow> STraitListView::HandleGenerateRow(TSharedRef<FTraitListEntry> InEntry, const TSharedRef<STableViewBase>& InOwnerTable)
+{
+	if (InEntry->IsCategory())
+	{
+		return 
+			SNew(SCategoryHeaderTableRow< TSharedPtr<FTraitListEntry> >, InOwnerTable)
+			.Padding(InEntry->IsRoot() ? FMargin(2.0f, 2.0f, 2.0f, 2.0f) : FMargin(InEntry->GetDepthLevel() * 10.0f, 2.0f, 2.0f, 2.0f))
+			[
+				SNew(SHorizontalBox)
+				+SHorizontalBox::Slot()
+				.FillWidth(1.0)
+				.Padding(FMargin(2, 2))
+				[
+					SNew(SHorizontalBox)
+					+ SHorizontalBox::Slot()
+					.VAlign(VAlign_Center)
+					[
+						SNew(STextBlock)
+						.Font(FCoreStyle::GetDefaultFontStyle("Bold", 9))
+						.Text(InEntry->GetCategoryText())
+						.HighlightText_Lambda([this]() { return FilterText; })
+					]
+				]
+			];
+	}
+	else
+	{
+		TWeakPtr<FTraitListEntry> EntryWeak = InEntry.ToWeakPtr();
+		const TSharedPtr<FTraitEditorSharedData>& TraitEditorSharedDataLocal = TraitEditorSharedData;
+		const TSharedPtr<FTraitDataEditorDef>& TraitDataShared = InEntry->GetTraitData();
+		const FTraitUID TraitUID = InEntry->GetTraitUID();
+		const ETraitMode TraitMode = InEntry->GetTraitMode();
+
+		return
+			SNew(STableRow<TSharedRef<FTraitListEntry>>, InOwnerTable)
+			.Padding(FMargin(0.f, 2.f))
+			.ShowSelection(false)
+			.IsEnabled_Lambda([this, TraitUID, TraitMode]()->bool
+			{
+				return IsTraitAvailable(TraitEditorSharedData->CurrentTraitsDataShared, TraitUID, TraitMode, GetSelectedTraitData());
+			})
+			.OnDragDetected_Lambda([TraitDataShared](const FGeometry& MyGeometry, const FPointerEvent& MouseEvent)
+			{
+				if (MouseEvent.IsMouseButtonDown(EKeys::LeftMouseButton))
+				{
+					if (TraitDataShared.IsValid())
+					{
+						TSharedRef<FTraitListDragDropOp> DragDropOp = FTraitListDragDropOp::New(TraitDataShared.ToWeakPtr());
+						return FReply::Handled().BeginDragDrop(DragDropOp);
+					}
+				}
+				return FReply::Unhandled();
+			})
+			.OnCanAcceptDrop_Lambda([](const FDragDropEvent& DragDropEvent, EItemDropZone DropZone, TSharedPtr<FTraitListEntry> TargetItem)
+			{
+				const TOptional<EItemDropZone> InvalidDropZone;
+				return InvalidDropZone;
+			})
+			[
+				SNew(SHorizontalBox)
+				+ SHorizontalBox::Slot()
+				.Padding(FMargin(0.f, 0.f))
+				.FillWidth(1.0)
+				[
+					SNew(SButton)
+					.ButtonStyle(FAppStyle::Get(), "HoverHintOnly")
+					.ContentPadding(FMargin(0.f, 0.f))
+					.ButtonColorAndOpacity_Lambda([this, EntryWeak]()
+					{
+						if (TSharedPtr<FTraitListEntry> Entry = EntryWeak.Pin())
+						{
+							return FTraitEditorUtils::GetTraitBackroundDisplayColor(Entry->GetTraitMode(), /*bIsSelected*/ false, /*bIsHovered*/ false);
+						}
+						return FSlateColor(FColor::Red);
+					})
+					.ClickMethod(EButtonClickMethod::PreciseClick)
+					.OnClicked_Lambda([this, TraitUID]()->FReply
+						{
+							if (OnTraitClicked.IsBound())
+							{
+								return OnTraitClicked.Execute(TraitUID);
+							}
+
+							return FReply::Unhandled();
+						})
+					.Content()
+					[
+						SNew(SBorder)
+						.BorderImage(FAppStyle::Get().GetBrush("Brushes.White"))
+						.BorderBackgroundColor_Lambda([EntryWeak, this]()
+						{
+							if (TSharedPtr<FTraitListEntry> Entry = EntryWeak.Pin())
+							{
+								return FTraitEditorUtils::GetTraitBackroundDisplayColor(Entry->GetTraitMode(), /*bIsSelected*/ false, /*bIsHovered*/ false);
+							}
+							return FSlateColor(FColor::Red);
+						})
+						[
+							SNew(SVerticalBox)
+
+							// --- Trait Main Button ---
+							+ SVerticalBox::Slot()
+							[
+								SNew(SBox)
+								.MinDesiredHeight(25.0f)
+								.VAlign(VAlign_Center)
+								[
+									SNew(SHorizontalBox)
+
+									+ SHorizontalBox::Slot()
+									.AutoWidth()
+									.HAlign(HAlign_Left)
+									[
+										SNew(SBox)
+										.VAlign(VAlign_Top)
+										.HAlign(HAlign_Right)
+										[
+											SNew(SImage)
+											.Image(FAppStyle::Get().GetBrush("Icons.Info.Small"))
+											.ToolTipText_Lambda([EntryWeak, this]()->FText
+											{
+												if (TSharedPtr<FTraitListEntry> Entry = EntryWeak.Pin())
+												{
+													if (const TSharedPtr<FTraitDataEditorDef>& TraitData = Entry->GetTraitData())
+													{
+														// TODO zzz : Create a better looking Tooltip that is not just text based
+														TStringBuilder<1024> TraitInfoString;
+
+														TraitInfoString.Append(TraitData->TraitDisplayName.ToString())
+															.Append(TEXT("\n\n"));
+													
+														if (TraitData->ImplementedInterfaces.Num() > 0)
+														{
+															TraitInfoString.Append(LOCTEXT("TraitInfoImplementedInterfaces", "Implements :").ToString())
+																.AppendChar(TEXT('\n'));
+
+															for (const auto& ImplementedInterface : TraitData->ImplementedInterfaces)
+															{
+																TraitInfoString.Append(TEXT("- "))
+																	.Append(ImplementedInterface.GetInterfaceName())
+																	.AppendChar(TEXT('\n'));
+															}
+														}
+
+														if (TraitData->RequiredInterfaces.Num() > 0)
+														{
+															TraitInfoString.AppendChar(TEXT('\n'));
+															TraitInfoString.Append(LOCTEXT("TraitInfoRequiredInterfaces", "Requires :").ToString())
+															.AppendChar(TEXT('\n'));
+
+															for (const auto& RequiredInterface : TraitData->RequiredInterfaces)
+															{
+																TraitInfoString.Append(TEXT("- "))
+																	.Append(RequiredInterface.GetInterfaceName())
+																	.AppendChar(TEXT('\n'));
+															}
+														}
+
+														return FText::FromString(TraitInfoString.ToString());
+													}
+												}
+												return FText();
+											})
+										]
+									]
+
+									+ SHorizontalBox::Slot()
+									.HAlign(HAlign_Center)
+									.FillWidth(1.f)
+									[
+										SNew(STextBlock)
+										.Text(InEntry->GetTraitNameText())
+										.HighlightText_Lambda([this]() { return FilterText; })
+										.Justification(ETextJustify::Center)
+										.ColorAndOpacity_Lambda([EntryWeak]()
+										{
+											if (TSharedPtr<FTraitListEntry> Entry = EntryWeak.Pin())
+											{
+												return FTraitEditorUtils::GetTraitTextDisplayColor(Entry->GetTraitMode());
+											}
+
+											return FSlateColor(FColor::Red);
+										})
+									]
+								]
+							]
+
+							// --- Trait Implemented Interfaces ---
+							+ SVerticalBox::Slot()
+							.AutoHeight()
+							[
+								FTraitEditorUtils::GetInterfaceListWidget(FTraitEditorUtils::EInterfaceDisplayType::ListImplemented, TraitDataShared, TraitEditorSharedDataLocal)
+							]
+
+						]
+					]
+				]
+			];
+	}
+}
+
+void STraitListView::HandleGetChildren(TSharedRef<FTraitListEntry> InEntry, TArray<TSharedRef<FTraitListEntry>>& OutChildren)
+{
+	if (InEntry->IsCategory())
+	{
+		InEntry->GetChildren(OutChildren);
+	}
+}
+
+void STraitListView::HandleItemScrolledIntoView(TSharedRef<FTraitListEntry> InEntry, const TSharedPtr<ITableRow>& InWidget)
+{
+}
+
+void STraitListView::HandleSelectionChanged(TSharedPtr<FTraitListEntry> InEntry, ESelectInfo::Type InSelectionType)
+{
+}
+
+void STraitListView::GenerateTraitList()
+{
+	BaseTraitCategories.Reset();
+	AdditiveTraitCategories.Reset();
+
+	if (!HasValidEditorData())
+	{
+		return;
+	}
+
+	FTraitRegistry& TraitRegistry = FTraitRegistry::Get();
+
+	const TArray<const FTrait*> Traits = TraitRegistry.GetTraits();
+	for (const FTrait* Trait : Traits)
+	{
+		FName Category = NAME_None;
+		FText CategoryText;
+
+		if (const FString* CategoryMetaData = Trait->GetTraitSharedDataStruct()->FindMetaData(NAME_Category))
+		{
+			Category = **CategoryMetaData;
+			CategoryText = FObjectEditorUtils::GetCategoryText(Trait->GetTraitSharedDataStruct());
+		}
+		else
+		{
+			Category = NAME_DefaultCategory;
+			CategoryText = LOCTEXT("DefaultTraitCategory", "Default");
+		}
+
+		static_assert((uint8)ETraitMode::Num == 2);
+		TMap<FName, FTraitCategoryData>& TraitModeCategories = (Trait->GetTraitMode() == ETraitMode::Base) ? BaseTraitCategories : AdditiveTraitCategories;
+		FTraitCategoryData* TraitCategoryData = TraitModeCategories.Find(Category);
+		if (TraitCategoryData == nullptr)
+		{
+			TraitCategoryData = &TraitModeCategories.Add(Category, FTraitCategoryData(Category, CategoryText));
+		}
+
+		UScriptStruct* TraitSharedDataStruct = Trait->GetTraitSharedDataStruct();
+		
+		const FText TraitDisplayName = (TraitSharedDataStruct != nullptr) ? TraitSharedDataStruct->GetDisplayNameText() : FText::FromString(Trait->GetTraitName());
+		
+		TSharedPtr<FTraitDataEditorDef> TraitData = MakeShared<FTraitDataEditorDef>(*Trait, TraitDisplayName); // TODO zzz : Why TraitName is not an FName as in the Decorator ???
+		FTraitEditorUtils::GenerateStackInterfacesUsedIndexes(TraitData, TraitEditorSharedData);
+
+		TraitCategoryData->TraitList.Add(TraitData); 
+	}
+}
+
+} // end namespace UE::Workspace
+
+#undef LOCTEXT_NAMESPACE
diff --git a/Engine/Plugins/Experimental/AnimNext/Source/AnimNextEditor/Private/Graph/STraitListView.h b/Engine/Plugins/Experimental/AnimNext/Source/AnimNextEditor/Private/Graph/STraitListView.h
new file mode 100644
--- /dev/null
+++ b/Engine/Plugins/Experimental/AnimNext/Source/AnimNextEditor/Private/Graph/STraitListView.h
@@ -0,0 +1,97 @@
+// Copyright Epic Games, Inc. All Rights Reserved.
+
+#pragma once
+
+#include "ContentBrowserDelegates.h"
+#include "Widgets/SCompoundWidget.h"
+#include "Widgets/Views/STreeView.h"
+#include "Widgets/Views/STableViewBase.h"
+#include "Widgets/Views/ITableRow.h"
+#include "Graph/TraitEditorDefs.h"
+
+class UAnimNextGraph_EdGraphNode;
+
+namespace UE::Workspace
+{
+class IWorkspaceEditor;
+}
+
+
+namespace UE::AnimNext::Editor
+{
+
+struct FTraitListEntry;
+struct FTraitListCategoryEntry;
+
+class STraitListView : public SCompoundWidget
+{
+public:
+	/** Used to notify the Trait Editor a Trait has been clicked on the available list */
+	DECLARE_DELEGATE_RetVal_OneParam(FReply, FOnTraitClicked, const FTraitUID);
+
+	/** Used to obtain the currently selected Trait in the Stack */
+	DECLARE_DELEGATE_RetVal(TWeakPtr<FTraitDataEditorDef>, FOnGetSelectedTraitData);
+
+
+	SLATE_BEGIN_ARGS(STraitListView) {}
+
+	SLATE_EVENT(FOnTraitClicked, OnTraitClicked)
+
+	SLATE_EVENT(FOnGetSelectedTraitData, OnGetSelectedTraitData)
+
+	SLATE_END_ARGS()
+
+	void Construct(const FArguments& InArgs, TSharedPtr<FTraitEditorSharedData>& InTraitEditorSharedData);
+
+	void RefreshList();
+
+private:
+	// SWidget interface
+	virtual FReply OnKeyDown(const FGeometry& MyGeometry, const FKeyEvent& InKeyEvent) override;
+
+	//void HandleWorkspaceModified(UWorkspace* InWorkspace);
+
+	void OnFilterTextChanged(const FText& InFilterText);
+
+	TSharedPtr<FTraitEditorSharedData> TraitEditorSharedData;
+
+	TMap<FName, FTraitCategoryData> BaseTraitCategories;
+	TMap<FName, FTraitCategoryData> AdditiveTraitCategories;
+
+	TSharedPtr<FUICommandList> UICommandList;
+
+	TSharedPtr< SSearchBox > TraitListFilterBox;
+	TSharedPtr<STreeView<TSharedRef<FTraitListEntry>>> EntriesList;
+	TArray<TSharedRef<FTraitListEntry>> Categories;
+
+	FText FilterText;
+	TArray<TSharedRef<FTraitListEntry>> FilteredEntries;
+
+	TSet< TSharedRef<FTraitListEntry> > OldExpansionState;
+
+	FOnTraitClicked OnTraitClicked;
+	FOnGetSelectedTraitData OnGetSelectedTraitData;
+
+	bool bIgnoreExpansion = false;
+
+	// --- ---
+	bool HasValidEditorData() const;
+	TSharedPtr<FTraitDataEditorDef> GetSelectedTraitData() const;
+
+	void StoreExpansionState();
+	void RestoreExpansionState(TArray<TSharedRef<FTraitListEntry>>& AllEntries);
+	void GetAllEntries(TArray<TSharedRef<FTraitListEntry>>& AllEntries) const;
+	void ExpandAllCategories(const TArray<TSharedRef<FTraitListEntry>>& AllEntries);
+	void RefreshEntries();
+	void CreateTraitCategories(const FName& BaseCategoryName, const FText& BaseCategoryText, TMap<FName, FTraitCategoryData>& CategoriesMap);
+	void RefreshFilter();
+
+	TSharedRef<ITableRow> HandleGenerateRow(TSharedRef<FTraitListEntry> InEntry, const TSharedRef<STableViewBase>& InOwnerTable);
+	void HandleGetChildren(TSharedRef<FTraitListEntry> InEntry, TArray<TSharedRef<FTraitListEntry>>& OutChildren);
+	void HandleItemScrolledIntoView(TSharedRef<FTraitListEntry> InEntry, const TSharedPtr<ITableRow>& InWidget);
+	void HandleSelectionChanged(TSharedPtr<FTraitListEntry> InEntry, ESelectInfo::Type InSelectionType);
+
+	void GenerateTraitList();
+};
+
+};
diff --git a/Engine/Plugins/Experimental/AnimNext/Source/AnimNextEditor/Private/Graph/STraitStackView.cpp b/Engine/Plugins/Experimental/AnimNext/Source/AnimNextEditor/Private/Graph/STraitStackView.cpp
new file mode 100644
--- /dev/null
+++ b/Engine/Plugins/Experimental/AnimNext/Source/AnimNextEditor/Private/Graph/STraitStackView.cpp
@@ -0,0 +1,634 @@
+// Copyright Epic Games, Inc. All Rights Reserved.
+
+#include "STraitStackView.h"
+#include "ContentBrowserModule.h"
+#include "Framework/Commands/GenericCommands.h"
+#include "Graph/AnimNextGraph_EdGraphNode.h"
+#include "IContentBrowserSingleton.h"
+#include "RigVMModel/RigVMNode.h"
+#include "SAssetDropTarget.h"
+#include "ScopedTransaction.h"
+#include "TraitCore/Trait.h"
+#include "TraitCore/TraitRegistry.h"
+#include "Widgets/Views/SListView.h"
+#include "WorkspaceSchema.h"
+#include "Styling/AppStyle.h"
+#include "Widgets/Text/SInlineEditableTextBlock.h"
+#include "Layout/Margin.h"
+#include "DragAndDrop/DecoratedDragDropOp.h"
+#include "DragAndDrop/AssetDragDropOp.h"
+#include "Widgets/Input/SButton.h"
+
+#define LOCTEXT_NAMESPACE "TraitListEditor"
+
+namespace UE::AnimNext::Editor
+{
+
+// --- FTraitStackViewEntry ---
+
+struct FTraitStackViewEntry
+{
+	FTraitStackViewEntry() = default;
+
+	explicit FTraitStackViewEntry(const TSharedPtr<FTraitDataEditorDef>& InTraitDataEditorDef)
+		: TraitData(InTraitDataEditorDef)
+	{
+	}
+
+	TSharedPtr<FTraitDataEditorDef> TraitData;
+};
+
+// --- FTraitStackDragDropOp ---
+class FTraitStackDragDropOp : public FTraitListDragDropBase
+{
+public:
+	DRAG_DROP_OPERATOR_TYPE(FTraitStackDragDropOp, FDecoratedDragDropOp)
+
+	static TSharedRef<FTraitStackDragDropOp> New(TWeakPtr<FTraitDataEditorDef> InDraggedTraitDataWeak)
+	{
+		TSharedRef<FTraitStackDragDropOp> Operation = MakeShared<FTraitStackDragDropOp>();
+		Operation->DraggedTraitDataWeak = InDraggedTraitDataWeak;
+		Operation->Construct();
+		return Operation;
+	}
+};
+
+
+// --- STraitStackView ---
+
+STraitStackView::STraitStackView()
+	: EntriesList(MakeShared<SListView<TSharedRef<FTraitStackViewEntry>>>())
+{
+}
+
+void STraitStackView::Construct(const FArguments& InArgs, TSharedPtr<FTraitEditorSharedData>& InTraitEditorSharedData)
+{
+	OnTraitDeleteRequest = InArgs._OnTraitDeleteRequest;
+	OnStatckTraitSelectionChanged = InArgs._OnStatckTraitSelectionChanged;
+	OnStackTraitDragAccepted = InArgs._OnStackTraitDragAccepted;
+
+	TraitEditorSharedData = InTraitEditorSharedData;
+
+	UICommandList = MakeShared<FUICommandList>();
+
+	UICommandList->MapAction(FGenericCommands::Get().Delete,
+		FExecuteAction::CreateSP(this, &STraitStackView::HandleDelete),
+		FCanExecuteAction::CreateSP(this, &STraitStackView::HasValidSelection));
+
+	const TSharedPtr<FTraitEditorSharedData>& TraitEditorSharedDataLocal = TraitEditorSharedData;
+	const TWeakObjectPtr<UAnimNextGraph_EdGraphNode> EdGraphNodeWeak = TraitEditorSharedData->EdGraphNodeWeak;
+
+	ChildSlot
+	[
+		SNew(SVerticalBox)
+		+ SVerticalBox::Slot()
+		.FillHeight(1.0)
+		[
+			SNew(SBorder)
+			.Padding(4.0f)
+			.BorderImage(FAppStyle::GetBrush("ToolPanel.GroupBorder"))
+			[
+				SNew(SVerticalBox)
+
+				+ SVerticalBox::Slot()
+				.AutoHeight()
+				[
+					SNew(SHorizontalBox)
+					+ SHorizontalBox::Slot()
+					.AutoWidth()
+					[
+						SNew(STextBlock)
+						.Text(LOCTEXT("TraitStackNodeName", "Node Name :"))
+					]
+					+ SHorizontalBox::Slot()
+					.Padding(40.f, 0.f, 0.f, 0.f)
+					.FillWidth(1.f)
+					[
+						SNew(SBorder)
+						.BorderImage(FAppStyle::Get().GetBrush("Brushes.Black"))
+						.Padding(FMargin(1.0f, 1.0f))
+						[
+							SNew(SInlineEditableTextBlock)
+							.Style(FAppStyle::Get(), "Graph.Node.NodeTitleInlineEditableText")
+							.Text_Lambda([this] { return TraitEditorSharedData->EdGraphNodeWeak.IsValid() ? TraitEditorSharedData->EdGraphNodeWeak->GetNodeTitle(ENodeTitleType::EditableTitle) : FText::GetEmpty(); })
+							.IsReadOnly_Lambda([EdGraphNodeWeak]()->bool
+							{
+								bool bIsReadOnly = false;
+								if (UAnimNextGraph_EdGraphNode* EdGraphNode = Cast<UAnimNextGraph_EdGraphNode>(EdGraphNodeWeak.Get()))
+								{
+									bIsReadOnly = EdGraphNode->IsDeprecated() || EdGraphNode->IsOutDated();
+								}
+								return bIsReadOnly;
+							})
+							.OnTextCommitted_Lambda([EdGraphNodeWeak](const FText& InNewText, ETextCommit::Type InCommitType)
+							{
+								if(InCommitType == ETextCommit::OnEnter)
+								{
+									if(UAnimNextGraph_EdGraphNode* EdGraphNode = Cast<UAnimNextGraph_EdGraphNode>(EdGraphNodeWeak.Get()))
+									{
+										// TODO zzz : Find the correct way to change the DisplayName / Title
+										//FScopedTransaction Transaction(LOCTEXT("SetNodeName", "Set Node title"));
+										//EdGraphNode->SetNodeTitle(InNewText);
+									}
+								}
+							})
+						]
+					]
+				]
+
+				+ SVerticalBox::Slot()
+				.FillHeight(1.0)
+				[
+					SNew(SBorder)
+					.Visibility(EVisibility::Visible)
+					.BorderImage(FAppStyle::GetBrush("Menu.Background"))
+					[
+						SNew(SVerticalBox)
+						+ SVerticalBox::Slot()
+						.FillHeight(1.0)
+						[
+							SAssignNew(EntriesList, SListView<TSharedRef<FTraitStackViewEntry>>)
+							.ListItemsSource(&Entries)
+							.ItemHeight_Lambda([TraitEditorSharedDataLocal]()
+								{
+									constexpr const float ButtonSize = 20.f;
+									constexpr const float ButtonSizeWithInterfaces = 40.f;
+
+									if (TraitEditorSharedDataLocal.IsValid() && TraitEditorSharedDataLocal->bShowTraitInterfaces)
+									{
+										return ButtonSizeWithInterfaces; 
+									}
+									return ButtonSize;
+								})
+							.OnGenerateRow(this, &STraitStackView::HandleGenerateRow)
+							.SelectionMode(ESelectionMode::SingleToggle)
+							.OnSelectionChanged_Lambda([this](TSharedPtr<FTraitStackViewEntry> InEntry, ESelectInfo::Type InSelectionType)
+							{
+								if (InEntry.IsValid())
+								{
+									SelectedTraitData = InEntry->TraitData;
+								}
+								else
+								{
+									SelectedTraitData.Reset();
+								}
+
+								if (OnStatckTraitSelectionChanged.IsBound())
+								{
+									OnStatckTraitSelectionChanged.Execute(GetSelectedTraitUID());
+								}
+							})
+						]
+					]
+				]
+			]
+		]
+	];
+}
+
+void STraitStackView::RefreshList()
+{
+	if (TraitEditorSharedData.IsValid() && TraitEditorSharedData->CurrentTraitsDataShared.IsValid())
+	{
+		const TArray<TSharedPtr<FTraitDataEditorDef>>& CurrentTraitsData = *(TraitEditorSharedData->CurrentTraitsDataShared.Get());
+		
+		Entries.Reset(CurrentTraitsData.Num());
+
+		for (const TSharedPtr<FTraitDataEditorDef>& TraitData : CurrentTraitsData)
+		{
+			Entries.Add(MakeShared<FTraitStackViewEntry>(TraitData));
+		}
+	}
+	else
+	{
+		Entries.Reset();
+	}
+
+	EntriesList->RebuildList();
+}
+
+FReply STraitStackView::OnKeyDown(const FGeometry& MyGeometry, const FKeyEvent& InKeyEvent)
+{
+	if (UICommandList.IsValid() && UICommandList->ProcessCommandBindings(InKeyEvent))
+	{
+		return FReply::Handled();
+	}
+	return FReply::Unhandled();
+}
+
+void STraitStackView::HandleDelete()
+{
+	const FTraitUID SelectedTraitUID = GetSelectedTraitUID();
+	if (SelectedTraitUID != FTraitUID())
+	{
+		ExecuteDelete(SelectedTraitUID);
+	}
+}
+
+bool STraitStackView::HasValidSelection() const 
+{
+	const FTraitUID SelectedTraitUID = GetSelectedTraitUID();
+
+	if (SelectedTraitUID != FTraitUID())
+	{
+		return true;
+	}
+
+	return false;
+}
+
+// Helper to get direct access to the Paint delegate
+class STraitStackTableRow : public STableRow<TSharedRef<FTraitStackViewEntry>>
+{
+public:
+	/** Optional delegate for painting drop indicators */
+	FOnPaintDropIndicator GetOnPaintDropIndicatorDelegate() { return PaintDropIndicatorEvent; }
+};
+
+TSharedRef<ITableRow> STraitStackView::HandleGenerateRow(TSharedRef<FTraitStackViewEntry> InEntry, const TSharedRef<STableViewBase>& InOwnerTable)
+{
+	const TWeakPtr<FTraitStackViewEntry> InEntryWeak = InEntry.ToWeakPtr();
+	const TSharedPtr<FTraitEditorSharedData> TraitEditorSharedDataLocal = TraitEditorSharedData;
+	const TSharedPtr<FTraitDataEditorDef>& TraitDataShared = InEntry->TraitData;
+
+	TSharedPtr<STraitStackTableRow> Row;
+
+	SAssignNew(Row, STraitStackTableRow, InOwnerTable)
+		.Padding(FMargin(4.f, 2.f))
+		.OnDragDetected_Lambda([InEntryWeak](const FGeometry& MyGeometry, const FPointerEvent& MouseEvent)
+			{
+				if (MouseEvent.IsMouseButtonDown(EKeys::LeftMouseButton))
+				{
+					if (TSharedPtr<FTraitStackViewEntry> Entry = InEntryWeak.Pin())
+					{
+						if (Entry->TraitData.IsValid())
+						{
+							//TSharedRef<FTraitStackDragDropOp> DragDropOp = FTraitStackDragDropOp::New(Entry);
+							TSharedRef<FTraitStackDragDropOp> DragDropOp = FTraitStackDragDropOp::New(Entry->TraitData.ToWeakPtr());
+							return FReply::Handled().BeginDragDrop(DragDropOp);
+						}
+					}
+				}
+				return FReply::Unhandled();
+			})
+		.OnCanAcceptDrop_Lambda([TraitEditorSharedDataLocal](const FDragDropEvent& DragDropEvent, EItemDropZone DropZone, TSharedPtr<FTraitStackViewEntry> TargetItem)
+			{
+				const TOptional<EItemDropZone> InvalidDropZone;
+				TOptional<EItemDropZone> ReturnDropZone = InvalidDropZone;
+
+				if (!TargetItem.IsValid())
+				{
+					return ReturnDropZone;
+				}
+
+				// --- Drops from the Trait List ---
+				if (TSharedPtr<FTraitListDragDropOp> TraitListDragDropOp = DragDropEvent.GetOperationAs<FTraitListDragDropOp>())
+				{
+					if (TSharedPtr<FTraitDataEditorDef> DraggedEntryTraitData = TraitListDragDropOp->GetDraggedTraitData().Pin(); DraggedEntryTraitData.IsValid())
+					{
+						const ETraitMode DraggedTraitMode = DraggedEntryTraitData->TraitMode;
+						const ETraitMode TargetTraitMode = TargetItem->TraitData->TraitMode;
+
+						if (TargetTraitMode == ETraitMode::Base)
+						{
+							if (DraggedTraitMode == ETraitMode::Additive)
+							{
+								ReturnDropZone = EItemDropZone::BelowItem;	// Force BelowItem for additive traits
+							}
+							else
+							{
+								ReturnDropZone = EItemDropZone::OntoItem; // Force OntoItem independently of the zone for base traits
+							}
+						}
+						else
+						{
+							if (DraggedTraitMode == ETraitMode::Base)
+							{
+								ReturnDropZone = InvalidDropZone;
+							}
+							else
+							{
+								ReturnDropZone = DropZone == EItemDropZone::AboveItem ? EItemDropZone::OntoItem : DropZone; // For additives disalow AboveItem
+							}
+						}
+					}
+				}
+				// --- Drops from the Trait Stack itself (rearrange items in the stack) ---
+				else if (TSharedPtr<FTraitStackDragDropOp> TraitStackDragDropOp = DragDropEvent.GetOperationAs<FTraitStackDragDropOp>())
+				{
+					if (DropZone == EItemDropZone::OntoItem)
+					{
+						return InvalidDropZone;
+					}
+					
+					if (TSharedPtr<FTraitDataEditorDef> DraggedEntryTraitData = TraitStackDragDropOp->GetDraggedTraitData().Pin(); DraggedEntryTraitData.IsValid())
+					{
+						// Base can not be dropped anywhere
+						if (DraggedEntryTraitData->TraitMode == ETraitMode::Base)
+						{
+							return InvalidDropZone;
+						}
+
+						// An additive can only be dropped below a base
+						if (TargetItem->TraitData->TraitMode == ETraitMode::Base && DropZone != EItemDropZone::BelowItem)
+						{
+							return InvalidDropZone;
+						}
+
+						int32 TargetTraitIndex = INDEX_NONE;
+						FTraitEditorUtils::FindTraitInCurrentStackData(TargetItem->TraitData->TraitUID, TraitEditorSharedDataLocal->CurrentTraitsDataShared, &TargetTraitIndex);
+						int32 DraggedTraitIndex = INDEX_NONE;
+						FTraitEditorUtils::FindTraitInCurrentStackData(DraggedEntryTraitData->TraitUID, TraitEditorSharedDataLocal->CurrentTraitsDataShared, &DraggedTraitIndex);
+
+						if (TargetTraitIndex == INDEX_NONE || DraggedTraitIndex == INDEX_NONE)
+						{
+							return InvalidDropZone;
+						}
+
+						const int32 IndexDiff = TargetTraitIndex - DraggedTraitIndex;
+						const int32 AbsIndexDiff = FMath::Abs(IndexDiff);
+						if (AbsIndexDiff < 1 // Can not drop on self
+							|| DropZone == EItemDropZone::AboveItem // Only BelowItem is allowed
+							|| (AbsIndexDiff == 1 && IndexDiff < 0))// And on neighbors only if it is the down neighbor
+						{
+							return InvalidDropZone;
+						}
+
+						ReturnDropZone = DropZone;
+					}
+				}
+
+				return ReturnDropZone;
+			})
+		.OnAcceptDrop_Lambda([this](const FDragDropEvent& DragDropEvent, EItemDropZone DropZone, TSharedPtr<FTraitStackViewEntry> TargetItem)->FReply
+			{
+				if(TargetItem.IsValid() && TargetItem->TraitData.IsValid())
+				{
+					TSharedPtr<FTraitDataEditorDef> DraggedEntryTraitData;
+
+					if (TSharedPtr<FTraitStackDragDropOp> TraitStackDragDropOp = DragDropEvent.GetOperationAs<FTraitStackDragDropOp>())
+					{
+						DraggedEntryTraitData = TraitStackDragDropOp->GetDraggedTraitData().Pin();
+					}
+					else if (TSharedPtr<FTraitListDragDropOp> TraitListDragDropOp = DragDropEvent.GetOperationAs<FTraitListDragDropOp>())
+					{
+						DraggedEntryTraitData = TraitListDragDropOp->GetDraggedTraitData().Pin();
+					}
+
+					if (DraggedEntryTraitData.IsValid())
+					{
+						if (OnStackTraitDragAccepted.IsBound())
+						{
+							return OnStackTraitDragAccepted.Execute(DraggedEntryTraitData->TraitUID, TargetItem->TraitData->TraitUID, DropZone);
+						}
+					}
+				}
+
+				return FReply::Unhandled();
+			})
+		[
+			SNew(SHorizontalBox)
+
+			+ SHorizontalBox::Slot()
+			[
+				SNew(SBorder)
+				.BorderImage(FAppStyle::Get().GetBrush("Brushes.White"))
+				.BorderBackgroundColor_Lambda([InEntryWeak, this]()
+				{
+					if (TSharedPtr<FTraitStackViewEntry> Entry = InEntryWeak.Pin())
+					{
+						const FTraitUID SelectedTraitUID = GetSelectedTraitUID();
+						const bool bIsSelected = SelectedTraitUID != FTraitUID() && Entry->TraitData.IsValid() && Entry->TraitData->TraitUID == SelectedTraitUID;
+
+						return FTraitEditorUtils::GetTraitBackroundDisplayColor(Entry->TraitData->TraitMode, bIsSelected);
+					}
+					return FSlateColor(FColor::Red);
+				})
+				.Padding(FMargin(1.0f, 1.0f))
+				[
+					SNew(SVerticalBox)
+
+						+ SVerticalBox::Slot()
+						.AutoHeight()
+						[
+							FTraitEditorUtils::GetInterfaceListWidget(FTraitEditorUtils::EInterfaceDisplayType::StackRequired, TraitDataShared, TraitEditorSharedDataLocal)
+						]
+
+						+ SVerticalBox::Slot()
+						[
+							SNew(SBox)
+							//.Padding(30.f, 10.f)
+							.MinDesiredHeight(30.0f)
+							.VAlign(VAlign_Center)
+							[
+								GetStackListItemWidget(InEntryWeak, TraitEditorSharedDataLocal)
+							]
+						]
+
+						+ SVerticalBox::Slot()
+						.AutoHeight()
+						[
+							FTraitEditorUtils::GetInterfaceListWidget(FTraitEditorUtils::EInterfaceDisplayType::StackImplemented, TraitDataShared, TraitEditorSharedDataLocal)
+						]
+				]
+			]
+		];
+
+		const TWeakPtr<STraitStackTableRow> RowWeak = Row.ToWeakPtr();
+		
+		// Use direct access to the paint delegate, in order to be able to pass the Row to the lambda
+		Row->GetOnPaintDropIndicatorDelegate().BindLambda([RowWeak](EItemDropZone InItemDropZone, const FPaintArgs& Args, const FGeometry& AllottedGeometry, const FSlateRect& MyCullingRect, FSlateWindowElementList& OutDrawElements, int32 LayerId, const FWidgetStyle& InWidgetStyle, bool bParentEnabled)
+			{
+				if (const TSharedPtr<STraitStackTableRow> RowShared = RowWeak.Pin())
+				{
+					const FSlateBrush* DropIndicatorBrush = RowShared->GetDropIndicatorBrush(InItemDropZone);
+					static float OffsetX = 10.0f;
+					FVector2D Offset(OffsetX * RowShared->GetIndentLevel(), 0.0f);
+
+					FSlateDrawElement::MakeBox(
+						OutDrawElements,
+						LayerId++,
+						AllottedGeometry.ToPaintGeometry(FVector2D(AllottedGeometry.GetLocalSize() - Offset), FSlateLayoutTransform(Offset)),
+						DropIndicatorBrush,
+						ESlateDrawEffect::None,
+						DropIndicatorBrush->GetTint(InWidgetStyle) * InWidgetStyle.GetColorAndOpacityTint()
+					);
+				}
+
+				return LayerId;
+			});
+
+	return Row.ToSharedRef();
+}
+
+TSharedRef<SWidget> STraitStackView::GetStackListItemWidget(const TWeakPtr<FTraitStackViewEntry>& InEntryWeak, const TSharedPtr<FTraitEditorSharedData>& TraitEditorSharedDataLocal)
+{
+	return SNew(SHorizontalBox)
+			+ SHorizontalBox::Slot()
+			[
+				SNew(SHorizontalBox)
+
+				+ SHorizontalBox::Slot()
+				.AutoWidth()
+				.HAlign(HAlign_Left)
+				[
+					SNew(SBox)
+					.MaxDesiredHeight(20.0f)
+					.MaxDesiredWidth(20.0f)
+					.VAlign(VAlign_Center)
+					.HAlign(HAlign_Left)
+					.Padding(FMargin(5.f, 0.f, 0.f, 0.f))
+					[
+						SNew(SImage)
+						.Image(FAppStyle::Get().GetBrush("Icons.Error"))
+						.ColorAndOpacity_Lambda([InEntryWeak]()
+						{
+							if (TSharedPtr<FTraitStackViewEntry> Entry = InEntryWeak.Pin())
+							{
+								return FTraitEditorUtils::GetTraitIconErrorDisplayColor(Entry->TraitData->StackStatus);
+							}
+
+							return FSlateColor(FColor::Red);
+						})
+						.ToolTipText_Lambda([InEntryWeak]()->FText
+						{
+							if (TSharedPtr<FTraitStackViewEntry> Entry = InEntryWeak.Pin())
+							{
+								if (Entry->TraitData->StackStatus.TraitStatus != FTraitStackTraitStatus::EStackStatus::Ok)
+								{
+									TStringBuilder<1024> ErrorMessage;
+									for (const auto& Status : Entry->TraitData->StackStatus.StatusMessages)
+									{
+										ErrorMessage.Append(Status.MessageText.ToString()).AppendChar(TEXT('\n'));
+									}
+									return FText::FromString(ErrorMessage.ToString());
+								}
+							}
+							return FText::GetEmpty();
+						})
+						.Visibility_Lambda([InEntryWeak]()->EVisibility
+						{
+							if (TSharedPtr<FTraitStackViewEntry> Entry = InEntryWeak.Pin())
+							{
+								if (Entry->TraitData->StackStatus.TraitStatus != FTraitStackTraitStatus::EStackStatus::Ok)
+								{
+									return EVisibility::Visible;
+								}
+							}
+							return EVisibility::Hidden;
+						})
+					]
+				]
+
+				+ SHorizontalBox::Slot()
+				.FillWidth(1.f)
+				.HAlign(HAlign_Center)
+				[
+					SNew(STextBlock)
+					.Justification(ETextJustify::Center)
+					.Text(InEntryWeak.Pin()->TraitData->TraitDisplayName)
+					.ColorAndOpacity_Lambda([InEntryWeak]()
+					{
+						if (TSharedPtr<FTraitStackViewEntry> Entry = InEntryWeak.Pin())
+						{
+							return FTraitEditorUtils::GetTraitTextDisplayColor(Entry->TraitData->TraitMode);
+						}
+
+						return FSlateColor(FColor::Red);
+					})
+				]
+
+				+ SHorizontalBox::Slot()
+				.AutoWidth()
+				.HAlign(HAlign_Right)
+				[
+					SNew(SBox)
+					.MaxDesiredHeight(20.0f)
+					.MaxDesiredWidth(20.0f)
+					.VAlign(VAlign_Top)
+					.HAlign(HAlign_Right)
+					[
+						SNew(SButton)
+						.ClickMethod(EButtonClickMethod::MouseUp)
+						.ButtonStyle(FAppStyle::Get(), "HoverHintOnly")
+						.Visibility_Lambda([InEntryWeak, TraitEditorSharedDataLocal]()
+						{
+							if (TSharedPtr<FTraitStackViewEntry> Entry = InEntryWeak.Pin())
+							{
+								if (TraitEditorSharedDataLocal.IsValid() && TraitEditorSharedDataLocal->CurrentTraitsDataShared.IsValid())
+								{
+									// If the user has deleted the base but there are still additive Traits in the stack
+									if (Entry->TraitData->TraitMode == ETraitMode::Base 
+										&& TraitEditorSharedDataLocal->CurrentTraitsDataShared->Num() > 1
+										&& (Entry->TraitData->StackStatus.TraitStatus == FTraitStackTraitStatus::EStackStatus::Invalid)
+										&& Entry->TraitData->TraitUID == FTraitUID())
+									{
+										return EVisibility::Hidden;
+									}
+								}
+							}
+							return EVisibility::Visible;
+						})
+						.OnClicked_Lambda([InEntryWeak, this]()->FReply
+						{
+							FReply Reply = FReply::Unhandled();
+
+							if (TSharedPtr<FTraitStackViewEntry> Entry = InEntryWeak.Pin())
+							{
+								Reply = ExecuteDelete(Entry->TraitData->TraitUID);
+							}
+
+							return Reply;
+						})
+						.Content()
+						[
+							SNew(SImage)
+							.Image(FAppStyle::Get().GetBrush("Icons.Delete"))
+						]
+					]
+				]
+			];
+
+}
+
+const TSharedPtr<FTraitDataEditorDef>& STraitStackView::GetSelectedTraitData() const
+{
+	return SelectedTraitData;
+}
+
+const FTraitUID STraitStackView::GetSelectedTraitUID() const
+{
+	FTraitUID TraitUID;
+
+	if (const TSharedPtr<FTraitDataEditorDef> TraitDataEditorData = GetSelectedTraitData())
+	{
+		if (TraitDataEditorData.IsValid())
+		{
+			TraitUID = TraitDataEditorData->TraitUID;
+		}
+	}
+
+	return TraitUID;
+}
+
+FReply STraitStackView::ExecuteDelete(FTraitUID TraitUID)
+{
+	FReply Reply = FReply::Unhandled();
+
+	if (OnTraitDeleteRequest.IsBound())
+	{
+		Reply = OnTraitDeleteRequest.Execute(TraitUID);
+	}
+
+	SelectedTraitData.Reset();
+	if (OnStatckTraitSelectionChanged.IsBound())
+	{
+		OnStatckTraitSelectionChanged.Execute(FTraitUID());
+	}
+
+	return Reply;
+}
+
+} // end namespace UE::Workspace
+
+#undef LOCTEXT_NAMESPACE
diff --git a/Engine/Plugins/Experimental/AnimNext/Source/AnimNextEditor/Private/Graph/STraitStackView.h b/Engine/Plugins/Experimental/AnimNext/Source/AnimNextEditor/Private/Graph/STraitStackView.h
new file mode 100644
--- /dev/null
+++ b/Engine/Plugins/Experimental/AnimNext/Source/AnimNextEditor/Private/Graph/STraitStackView.h
@@ -0,0 +1,89 @@
+// Copyright Epic Games, Inc. All Rights Reserved.
+
+#pragma once
+
+#include "ContentBrowserDelegates.h"
+#include "TraitCore/TraitUID.h"
+#include "Widgets/SCompoundWidget.h"
+#include "Widgets/Views/STableViewBase.h"
+#include "Widgets/Views/STableRow.h"
+#include "TraitCore/TraitMode.h"
+#include "Graph/TraitEditorDefs.h"
+
+class UAnimNextGraph_EdGraphNode;
+
+namespace UE::Workspace
+{
+
+class IWorkspaceEditor;
+
+}
+
+namespace UE::AnimNext::Editor
+{
+
+struct FTraitStackViewEntry;
+
+
+class STraitStackView : public SCompoundWidget
+{
+public:
+	/** Used to notify the Trait Editor a Trait has been selected on the Stack */
+	DECLARE_DELEGATE_OneParam(FOnStatckTraitSelectionChanged, const FTraitUID /*SelectedUID*/);
+
+	/** Used to notify the Trait Editor a request to delete a Trait */
+	DECLARE_DELEGATE_RetVal_OneParam(FReply, FOnStackTraitDeleteRequest, const FTraitUID /*DeleteUID*/);
+
+	/** Used to notify the Trait Editor a Trait Drag has been accepted on the Stack */
+	DECLARE_DELEGATE_RetVal_ThreeParams(FReply, FOnStackTraitDragAccepted, const FTraitUID /*DraggedUID*/, const FTraitUID /*TargetUID*/, EItemDropZone /*DropZone*/);
+
+
+	STraitStackView();
+
+	SLATE_BEGIN_ARGS(STraitStackView) {}
+
+	/** Called to notify a click on the delete button */
+	SLATE_EVENT(FOnStackTraitDeleteRequest, OnTraitDeleteRequest)
+
+	/** Called to notify a trait has been selected */
+	SLATE_EVENT(FOnStatckTraitSelectionChanged, OnStatckTraitSelectionChanged)
+
+	/** Called to notify a trait has been dragged onto the stack */
+	SLATE_EVENT(FOnStackTraitDragAccepted, OnStackTraitDragAccepted)
+
+	SLATE_END_ARGS()
+
+	void Construct(const FArguments& InArgs, TSharedPtr<FTraitEditorSharedData>& InTraitEditorSharedData);
+
+	void RefreshList();
+
+private:
+
+	// SWidget interface
+	virtual FReply OnKeyDown(const FGeometry& MyGeometry, const FKeyEvent& InKeyEvent) override;
+
+	void HandleDelete();
+	bool HasValidSelection() const;
+
+	TSharedRef<ITableRow> HandleGenerateRow(TSharedRef<FTraitStackViewEntry> InEntry, const TSharedRef<STableViewBase>& InOwnerTable);
+	TSharedRef<SWidget> GetStackListItemWidget(const TWeakPtr<FTraitStackViewEntry>& InEntryWeak, const TSharedPtr<FTraitEditorSharedData>& TraitEditorSharedDataLocal);
+
+	const TSharedPtr<FTraitDataEditorDef>& GetSelectedTraitData() const;
+	const FTraitUID GetSelectedTraitUID() const;
+	
+	FReply ExecuteDelete(FTraitUID TraitUID);
+
+	TSharedPtr<FUICommandList> UICommandList;
+
+	TSharedPtr<FTraitDataEditorDef> SelectedTraitData;
+	TSharedPtr<FTraitEditorSharedData> TraitEditorSharedData;
+
+	FOnStackTraitDeleteRequest OnTraitDeleteRequest;
+	FOnStatckTraitSelectionChanged OnStatckTraitSelectionChanged;
+	FOnStackTraitDragAccepted OnStackTraitDragAccepted;
+
+	TSharedRef<SListView<TSharedRef<FTraitStackViewEntry>>> EntriesList;
+	TArray<TSharedRef<FTraitStackViewEntry>> Entries;
+};
+
+};
diff --git a/Engine/Plugins/Experimental/AnimNext/Source/AnimNextEditor/Private/Graph/TraitEditorDefs.cpp b/Engine/Plugins/Experimental/AnimNext/Source/AnimNextEditor/Private/Graph/TraitEditorDefs.cpp
new file mode 100644
--- /dev/null
+++ b/Engine/Plugins/Experimental/AnimNext/Source/AnimNextEditor/Private/Graph/TraitEditorDefs.cpp
@@ -0,0 +1,335 @@
+// Copyright Epic Games, Inc. All Rights Reserved.
+
+#include "Graph/TraitEditorDefs.h"
+#include "TraitInterfaces/IGarbageCollection.h"
+#include "Widgets/Layout/SBox.h"
+
+#define LOCTEXT_NAMESPACE "TraitListEditor"
+
+namespace UE::AnimNext::Editor
+{
+
+/*static*/ FSlateColor FTraitEditorUtils::GetTraitIconErrorDisplayColor(const FTraitStackTraitStatus& InTraitStatus)
+{
+	switch (InTraitStatus.TraitStatus)
+	{
+		case FTraitStackTraitStatus::EStackStatus::Error:
+		case FTraitStackTraitStatus::EStackStatus::Invalid:
+		{
+			return FSlateColor(FColor::Red);
+			break;
+		}
+		case FTraitStackTraitStatus::EStackStatus::Warning:
+		{
+			return FSlateColor(FColor::Yellow);
+			break;
+		}
+		default:
+		{
+			break;
+		}
+	}
+
+	return FSlateColor::UseForeground();
+}
+
+/*static*/ FSlateColor FTraitEditorUtils::GetTraitTextDisplayColor(const ETraitMode TraitMode)
+{
+	switch (TraitMode)
+	{
+		case ETraitMode::Base: return FSlateColor(FColor::White);
+		case ETraitMode::Additive: return FSlateColor(FColor::White);
+		default:
+		{
+			break;
+		}
+	}
+
+	return FSlateColor(FSlateColor::UseSubduedForeground());
+}
+
+/*static*/ FSlateColor FTraitEditorUtils::GetTraitBackroundDisplayColor(const ETraitMode TraitMode, bool bIsSelected, bool bIsHovered)
+{
+	static const FLinearColor SelectedColor(FColor::FromHex(TEXT("#5555ff")));
+	static const FLinearColor BaseColor(FColor::FromHex(TEXT("#505050")));
+	static const FLinearColor BaseColorHovered(FColor::FromHex(TEXT("#555555")));
+	static const FLinearColor AdditiveColor(FColor::FromHex(TEXT("#707070")));
+	static const FLinearColor AdditiveColorHovered(FColor::FromHex(TEXT("#757575")));
+
+	if (bIsSelected)
+	{
+		return SelectedColor;
+	}
+
+	switch (TraitMode)
+	{
+		case ETraitMode::Base: return bIsHovered ? BaseColorHovered : BaseColor;
+		case ETraitMode::Additive: return bIsHovered ? AdditiveColorHovered : AdditiveColor;
+		default:
+		{
+			break;
+		}
+	}
+
+	return FSlateColor(FSlateColor::UseSubduedForeground());
+}
+
+/*static*/ TSharedRef<SWidget> FTraitEditorUtils::GetInterfaceListWidget(EInterfaceDisplayType InterfaceDisplayType, const TSharedPtr<FTraitDataEditorDef>& InTraitDataShared, const TSharedPtr<FTraitEditorSharedData>& InTraitEditorSharedDataShared)
+{
+	if (!InTraitEditorSharedDataShared.IsValid()
+		|| (!InTraitEditorSharedDataShared->bShowTraitInterfaces
+			&& (!InTraitEditorSharedDataShared->bStackContainsErrors || !InTraitEditorSharedDataShared->bShowTraitInterfacesIfWarningsOrErrors)))
+	{
+		return SNullWidget::NullWidget;
+	}
+
+	TSharedRef<SHorizontalBox> InterfaceWidgetsHorizontalBox = SNew(SHorizontalBox);
+
+	const int32 NumSlots = InTraitEditorSharedDataShared->StackUsedInterfaces.Num();
+	for (int32 SlotIndex = 0; SlotIndex < NumSlots; SlotIndex++)
+	{
+		InterfaceWidgetsHorizontalBox->AddSlot()
+		.AutoWidth()
+		.Padding(FMargin(4.f, 1.f))
+		[
+			GetInterfaceWidget(InterfaceDisplayType, SlotIndex, InTraitDataShared, InTraitEditorSharedDataShared)
+		];
+	}
+
+	return SNew(SBorder)
+		[
+			SNew(SHorizontalBox)
+			+ SHorizontalBox::Slot()
+			.HAlign(HAlign_Left)
+			.AutoWidth()
+			[
+				InterfaceWidgetsHorizontalBox
+			]
+		];
+}
+
+/*static*/ TSharedRef<SWidget> FTraitEditorUtils::GetInterfaceWidget(EInterfaceDisplayType InterfaceDisplayType, int32 SlotIndex, const TSharedPtr<FTraitDataEditorDef>& InTraitDataShared, const TSharedPtr<FTraitEditorSharedData>& InTraitEditorSharedDataShared)
+{
+	return SNew(SBorder)
+		.BorderImage_Lambda([InterfaceDisplayType, InTraitDataShared, InTraitEditorSharedDataShared, SlotIndex]()
+			{
+				switch (InterfaceDisplayType)
+				{
+					case EInterfaceDisplayType::StackRequired:
+					{
+						if (InTraitEditorSharedDataShared.IsValid() && SlotIndex < InTraitEditorSharedDataShared->StackUsedInterfaces.Num())
+						{
+							if (InTraitDataShared.IsValid())
+							{
+								if (InTraitDataShared->StackStatus.MissingInterfaces.Find(InTraitEditorSharedDataShared->StackUsedInterfaces[SlotIndex]) != INDEX_NONE)
+								{
+									return FAppStyle::Get().GetBrush("Brushes.Error");
+								}
+							}
+						}
+						break;
+					}
+					case EInterfaceDisplayType::ListImplemented:
+					{
+						if (InTraitEditorSharedDataShared.IsValid() && SlotIndex < InTraitEditorSharedDataShared->StackUsedInterfaces.Num())
+						{
+							if (InTraitDataShared.IsValid())
+							{
+								if (InTraitEditorSharedDataShared->StackMissingInterfaces.Find(InTraitEditorSharedDataShared->StackUsedInterfaces[SlotIndex]) != INDEX_NONE)
+								{
+									return FAppStyle::Get().GetBrush("Brushes.Select");
+								}
+							}
+						}
+						break;
+					}
+					default:
+					{
+						break;
+					}
+				}
+
+				return FAppStyle::Get().GetBrush("Brushes.Background");
+			})
+		.Visibility_Lambda([InterfaceDisplayType, InTraitDataShared, InTraitEditorSharedDataShared, SlotIndex]()
+			{
+				if (InTraitDataShared.IsValid())
+				{
+					switch (InterfaceDisplayType)
+					{
+						case EInterfaceDisplayType::ListImplemented:
+						case EInterfaceDisplayType::StackImplemented:
+						{
+							if (InTraitDataShared->ImplementedInterfacesStackListIndexes.Find(SlotIndex) != INDEX_NONE)
+							{
+								return EVisibility::Visible;
+							}
+							break;
+						}
+						case EInterfaceDisplayType::StackRequired:
+						{
+							if (InTraitDataShared->RequiredInterfacesStackListIndexes.Find(SlotIndex) != INDEX_NONE)
+							{
+								return EVisibility::Visible;
+							}
+							break;
+						}
+						default:
+						{
+							break;
+						}
+					}
+				}
+				return EVisibility::Hidden;
+			})
+		[
+			SNew(STextBlock)
+			.Font(FAppStyle::GetFontStyle("PropertyWindow.BoldFont"))
+			.ColorAndOpacity(FSlateColor(FColor::White))
+			.Text_Lambda([InTraitEditorSharedDataShared, SlotIndex]()->FText
+			{
+				if (InTraitEditorSharedDataShared.IsValid())
+				{
+					// TODO zzz : Should the interface names be LOCTEXT in origin?
+					const FString InterfaceShortName(InTraitEditorSharedDataShared->StackUsedInterfaces.IsValidIndex(SlotIndex)
+						? InTraitEditorSharedDataShared->StackUsedInterfaces[SlotIndex].GetInterfaceShortName()
+						: FText(LOCTEXT("MissingInterfaceShortName", "???")).ToString());
+					return FText::FromString(InterfaceShortName);
+				}
+				return FText();
+			})
+			.ToolTipText_Lambda([InTraitEditorSharedDataShared, SlotIndex]()->FText
+			{
+				if (InTraitEditorSharedDataShared.IsValid())
+				{
+					// TODO zzz : Should the interface names be LOCTEXT in origin?
+					const FString InterfaceName(InTraitEditorSharedDataShared->StackUsedInterfaces.IsValidIndex(SlotIndex)
+						? InTraitEditorSharedDataShared->StackUsedInterfaces[SlotIndex].GetInterfaceName()
+						: FText(LOCTEXT("MissingInterfaceName", "Invalid or Missing Interface")).ToString());
+					return FText::FromString(InterfaceName);
+				}
+				return FText();
+			})
+		];
+}
+
+/*static*/ void FTraitEditorUtils::GenerateStackInterfacesUsedIndexes(TSharedPtr<FTraitDataEditorDef>& InTraitData, const TSharedPtr<FTraitEditorSharedData>& InTraitEditorSharedData)
+{
+	if (InTraitData.IsValid() && InTraitEditorSharedData.IsValid())
+	{
+		for (const FTraitInterfaceUID& ImplementedInterface : InTraitData->ImplementedInterfaces)
+		{
+			const int32 IndexInStackList = InTraitEditorSharedData->StackUsedInterfaces.Find(ImplementedInterface);
+			InTraitData->ImplementedInterfacesStackListIndexes.Add(IndexInStackList);
+		}
+		for (const FTraitInterfaceUID& RequiredInterface : InTraitData->RequiredInterfaces)
+		{
+			const int32 IndexInStackList = InTraitEditorSharedData->StackUsedInterfaces.Find(RequiredInterface);
+			InTraitData->RequiredInterfacesStackListIndexes.Add(IndexInStackList);
+		}
+	}
+}
+
+/*static*/ TSharedPtr<FTraitDataEditorDef> FTraitEditorUtils::FindTraitInCurrentStackData(const FTraitUID InTraitUID, TSharedPtr<TArray<TSharedPtr<FTraitDataEditorDef>>> InTraitsDataShared, int32* OutIndex)
+{
+	if (OutIndex != nullptr)
+	{
+		*OutIndex = INDEX_NONE;
+	}
+
+	if (InTraitsDataShared.IsValid())
+	{
+		TArray<TSharedPtr<FTraitDataEditorDef>>* CurrentTraitsData = InTraitsDataShared.Get();
+
+		const int32 NumTraits = CurrentTraitsData->Num();
+		for (int32 i = 0; i < NumTraits; ++i)
+		{
+			TSharedPtr<FTraitDataEditorDef>& TraitData = (*CurrentTraitsData)[i];
+			if (TraitData->TraitUID == InTraitUID)
+			{
+				if (OutIndex != nullptr)
+				{
+					*OutIndex = i;
+				}
+				return TraitData;
+			}
+		}
+	}
+
+	return nullptr;
+}
+
+/*static*/ bool FTraitEditorUtils::IsInternalTraitInteface(const FTraitInterfaceUID& InTraitInterfaceUID) 
+{
+	switch (InTraitInterfaceUID.GetUID())
+	{
+		case IGarbageCollection::InterfaceUID.GetUID():		return true;
+		default:											return false;
+	}
+}
+
+
+// --- FTraitListDragDropBase ---
+TSharedPtr<SWidget> FTraitListDragDropBase::GetDefaultDecorator() const
+{
+	TWeakPtr<FTraitDataEditorDef> DraggedTraitDataWeakLocal = DraggedTraitDataWeak;
+
+	return SNew(SHorizontalBox)
+			+ SHorizontalBox::Slot()
+			[
+				SNew(SBorder)
+				.BorderImage(FAppStyle::Get().GetBrush("Brushes.White"))
+				.BorderBackgroundColor_Lambda([DraggedTraitDataWeakLocal]()
+				{
+					if (TSharedPtr<FTraitDataEditorDef> TraitData = DraggedTraitDataWeakLocal.Pin())
+					{
+						return FTraitEditorUtils::GetTraitBackroundDisplayColor(TraitData->TraitMode, false);
+					}
+					return FSlateColor(FColor::Red);
+				})
+				.Padding(FMargin(1.0f, 1.0f))
+				[
+					SNew(SVerticalBox)
+					
+					+ SVerticalBox::Slot()
+					[
+						SNew(SBox)
+						.Padding(10.f, 10.f)
+						.MinDesiredHeight(30.0f)
+						.MinDesiredWidth(200.0f)
+						.VAlign(VAlign_Center)
+						[
+							SNew(SHorizontalBox)
+							+ SHorizontalBox::Slot()
+							.FillWidth(1.f)
+							.HAlign(HAlign_Center)
+							[
+								SNew(STextBlock)
+								.Justification(ETextJustify::Center)
+								.Text_Lambda([DraggedTraitDataWeakLocal]()
+								{
+									if (TSharedPtr<FTraitDataEditorDef> TraitData = DraggedTraitDataWeakLocal.Pin())
+									{
+										return TraitData->TraitDisplayName;
+									}
+									return FText();
+								})
+								.ColorAndOpacity_Lambda([DraggedTraitDataWeakLocal]()
+								{
+									if (TSharedPtr<FTraitDataEditorDef> TraitData = DraggedTraitDataWeakLocal.Pin())
+									{
+										return FTraitEditorUtils::GetTraitTextDisplayColor(TraitData->TraitMode);
+									}
+
+									return FSlateColor(FColor::Red);
+								})
+							]
+						]
+					]
+				]
+			];
+}
+
+} // namespace UE::AnimNext::Editor
+
+#undef LOCTEXT_NAMESPACE
\ No newline at end of file
diff --git a/Engine/Plugins/Experimental/AnimNext/Source/AnimNextEditor/Private/Graph/TraitEditorDefs.h b/Engine/Plugins/Experimental/AnimNext/Source/AnimNextEditor/Private/Graph/TraitEditorDefs.h
new file mode 100644
--- /dev/null
+++ b/Engine/Plugins/Experimental/AnimNext/Source/AnimNextEditor/Private/Graph/TraitEditorDefs.h
@@ -0,0 +1,193 @@
+// Copyright Epic Games, Inc. All Rights Reserved.
+
+#pragma once
+
+#include "TraitCore/Trait.h"
+#include "TraitCore/TraitMode.h"
+#include "TraitCore/TraitUID.h"
+#include "TraitCore/TraitInterfaceUID.h"
+#include "DragAndDrop/DecoratedDragDropOp.h"
+#include "DragAndDrop/AssetDragDropOp.h"
+
+class UAnimNextGraph_EdGraphNode;
+struct FSlateColor;
+
+namespace UE::Workspace
+{
+
+class IWorkspaceEditor;
+
+}
+
+namespace UE::AnimNext::Editor
+{
+
+struct FTraitStackTraitStatus
+{
+	enum class EStackStatus : uint8
+	{
+		Invalid,
+		Ok,
+		Warning,
+		Error,
+	};
+
+	struct FStatusMessage
+	{
+		FStatusMessage() = default;
+		FStatusMessage(EStackStatus InStatus, const FText& InMessageText)
+			: Status(InStatus)
+			, MessageText(InMessageText)
+		{}
+
+		EStackStatus Status = EStackStatus::Invalid;
+		FText MessageText;
+	};
+
+	bool HasWarnings() const
+	{
+		return HasStatusType(EStackStatus::Warning);
+	}
+	bool HasErrors() const
+	{
+		return HasStatusType(EStackStatus::Error);
+	}
+
+	EStackStatus TraitStatus = EStackStatus::Invalid;
+	TArray<FStatusMessage> StatusMessages;
+	TArray<FTraitInterfaceUID> MissingInterfaces;
+
+	bool HasStatusType(EStackStatus InStatus) const
+	{
+		for (const FStatusMessage& Status : StatusMessages)
+		{
+			if (Status.Status == InStatus)
+			{
+				return true;
+			}
+		}
+		return false;
+	}
+};
+
+struct FTraitDataEditorDef
+{
+	FTraitDataEditorDef() = default;
+
+	FTraitDataEditorDef(const FTrait& Trait, const FText& InTraitDisplayName)
+		: TraitName(Trait.GetTraitName())
+		, TraitDisplayName(InTraitDisplayName)
+		, TraitUID(Trait.GetTraitUID())
+		, TraitMode(Trait.GetTraitMode())
+		, ImplementedInterfaces(Trait.GetTraitInterfaces())
+		, RequiredInterfaces(Trait.GetTraitRequiredInterfaces())
+	{
+	}
+
+	FTraitDataEditorDef(const FName& InTraitName, const FText& InTraitDisplayName, FTraitUID InTraitUUID, ETraitMode InTraitMode, const TArray<FTraitInterfaceUID>& InImplementedInterfaces, const TArray<FTraitInterfaceUID>& InRequiredInterfaces)
+		: TraitName(InTraitName)
+		, TraitDisplayName(InTraitDisplayName)
+		, TraitUID(InTraitUUID)
+		, TraitMode(InTraitMode)
+		, ImplementedInterfaces(InImplementedInterfaces)
+		, RequiredInterfaces(InRequiredInterfaces)
+	{}
+
+	FName TraitName;
+	FText TraitDisplayName;
+	FTraitUID TraitUID;
+	ETraitMode TraitMode = ETraitMode::Invalid;
+
+	TArray<FTraitInterfaceUID> ImplementedInterfaces;
+	TArray<int32> ImplementedInterfacesStackListIndexes;
+	TArray<FTraitInterfaceUID> RequiredInterfaces;
+	TArray<int32> RequiredInterfacesStackListIndexes;
+
+	FTraitStackTraitStatus StackStatus;
+};
+
+struct FTraitCategoryData
+{
+	FTraitCategoryData() = default;
+
+	FTraitCategoryData( const FName& InCategory, const FText& InCategoryText)
+		: Category(InCategory)
+		, CategoryText(InCategoryText)
+	{}
+
+	FName Category;
+	FText CategoryText;
+
+	TArray<TSharedPtr<FTraitDataEditorDef>> TraitList;
+};
+
+struct FTraitEditorSharedData
+{
+	TWeakPtr<UE::Workspace::IWorkspaceEditor> WorkspaceEditorWeak;
+	TWeakObjectPtr<UAnimNextGraph_EdGraphNode> EdGraphNodeWeak = nullptr;
+	TSharedPtr<TArray<TSharedPtr<FTraitDataEditorDef>>> CurrentTraitsDataShared;
+
+	TArray<FTraitInterfaceUID> StackUsedInterfaces;
+	TArray<FTraitInterfaceUID> StackMissingInterfaces;
+	TArray<int32> StackUsedInteraceMissingIndexes;
+
+	bool bStackContainsErrors = false;
+	bool bShowTraitInterfaces = false;
+	bool bShowTraitInterfacesIfWarningsOrErrors = false;
+};
+
+struct FTraitEditorUtils
+{
+	static FSlateColor GetTraitIconErrorDisplayColor(const FTraitStackTraitStatus& InTraitStatus);
+	static FSlateColor GetTraitTextDisplayColor(const ETraitMode TraitMode);
+	static FSlateColor GetTraitBackroundDisplayColor(const ETraitMode TraitMode, bool bIsSelected = false, bool bIsHovered = false);
+
+	enum class EInterfaceDisplayType : uint8
+	{
+		ListImplemented,
+		StackImplemented,
+		StackRequired
+	};
+	static TSharedRef<SWidget> GetInterfaceListWidget(EInterfaceDisplayType InterfaceDisplayType, const TSharedPtr<FTraitDataEditorDef>& InTraitDataShared, const TSharedPtr<FTraitEditorSharedData>& InTraitEditorSharedDataShared);
+	static TSharedRef<SWidget> GetInterfaceWidget(EInterfaceDisplayType InterfaceDisplayType, int32 SlotIndex, const TSharedPtr<FTraitDataEditorDef>& InTraitDataShared, const TSharedPtr<FTraitEditorSharedData>& InTraitEditorSharedDataShared);
+	static void GenerateStackInterfacesUsedIndexes(TSharedPtr<FTraitDataEditorDef>& TraitData, const TSharedPtr<FTraitEditorSharedData>& InTraitEditorSharedData);
+
+	static TSharedPtr<FTraitDataEditorDef> FindTraitInCurrentStackData(const FTraitUID InTraitUID, TSharedPtr<TArray<TSharedPtr<FTraitDataEditorDef>>> CurrentTraitsDataShared, int32* OutIndex = nullptr);
+
+	// Internal Interfaces are not shown in the Traits Editor
+	// TODO zzz : This should be in a TraitInterfaceRegistry
+	static bool IsInternalTraitInteface(const FTraitInterfaceUID& InterfaceUID);
+};
+
+// --- FTraitListDragDropBase ---
+class FTraitListDragDropBase : public FDecoratedDragDropOp
+{
+public:
+	/** @return The identifier being dragged */
+	const TWeakPtr<FTraitDataEditorDef>& GetDraggedTraitData() const
+	{
+		return DraggedTraitDataWeak;
+	}
+
+	virtual TSharedPtr<SWidget> GetDefaultDecorator() const override;
+
+protected:
+	TWeakPtr<FTraitDataEditorDef> DraggedTraitDataWeak;
+};
+
+// --- FTraitListDragDropOp ---
+class FTraitListDragDropOp : public FTraitListDragDropBase
+{
+public:
+	DRAG_DROP_OPERATOR_TYPE(FTraitListDragDropOp, FDecoratedDragDropOp)
+
+	static TSharedRef<FTraitListDragDropOp> New(TWeakPtr<FTraitDataEditorDef> InDraggedTraitDataWeak)
+	{
+		TSharedRef<FTraitListDragDropOp> Operation = MakeShared<FTraitListDragDropOp>();
+		Operation->DraggedTraitDataWeak = InDraggedTraitDataWeak;
+		Operation->Construct();
+		return Operation;
+	}
+};
+
+};
diff --git a/Engine/Plugins/Experimental/AnimNext/Source/AnimNextEditor/Private/Graph/TraitEditorTabSummoner.cpp b/Engine/Plugins/Experimental/AnimNext/Source/AnimNextEditor/Private/Graph/TraitEditorTabSummoner.cpp
new file mode 100644
--- /dev/null
+++ b/Engine/Plugins/Experimental/AnimNext/Source/AnimNextEditor/Private/Graph/TraitEditorTabSummoner.cpp
@@ -0,0 +1,38 @@
+// Copyright Epic Games, Inc. All Rights Reserved.
+
+#include "Graph/TraitEditorTabSummoner.h"
+#include "IWorkspaceEditor.h"
+
+#define LOCTEXT_NAMESPACE "WorkspaceTabSummoner"
+
+namespace UE::AnimNext::Editor
+{
+
+const FName TraitEditorTabName("TraitEditorTab");
+
+
+FTraitEditorTabSummoner::FTraitEditorTabSummoner(const TSharedPtr<UE::Workspace::IWorkspaceEditor>& InHostingApp)
+	: FWorkflowTabFactory(TraitEditorTabName, StaticCastSharedPtr<FAssetEditorToolkit>(InHostingApp))
+{
+	TabLabel = LOCTEXT("TraitEditorTabLabel", "Trait Editor");
+	TabIcon = FSlateIcon("EditorStyle", "LevelEditor.Tabs.Outliner");
+	ViewMenuDescription = LOCTEXT("TraitEditorTabMenuDescription", "Trait Editor");
+	ViewMenuTooltip = LOCTEXT("TraitEditorTabToolTip", "Shows the Trait Editor tab.");
+	bIsSingleton = true;
+
+	TraitEditorView = SNew(STraitEditorView, InHostingApp.ToWeakPtr());
+}
+
+TSharedRef<SWidget> FTraitEditorTabSummoner::CreateTabBody(const FWorkflowTabSpawnInfo& Info) const
+{
+	return TraitEditorView.ToSharedRef();
+}
+
+FText FTraitEditorTabSummoner::GetTabToolTipText(const FWorkflowTabSpawnInfo& Info) const
+{
+	return ViewMenuTooltip;
+}
+
+} // end namespace UE::AnimNext::Graph
+
+#undef LOCTEXT_NAMESPACE
\ No newline at end of file
diff --git a/Engine/Plugins/Experimental/AnimNext/Source/AnimNextEditor/Private/Graph/TraitEditorTabSummoner.h b/Engine/Plugins/Experimental/AnimNext/Source/AnimNextEditor/Private/Graph/TraitEditorTabSummoner.h
new file mode 100644
--- /dev/null
+++ b/Engine/Plugins/Experimental/AnimNext/Source/AnimNextEditor/Private/Graph/TraitEditorTabSummoner.h
@@ -0,0 +1,34 @@
+// Copyright Epic Games, Inc. All Rights Reserved.
+
+#pragma once
+
+#include "WorkflowOrientedApp/WorkflowTabFactory.h"
+#include "STraitEditorView.h"
+
+namespace UE::Workspace
+{
+	class IWorkspaceEditor;
+}
+
+namespace UE::AnimNext::Editor
+{
+
+extern const FName TraitEditorTabName;
+
+DECLARE_DELEGATE_OneParam(FOnTraitEditorCreated, TSharedRef<STraitEditorView>);
+
+struct FTraitEditorTabSummoner : public FWorkflowTabFactory
+{
+public:
+	FTraitEditorTabSummoner(const TSharedPtr<UE::Workspace::IWorkspaceEditor>& InHostingApp);
+
+private:
+	// FWorkflowTabFactory interface
+	virtual TSharedRef<SWidget> CreateTabBody(const FWorkflowTabSpawnInfo& Info) const override;
+	virtual FText GetTabToolTipText(const FWorkflowTabSpawnInfo& Info) const override;
+
+	// The widget this tab spawner wraps
+	TSharedPtr<STraitEditorView> TraitEditorView;
+};
+
+} // end namespace UE::AnimNext::Editor
diff --git a/Engine/Plugins/Experimental/AnimNext/Source/AnimNextTestSuite/Private/AnimNextTraitBaseTest.cpp b/Engine/Plugins/Experimental/AnimNext/Source/AnimNextTestSuite/Private/AnimNextTraitBaseTest.cpp
--- a/Engine/Plugins/Experimental/AnimNext/Source/AnimNextTestSuite/Private/AnimNextTraitBaseTest.cpp
+++ b/Engine/Plugins/Experimental/AnimNext/Source/AnimNextTestSuite/Private/AnimNextTraitBaseTest.cpp
@@ -36,7 +36,7 @@ namespace UE::AnimNext
 
 	struct IInterfaceA : ITraitInterface
 	{
-		DECLARE_ANIM_TRAIT_INTERFACE(IInterfaceA, 0x34cb8e62)
+		DECLARE_ANIM_TRAIT_INTERFACE(IInterfaceA, ITA, 0x34cb8e62)
 
 		virtual void FuncA(const FExecutionContext& Context, const TTraitBinding<IInterfaceA>& Binding) const;
 	};
@@ -66,7 +66,7 @@ namespace UE::AnimNext
 
 	struct IInterfaceB : ITraitInterface
 	{
-		DECLARE_ANIM_TRAIT_INTERFACE(IInterfaceB, 0x33cb8ccf)
+		DECLARE_ANIM_TRAIT_INTERFACE(IInterfaceB, ITB, 0x33cb8ccf)
 
 		virtual void FuncB(const FExecutionContext& Context, const TTraitBinding<IInterfaceB>& Binding) const;
 	};
@@ -96,7 +96,7 @@ namespace UE::AnimNext
 
 	struct IInterfaceC : ITraitInterface
 	{
-		DECLARE_ANIM_TRAIT_INTERFACE(IInterfaceC, 0x32cb8b3c)
+		DECLARE_ANIM_TRAIT_INTERFACE(IInterfaceC, ITC, 0x32cb8b3c)
 
 		virtual void FuncC(const FExecutionContext& Context, const TTraitBinding<IInterfaceC>& Binding) const;
 	};
@@ -161,7 +161,7 @@ namespace UE::AnimNext
 	#define TRAIT_INTERFACE_ENUMERATOR(GeneratorMacro) \
 		GeneratorMacro(IInterfaceA) \
 
-	GENERATE_ANIM_TRAIT_IMPLEMENTATION(FTraitA_Base, TRAIT_INTERFACE_ENUMERATOR, NULL_ANIM_TRAIT_EVENT_ENUMERATOR)
+	GENERATE_ANIM_TRAIT_IMPLEMENTATION(FTraitA_Base, TRAIT_INTERFACE_ENUMERATOR, NULL_ANIM_TRAIT_INTERFACE_ENUMERATOR, NULL_ANIM_TRAIT_EVENT_ENUMERATOR)
 	#undef TRAIT_INTERFACE_ENUMERATOR
 
 	//////////////////////////////////////////////////////////////////////////
@@ -209,7 +209,7 @@ namespace UE::AnimNext
 		GeneratorMacro(IInterfaceA) \
 		GeneratorMacro(IInterfaceB) \
 
-	GENERATE_ANIM_TRAIT_IMPLEMENTATION(FTraitAB_Add, TRAIT_INTERFACE_ENUMERATOR, NULL_ANIM_TRAIT_EVENT_ENUMERATOR)
+	GENERATE_ANIM_TRAIT_IMPLEMENTATION(FTraitAB_Add, TRAIT_INTERFACE_ENUMERATOR, NULL_ANIM_TRAIT_INTERFACE_ENUMERATOR, NULL_ANIM_TRAIT_EVENT_ENUMERATOR)
 	#undef TRAIT_INTERFACE_ENUMERATOR
 
 	//////////////////////////////////////////////////////////////////////////
@@ -257,7 +257,7 @@ namespace UE::AnimNext
 		GeneratorMacro(IInterfaceA) \
 		GeneratorMacro(IInterfaceC) \
 
-	GENERATE_ANIM_TRAIT_IMPLEMENTATION(FTraitAC_Add, TRAIT_INTERFACE_ENUMERATOR, NULL_ANIM_TRAIT_EVENT_ENUMERATOR)
+	GENERATE_ANIM_TRAIT_IMPLEMENTATION(FTraitAC_Add, TRAIT_INTERFACE_ENUMERATOR, NULL_ANIM_TRAIT_INTERFACE_ENUMERATOR, NULL_ANIM_TRAIT_EVENT_ENUMERATOR)
 	#undef TRAIT_INTERFACE_ENUMERATOR
 
 
@@ -274,7 +274,7 @@ namespace UE::AnimNext
 	#define TRAIT_INTERFACE_ENUMERATOR(GeneratorMacro) \
 		GeneratorMacro(IInterfaceA) \
 
-	GENERATE_ANIM_TRAIT_IMPLEMENTATION(FTraitSerialization_Base, TRAIT_INTERFACE_ENUMERATOR, NULL_ANIM_TRAIT_EVENT_ENUMERATOR)
+	GENERATE_ANIM_TRAIT_IMPLEMENTATION(FTraitSerialization_Base, TRAIT_INTERFACE_ENUMERATOR, NULL_ANIM_TRAIT_INTERFACE_ENUMERATOR, NULL_ANIM_TRAIT_EVENT_ENUMERATOR)
 	#undef TRAIT_INTERFACE_ENUMERATOR
 
 	//////////////////////////////////////////////////////////////////////////
@@ -290,7 +290,7 @@ namespace UE::AnimNext
 	#define TRAIT_INTERFACE_ENUMERATOR(GeneratorMacro) \
 		GeneratorMacro(IInterfaceB) \
 
-	GENERATE_ANIM_TRAIT_IMPLEMENTATION(FTraitSerialization_Add, TRAIT_INTERFACE_ENUMERATOR, NULL_ANIM_TRAIT_EVENT_ENUMERATOR)
+	GENERATE_ANIM_TRAIT_IMPLEMENTATION(FTraitSerialization_Add, TRAIT_INTERFACE_ENUMERATOR, NULL_ANIM_TRAIT_INTERFACE_ENUMERATOR, NULL_ANIM_TRAIT_EVENT_ENUMERATOR)
 	#undef TRAIT_INTERFACE_ENUMERATOR
 
 	//////////////////////////////////////////////////////////////////////////
@@ -306,7 +306,7 @@ namespace UE::AnimNext
 	#define TRAIT_INTERFACE_ENUMERATOR(GeneratorMacro) \
 		GeneratorMacro(IInterfaceC) \
 
-	GENERATE_ANIM_TRAIT_IMPLEMENTATION(FTraitNativeSerialization_Add, TRAIT_INTERFACE_ENUMERATOR, NULL_ANIM_TRAIT_EVENT_ENUMERATOR)
+	GENERATE_ANIM_TRAIT_IMPLEMENTATION(FTraitNativeSerialization_Add, TRAIT_INTERFACE_ENUMERATOR, NULL_ANIM_TRAIT_INTERFACE_ENUMERATOR, NULL_ANIM_TRAIT_EVENT_ENUMERATOR)
 	#undef TRAIT_INTERFACE_ENUMERATOR
 }
 
diff --git a/Engine/Plugins/Experimental/AnimNext/Source/AnimNextTestSuite/Private/AnimNextTraitEventTest.cpp b/Engine/Plugins/Experimental/AnimNext/Source/AnimNextTestSuite/Private/AnimNextTraitEventTest.cpp
--- a/Engine/Plugins/Experimental/AnimNext/Source/AnimNextTestSuite/Private/AnimNextTraitEventTest.cpp
+++ b/Engine/Plugins/Experimental/AnimNext/Source/AnimNextTestSuite/Private/AnimNextTraitEventTest.cpp
@@ -49,7 +49,7 @@ namespace UE::AnimNext
 		GeneratorMacro(FTraitCoreTest_EventAB_Base::OnEventA) \
 		GeneratorMacro(FTraitCoreTest_EventAB_Base::OnEventB) \
 
-	GENERATE_ANIM_TRAIT_IMPLEMENTATION(FTraitCoreTest_EventAB_Base, NULL_ANIM_TRAIT_INTERFACE_ENUMERATOR, TRAIT_EVENT_ENUMERATOR)
+	GENERATE_ANIM_TRAIT_IMPLEMENTATION(FTraitCoreTest_EventAB_Base, NULL_ANIM_TRAIT_INTERFACE_ENUMERATOR, NULL_ANIM_TRAIT_INTERFACE_ENUMERATOR, TRAIT_EVENT_ENUMERATOR)
 	#undef TRAIT_EVENT_ENUMERATOR
 
 	struct FTraitCoreTest_EventAB_Add : FAdditiveTrait
@@ -74,7 +74,7 @@ namespace UE::AnimNext
 		GeneratorMacro(FTraitCoreTest_EventAB_Add::OnEventA) \
 		GeneratorMacro(FTraitCoreTest_EventAB_Add::OnEventB) \
 
-	GENERATE_ANIM_TRAIT_IMPLEMENTATION(FTraitCoreTest_EventAB_Add, NULL_ANIM_TRAIT_INTERFACE_ENUMERATOR, TRAIT_EVENT_ENUMERATOR)
+	GENERATE_ANIM_TRAIT_IMPLEMENTATION(FTraitCoreTest_EventAB_Add, NULL_ANIM_TRAIT_INTERFACE_ENUMERATOR, NULL_ANIM_TRAIT_INTERFACE_ENUMERATOR, TRAIT_EVENT_ENUMERATOR)
 	#undef TRAIT_EVENT_ENUMERATOR
 
 	struct FTraitCoreTest_EventA_Add : FAdditiveTrait
@@ -92,7 +92,7 @@ namespace UE::AnimNext
 	#define TRAIT_EVENT_ENUMERATOR(GeneratorMacro) \
 		GeneratorMacro(FTraitCoreTest_EventA_Add::OnEventA) \
 
-	GENERATE_ANIM_TRAIT_IMPLEMENTATION(FTraitCoreTest_EventA_Add, NULL_ANIM_TRAIT_INTERFACE_ENUMERATOR, TRAIT_EVENT_ENUMERATOR)
+	GENERATE_ANIM_TRAIT_IMPLEMENTATION(FTraitCoreTest_EventA_Add, NULL_ANIM_TRAIT_INTERFACE_ENUMERATOR, NULL_ANIM_TRAIT_INTERFACE_ENUMERATOR, TRAIT_EVENT_ENUMERATOR)
 	#undef TRAIT_EVENT_ENUMERATOR
 
 	struct FTraitCoreTest_EventB_Add : FAdditiveTrait
@@ -110,7 +110,7 @@ namespace UE::AnimNext
 	#define TRAIT_EVENT_ENUMERATOR(GeneratorMacro) \
 		GeneratorMacro(FTraitCoreTest_EventB_Add::OnEventB) \
 
-	GENERATE_ANIM_TRAIT_IMPLEMENTATION(FTraitCoreTest_EventB_Add, NULL_ANIM_TRAIT_INTERFACE_ENUMERATOR, TRAIT_EVENT_ENUMERATOR)
+	GENERATE_ANIM_TRAIT_IMPLEMENTATION(FTraitCoreTest_EventB_Add, NULL_ANIM_TRAIT_INTERFACE_ENUMERATOR, NULL_ANIM_TRAIT_INTERFACE_ENUMERATOR, TRAIT_EVENT_ENUMERATOR)
 	#undef TRAIT_EVENT_ENUMERATOR
 }
 
diff --git a/Engine/Plugins/Experimental/AnimNext/Source/AnimNextUncookedOnly/Internal/AnimNextRigVMAssetEditorData.h b/Engine/Plugins/Experimental/AnimNext/Source/AnimNextUncookedOnly/Internal/AnimNextRigVMAssetEditorData.h
--- a/Engine/Plugins/Experimental/AnimNext/Source/AnimNextUncookedOnly/Internal/AnimNextRigVMAssetEditorData.h
+++ b/Engine/Plugins/Experimental/AnimNext/Source/AnimNextUncookedOnly/Internal/AnimNextRigVMAssetEditorData.h
@@ -1,4 +1,4 @@
-﻿// Copyright Epic Games, Inc. All Rights Reserved.
+// Copyright Epic Games, Inc. All Rights Reserved.
 
 #pragma once
 
@@ -42,6 +42,9 @@ namespace UE::AnimNext::UncookedOnly
 {
 	// A delegate for subscribing / reacting to editor data modifications.
 	DECLARE_MULTICAST_DELEGATE_OneParam(FOnEditorDataModified, UAnimNextRigVMAssetEditorData* /* InEditorData */);
+
+	// An interaction bracket count reached 0
+	DECLARE_MULTICAST_DELEGATE_OneParam(FOnInteractionBracketFinished, UAnimNextRigVMAssetEditorData* /* InEditorData */);
 }
 
 // Script-callable editor API hoisted onto UAnimNextRigVMAsset
@@ -271,6 +274,9 @@ protected:
 	// Delegate to subscribe to modifications to this editor data
 	UE::AnimNext::UncookedOnly::FOnEditorDataModified ModifiedDelegate;
 
+	// Delegate to get notified when an interaction bracket reaches 0
+	UE::AnimNext::UncookedOnly::FOnInteractionBracketFinished InteractionBracketFinished;
+
 	// Cached exports, generated lazily or on compilation
 	mutable TOptional<FAnimNextParameterProviderAssetRegistryExports> CachedExports;
 
diff --git a/Engine/Plugins/Experimental/AnimNext/Source/AnimNextUncookedOnly/Internal/Graph/AnimNextGraph_Controller.h b/Engine/Plugins/Experimental/AnimNext/Source/AnimNextUncookedOnly/Internal/Graph/AnimNextGraph_Controller.h
--- a/Engine/Plugins/Experimental/AnimNext/Source/AnimNextUncookedOnly/Internal/Graph/AnimNextGraph_Controller.h
+++ b/Engine/Plugins/Experimental/AnimNext/Source/AnimNextUncookedOnly/Internal/Graph/AnimNextGraph_Controller.h
@@ -11,6 +11,14 @@
 
 struct FAnimNextParamType;
 
+
+namespace UE::AnimNext
+{
+	struct FTrait;
+	struct FTraitUID;
+}
+
+
 /**
   * Implements AnimNext RigVM controller extensions
   */
@@ -20,6 +28,26 @@ class UAnimNextGraph_Controller : public URigVMController
 	GENERATED_BODY()
 
 public:
+	// Adds a new Trait to the Stack, with default struct values
+	// Returns Trait Instance Name (or NAME_None on failure)
+	UFUNCTION(BlueprintCallable, Category = RigVMController)
+	ANIMNEXTUNCOOKEDONLY_API FName AddTrait(FName InNodeName, FName InNewTraitTypeName, int32 InPinIndex, const FString& InNewTraitDefaultValue = TEXT(""), bool bSetupUndoRedo = true, bool bPrintPythonCommand = false);
+
+	// Removes a Trait from the Stack, using Trait Instance Name
+	// Returns operation success (true) or failure (false)
+	UFUNCTION(BlueprintCallable, Category = RigVMController)
+	ANIMNEXTUNCOOKEDONLY_API bool RemoveTrait(FName InNodeName, FName InTraitInstanceName, bool bSetupUndoRedo = true, bool bPrintPythonCommand = false);
+
+	// Swap a Trait from the Stack with a new one, using existing Trait Instance Name and new Trait Type Name
+	// Returns Trait Instance Name (or NAME_None on failure)
+	UFUNCTION(BlueprintCallable, Category = RigVMController)
+	ANIMNEXTUNCOOKEDONLY_API FName SwapTrait(FName InNodeName, FName InTraitInstanceName, int32 InCurrentTraitPinIndex, FName InNewTraitTypeName, const FString& InNewTraitDefaultValue = TEXT(""), bool bSetupUndoRedo = true, bool bPrintPythonCommand = false);
+
+	// Move a Trait from its current PinIndex to the specified one (moving it visually in the stack)
+	// Returns operation success (true) or failure (false)
+	UFUNCTION(BlueprintCallable, Category = RigVMController)
+	ANIMNEXTUNCOOKEDONLY_API bool SetTraitPinIndex(FName InNodeName, FName InTraitInstanceName, int32 InNewPinIndex, bool bSetupUndoRedo = true, bool bPrintPythonCommand = false);
+
 	// Adds a unit node with a dynamic number of pins
 	URigVMUnitNode* AddUnitNodeWithPins(UScriptStruct* InScriptStruct, const FRigVMPinInfoArray& PinArray, const FName& InMethodName = TEXT("Execute"), const FVector2D& InPosition = FVector2D::ZeroVector, const FString& InNodeName = TEXT(""), bool bSetupUndoRedo = true, bool bPrintPythonCommand = false);
 
@@ -45,4 +73,5 @@ public:
 
 private:
 	URigVMNode* AddGetAnimNextParameterNodeInternal(UScriptStruct* NodeStruct, const FVector2D& Position, FName ParameterName, const FAnimNextParamType& Type, const TInstancedStruct<FAnimNextParamInstanceIdentifier>& InstanceId, bool bSetupUndoRedo, bool bPrintPythonCommand);
+
 };
diff --git a/Engine/Plugins/Experimental/AnimNext/Source/AnimNextUncookedOnly/Internal/Graph/AnimNextGraph_EditorData.h b/Engine/Plugins/Experimental/AnimNext/Source/AnimNextUncookedOnly/Internal/Graph/AnimNextGraph_EditorData.h
--- a/Engine/Plugins/Experimental/AnimNext/Source/AnimNextUncookedOnly/Internal/Graph/AnimNextGraph_EditorData.h
+++ b/Engine/Plugins/Experimental/AnimNext/Source/AnimNextUncookedOnly/Internal/Graph/AnimNextGraph_EditorData.h
@@ -63,7 +63,8 @@ class UAnimNextGraph_EditorData : public UAnimNextRigVMAssetEditorData
 	friend class UE::AnimNext::Editor::SAnimNextGraphView;
 	friend struct FAnimNextGraphSchemaAction_RigUnit;
 	friend struct FAnimNextGraphSchemaAction_DispatchFactory;
-	friend class FAnimationAnimNextRuntimeTest_GraphAddTrait;
+	friend class FAnimationAnimNextEditorTest_GraphAddTrait;
+	friend class FAnimationAnimNextEditorTest_GraphTraitOperations;
 	friend class FAnimationAnimNextRuntimeTest_GraphExecute;
 	friend class FAnimationAnimNextRuntimeTest_GraphExecuteLatent;
 	
diff --git a/Engine/Plugins/Experimental/AnimNext/Source/AnimNextUncookedOnly/Private/AnimNextRigVMAssetEditorData.cpp b/Engine/Plugins/Experimental/AnimNext/Source/AnimNextUncookedOnly/Private/AnimNextRigVMAssetEditorData.cpp
--- a/Engine/Plugins/Experimental/AnimNext/Source/AnimNextUncookedOnly/Private/AnimNextRigVMAssetEditorData.cpp
+++ b/Engine/Plugins/Experimental/AnimNext/Source/AnimNextUncookedOnly/Private/AnimNextRigVMAssetEditorData.cpp
@@ -1,4 +1,4 @@
-﻿// Copyright Epic Games, Inc. All Rights Reserved.
+// Copyright Epic Games, Inc. All Rights Reserved.
 
 #include "AnimNextRigVMAssetEditorData.h"
 
@@ -602,6 +602,11 @@ void UAnimNextRigVMAssetEditorData::DecrementVMRecompileBracket()
 			RecompileVMIfRequired();
 		}
 		VMRecompilationBracket = 0;
+
+		if (InteractionBracketFinished.IsBound())
+		{
+			InteractionBracketFinished.Broadcast(this);
+		}
 	}
 	else if (VMRecompilationBracket > 0)
 	{
diff --git a/Engine/Plugins/Experimental/AnimNext/Source/AnimNextUncookedOnly/Private/Graph/AnimNextGraph_Controller.cpp b/Engine/Plugins/Experimental/AnimNext/Source/AnimNextUncookedOnly/Private/Graph/AnimNextGraph_Controller.cpp
--- a/Engine/Plugins/Experimental/AnimNext/Source/AnimNextUncookedOnly/Private/Graph/AnimNextGraph_Controller.cpp
+++ b/Engine/Plugins/Experimental/AnimNext/Source/AnimNextUncookedOnly/Private/Graph/AnimNextGraph_Controller.cpp
@@ -1,7 +1,9 @@
 // Copyright Epic Games, Inc. All Rights Reserved.
 
 #include "Graph/AnimNextGraph_Controller.h"
+
 #include "Graph/AnimNextGraph_UnitNode.h"
+#include "Graph/RigDecorator_AnimNextCppTrait.h"
 #include "Param/AnimNextEditorParam.h"
 #include "Param/AnimNextParam.h"
 #include "Param/AnimNextParamInstanceIdentifier.h"
@@ -10,6 +12,195 @@
 #include "Param/RigVMDispatch_GetParameter.h"
 #include "Param/RigVMDispatch_GetScopedParameter.h"
 #include "Param/RigVMDispatch_SetLayerParameter.h"
+#include "RigVMModel/RigVMControllerActions.h"
+#include "TraitCore/Trait.h"
+#include "TraitCore/TraitRegistry.h"
+#include "TraitCore/TraitUID.h"
+
+FName UAnimNextGraph_Controller::AddTrait(FName InNodeName, FName InNewTraitTypeName, int32 InPinIndex, const FString& InNewTraitDefaultValue, bool bSetupUndoRedo, bool bPrintPythonCommand)
+{
+	if (!IsValidGraph())
+	{
+		return NAME_None;
+	}
+
+	if (!IsTransacting() && !IsGraphEditable())
+	{
+		return NAME_None;
+	}
+
+	if (InNodeName == NAME_None)
+	{
+		ReportError(TEXT("Invalid node name."));
+		return NAME_None;
+	}
+
+	const URigVMGraph* Graph = GetGraph();
+	check(Graph);
+
+	URigVMNode * Node = Graph->FindNodeByName(InNodeName);
+	if (Node == nullptr)
+	{
+		ReportError(TEXT("This graph does not contain a node with the provided name."));
+		return NAME_None;
+	}
+
+	UE::AnimNext::FTraitRegistry& TraitRegistry = UE::AnimNext::FTraitRegistry::Get();
+
+	const UE::AnimNext::FTrait* Trait = TraitRegistry.Find(InNewTraitTypeName);
+	if (Trait == nullptr)
+	{
+		ReportError(TEXT("Unknown Trait Type."));
+		return NAME_None;
+	}
+
+	const UScriptStruct* CppDecoratorStruct = FRigDecorator_AnimNextCppDecorator::StaticStruct();
+	UScriptStruct* ScriptStruct = Trait->GetTraitSharedDataStruct();
+
+	FString DefaultValue = InNewTraitDefaultValue;
+	if (DefaultValue.IsEmpty())
+	{
+		const FRigDecorator_AnimNextCppDecorator DefaultCppDecoratorStructInstance;
+		FRigDecorator_AnimNextCppDecorator CppDecoratorStructInstance;
+		CppDecoratorStructInstance.DecoratorSharedDataStruct = ScriptStruct;
+
+		if (!CppDecoratorStructInstance.CanBeAddedToNode(Node, nullptr))
+		{
+			ReportError(TEXT("Trait is not supported by the Node."));
+			return NAME_None;	// This trait isn't supported on this node
+		}
+
+		const FProperty* Prop = FAnimNextCppDecoratorWrapper::StaticStruct()->FindPropertyByName(GET_MEMBER_NAME_STRING_CHECKED(FAnimNextCppDecoratorWrapper, CppDecorator));
+		Prop->ExportText_Direct(DefaultValue, &CppDecoratorStructInstance, &DefaultCppDecoratorStructInstance, nullptr, PPF_SerializedAsImportText);
+	}
+
+	// Avoid multiple VM recompilations for internal operations
+	FRigVMControllerCompileBracketScope CompileScope(this);
+
+	const FName TraitName = InNewTraitTypeName;
+
+	return AddDecorator(InNodeName, *CppDecoratorStruct->GetPathName(), TraitName, DefaultValue, InPinIndex, bSetupUndoRedo, bPrintPythonCommand);
+}
+
+bool UAnimNextGraph_Controller::RemoveTrait(FName InNodeName, FName InTraitInstanceName, bool bSetupUndoRedo, bool bPrintPythonCommand)
+{
+	if (!IsValidGraph())
+	{
+		return false;
+	}
+
+	if (!IsTransacting() && !IsGraphEditable())
+	{
+		return false;
+	}
+
+	if (InNodeName == NAME_None)
+	{
+		ReportError(TEXT("Invalid node name."));
+		return false;
+	}
+
+	// Avoid multiple VM recompilations for internal operations
+	FRigVMControllerCompileBracketScope CompileScope(this);
+	return RemoveDecorator(InNodeName, InTraitInstanceName, bSetupUndoRedo, bPrintPythonCommand);
+}
+
+FName UAnimNextGraph_Controller::SwapTrait(FName InNodeName, FName InTraitInstanceName, int32 InCurrentTraitPinIndex, FName InNewTraitTypeName, const FString& InNewTraitDefaultValue, bool bSetupUndoRedo, bool bPrintPythonCommand)
+{
+	if (!IsValidGraph())
+	{
+		return NAME_None;
+	}
+
+	if (!IsTransacting() && !IsGraphEditable())
+	{
+		return NAME_None;
+	}
+
+	if (InNodeName == NAME_None)
+	{
+		ReportError(TEXT("Invalid node name."));
+		return NAME_None;
+	}
+
+	// Avoid multiple VM recompilations, for each operation
+	FRigVMControllerCompileBracketScope CompileScope(this);
+	if (RemoveTrait(InNodeName, InTraitInstanceName, bSetupUndoRedo, bPrintPythonCommand))
+	{
+		return AddTrait(InNodeName, InNewTraitTypeName, InCurrentTraitPinIndex, InNewTraitDefaultValue, bSetupUndoRedo, bPrintPythonCommand);
+	}
+
+	return NAME_None;
+}
+
+bool UAnimNextGraph_Controller::SetTraitPinIndex(FName InNodeName, FName InTraitInstanceName, int32 InNewPinIndex, bool bSetupUndoRedo, bool bPrintPythonCommand)
+{
+	if (!IsValidGraph())
+	{
+		return false;
+	}
+
+	if (!IsTransacting() && !IsGraphEditable())
+	{
+		return false;
+	}
+
+	if (InNodeName == NAME_None)
+	{
+		ReportError(TEXT("Invalid node name."));
+		return false;
+	}
+
+	const URigVMGraph* Graph = GetGraph();
+	check(Graph);
+
+	URigVMNode* Node = Graph->FindNodeByName(InNodeName);
+	if (Node == nullptr)
+	{
+		ReportError(TEXT("This graph does not contain a node with the provided name."));
+		return false;
+	}
+
+	if (!IsTransacting() && !IsGraphEditable())
+	{
+		return false;
+	}
+
+	URigVMPin* TraitPin = Node->FindDecorator(InTraitInstanceName);
+	if (TraitPin == nullptr)
+	{
+		ReportError(TEXT("The node does not contain a Trait with the provided name."));
+		return false;
+	}
+	
+	// Save current pin data for later
+	const FString TraitDefaultValue = TraitPin->GetDefaultValue();
+
+	// TODO zzz : Is there a better way to get a Trait* from a TraitPin ?
+	if (TSharedPtr<FStructOnScope> ScopedTrait = Node->GetDecoratorInstance(TraitPin->GetFName()))
+	{
+		const FRigVMDecorator* Decorator = (FRigVMDecorator*)ScopedTrait->GetStructMemory();
+		if (const UScriptStruct* TraitSharedInstanceData = Decorator->GetDecoratorSharedDataStruct())
+		{
+			UE::AnimNext::FTraitRegistry& TraitRegistry = UE::AnimNext::FTraitRegistry::Get();
+			if (const UE::AnimNext::FTrait* Trait = TraitRegistry.Find(TraitSharedInstanceData))
+			{
+				// Avoid multiple VM recompilations, for each operation
+				FRigVMControllerCompileBracketScope CompileScope(this);
+				if (RemoveTrait(InNodeName, InTraitInstanceName, bSetupUndoRedo, bPrintPythonCommand))
+				{
+					// TOOO zzz : Why TraitName is a string ?
+					if (AddTrait(InNodeName, FName(Trait->GetTraitName()), InNewPinIndex, TraitDefaultValue, bSetupUndoRedo, bPrintPythonCommand) != NAME_None)
+					{
+						return true;
+					}
+				}
+			}
+		}
+	}
+
+	return false;
+}
 
 URigVMUnitNode* UAnimNextGraph_Controller::AddUnitNodeWithPins(UScriptStruct* InScriptStruct, const FRigVMPinInfoArray& PinArray, const FName& InMethodName, const FVector2D& InPosition, const FString& InNodeName, bool bSetupUndoRedo, bool bPrintPythonCommand)
 {
diff --git a/Engine/Plugins/Experimental/AnimNextAnimGraph/Source/AnimNextAnimGraph/Private/Traits/BlendByBool.cpp b/Engine/Plugins/Experimental/AnimNextAnimGraph/Source/AnimNextAnimGraph/Private/Traits/BlendByBool.cpp
--- a/Engine/Plugins/Experimental/AnimNextAnimGraph/Source/AnimNextAnimGraph/Private/Traits/BlendByBool.cpp
+++ b/Engine/Plugins/Experimental/AnimNextAnimGraph/Source/AnimNextAnimGraph/Private/Traits/BlendByBool.cpp
@@ -17,7 +17,7 @@ namespace UE::AnimNext
 		GeneratorMacro(IUpdate) \
 		GeneratorMacro(IUpdateTraversal) \
 
-	GENERATE_ANIM_TRAIT_IMPLEMENTATION(FBlendByBoolTrait, TRAIT_INTERFACE_ENUMERATOR, NULL_ANIM_TRAIT_EVENT_ENUMERATOR)
+	GENERATE_ANIM_TRAIT_IMPLEMENTATION(FBlendByBoolTrait, TRAIT_INTERFACE_ENUMERATOR, NULL_ANIM_TRAIT_INTERFACE_ENUMERATOR, NULL_ANIM_TRAIT_EVENT_ENUMERATOR)
 	#undef TRAIT_INTERFACE_ENUMERATOR
 
 	static constexpr int32 TRUE_CHILD_INDEX = 0;
diff --git a/Engine/Plugins/Experimental/AnimNextAnimGraph/Source/AnimNextAnimGraph/Private/Traits/BlendInertializer.cpp b/Engine/Plugins/Experimental/AnimNextAnimGraph/Source/AnimNextAnimGraph/Private/Traits/BlendInertializer.cpp
--- a/Engine/Plugins/Experimental/AnimNextAnimGraph/Source/AnimNextAnimGraph/Private/Traits/BlendInertializer.cpp
+++ b/Engine/Plugins/Experimental/AnimNextAnimGraph/Source/AnimNextAnimGraph/Private/Traits/BlendInertializer.cpp
@@ -16,8 +16,13 @@ namespace UE::AnimNext
 		GeneratorMacro(IDiscreteBlend) \
 		GeneratorMacro(ISmoothBlend) \
 
-	GENERATE_ANIM_TRAIT_IMPLEMENTATION(FBlendInertializerTrait, TRAIT_INTERFACE_ENUMERATOR, NULL_ANIM_TRAIT_EVENT_ENUMERATOR)
+	// Trait implementation boilerplate
+	#define TRAIT_REQUIRED_INTERFACE_ENUMERATOR(GeneratorMacro) \
+		GeneratorMacro(IDiscreteBlend) \
+
+	GENERATE_ANIM_TRAIT_IMPLEMENTATION(FBlendInertializerTrait, TRAIT_INTERFACE_ENUMERATOR, TRAIT_REQUIRED_INTERFACE_ENUMERATOR, NULL_ANIM_TRAIT_EVENT_ENUMERATOR)
 	#undef TRAIT_INTERFACE_ENUMERATOR
+	#undef TRAIT_REQUIRED_INTERFACE_ENUMERATOR
 
 	void FBlendInertializerTrait::OnBlendTransition(FExecutionContext& Context, const TTraitBinding<IDiscreteBlend>& Binding, int32 OldChildIndex, int32 NewChildIndex) const
 	{
diff --git a/Engine/Plugins/Experimental/AnimNextAnimGraph/Source/AnimNextAnimGraph/Private/Traits/BlendInertializer.h b/Engine/Plugins/Experimental/AnimNextAnimGraph/Source/AnimNextAnimGraph/Private/Traits/BlendInertializer.h
--- a/Engine/Plugins/Experimental/AnimNextAnimGraph/Source/AnimNextAnimGraph/Private/Traits/BlendInertializer.h
+++ b/Engine/Plugins/Experimental/AnimNextAnimGraph/Source/AnimNextAnimGraph/Private/Traits/BlendInertializer.h
@@ -12,7 +12,7 @@
 
 class UBlendProfile;
 
-USTRUCT(meta = (DisplayName = "Blend Inertializer"))
+USTRUCT(meta = (DisplayName = "Blend Inertializer", Category = "Default"))
 struct FAnimNextBlendInertializerTraitSharedData : public FAnimNextTraitSharedData
 {
 	GENERATED_BODY()
diff --git a/Engine/Plugins/Experimental/AnimNextAnimGraph/Source/AnimNextAnimGraph/Private/Traits/BlendSmoother.cpp b/Engine/Plugins/Experimental/AnimNextAnimGraph/Source/AnimNextAnimGraph/Private/Traits/BlendSmoother.cpp
--- a/Engine/Plugins/Experimental/AnimNextAnimGraph/Source/AnimNextAnimGraph/Private/Traits/BlendSmoother.cpp
+++ b/Engine/Plugins/Experimental/AnimNextAnimGraph/Source/AnimNextAnimGraph/Private/Traits/BlendSmoother.cpp
@@ -19,7 +19,7 @@ namespace UE::AnimNext
 		GeneratorMacro(ISmoothBlend) \
 		GeneratorMacro(IUpdate) \
 
-	GENERATE_ANIM_TRAIT_IMPLEMENTATION(FBlendSmootherTrait, TRAIT_INTERFACE_ENUMERATOR, NULL_ANIM_TRAIT_EVENT_ENUMERATOR)
+	GENERATE_ANIM_TRAIT_IMPLEMENTATION(FBlendSmootherTrait, TRAIT_INTERFACE_ENUMERATOR, NULL_ANIM_TRAIT_INTERFACE_ENUMERATOR, NULL_ANIM_TRAIT_EVENT_ENUMERATOR)
 	#undef TRAIT_INTERFACE_ENUMERATOR
 
 	void FBlendSmootherTrait::PostEvaluate(FEvaluateTraversalContext& Context, const TTraitBinding<IEvaluate>& Binding) const
diff --git a/Engine/Plugins/Experimental/AnimNextAnimGraph/Source/AnimNextAnimGraph/Private/Traits/BlendSmootherPerBone.cpp b/Engine/Plugins/Experimental/AnimNextAnimGraph/Source/AnimNextAnimGraph/Private/Traits/BlendSmootherPerBone.cpp
--- a/Engine/Plugins/Experimental/AnimNextAnimGraph/Source/AnimNextAnimGraph/Private/Traits/BlendSmootherPerBone.cpp
+++ b/Engine/Plugins/Experimental/AnimNextAnimGraph/Source/AnimNextAnimGraph/Private/Traits/BlendSmootherPerBone.cpp
@@ -19,7 +19,7 @@ namespace UE::AnimNext
 		GeneratorMacro(IEvaluate) \
 		GeneratorMacro(IUpdate) \
 
-	GENERATE_ANIM_TRAIT_IMPLEMENTATION(FBlendSmootherPerBoneTrait, TRAIT_INTERFACE_ENUMERATOR, NULL_ANIM_TRAIT_EVENT_ENUMERATOR)
+	GENERATE_ANIM_TRAIT_IMPLEMENTATION(FBlendSmootherPerBoneTrait, TRAIT_INTERFACE_ENUMERATOR, NULL_ANIM_TRAIT_INTERFACE_ENUMERATOR, NULL_ANIM_TRAIT_EVENT_ENUMERATOR)
 	#undef TRAIT_INTERFACE_ENUMERATOR
 
 	void FBlendSmootherPerBoneTrait::PostEvaluate(FEvaluateTraversalContext& Context, const TTraitBinding<IEvaluate>& Binding) const
diff --git a/Engine/Plugins/Experimental/AnimNextAnimGraph/Source/AnimNextAnimGraph/Private/Traits/BlendTwoWay.cpp b/Engine/Plugins/Experimental/AnimNextAnimGraph/Source/AnimNextAnimGraph/Private/Traits/BlendTwoWay.cpp
--- a/Engine/Plugins/Experimental/AnimNextAnimGraph/Source/AnimNextAnimGraph/Private/Traits/BlendTwoWay.cpp
+++ b/Engine/Plugins/Experimental/AnimNextAnimGraph/Source/AnimNextAnimGraph/Private/Traits/BlendTwoWay.cpp
@@ -18,7 +18,7 @@ namespace UE::AnimNext
 		GeneratorMacro(IUpdate) \
 		GeneratorMacro(IUpdateTraversal) \
 
-	GENERATE_ANIM_TRAIT_IMPLEMENTATION(FBlendTwoWayTrait, TRAIT_INTERFACE_ENUMERATOR, NULL_ANIM_TRAIT_EVENT_ENUMERATOR)
+	GENERATE_ANIM_TRAIT_IMPLEMENTATION(FBlendTwoWayTrait, TRAIT_INTERFACE_ENUMERATOR, NULL_ANIM_TRAIT_INTERFACE_ENUMERATOR, NULL_ANIM_TRAIT_EVENT_ENUMERATOR)
 	#undef TRAIT_INTERFACE_ENUMERATOR
 
 	void FBlendTwoWayTrait::PostEvaluate(FEvaluateTraversalContext& Context, const TTraitBinding<IEvaluate>& Binding) const
diff --git a/Engine/Plugins/Experimental/AnimNextAnimGraph/Source/AnimNextAnimGraph/Private/Traits/InertializationDeadBlend.cpp b/Engine/Plugins/Experimental/AnimNextAnimGraph/Source/AnimNextAnimGraph/Private/Traits/InertializationDeadBlend.cpp
--- a/Engine/Plugins/Experimental/AnimNextAnimGraph/Source/AnimNextAnimGraph/Private/Traits/InertializationDeadBlend.cpp
+++ b/Engine/Plugins/Experimental/AnimNextAnimGraph/Source/AnimNextAnimGraph/Private/Traits/InertializationDeadBlend.cpp
@@ -23,7 +23,7 @@ namespace UE::AnimNext
 #define TRAIT_EVENT_ENUMERATOR(GeneratorMacro) \
 		GeneratorMacro(FInertializationDeadBlendTrait::OnInertializationRequestEvent) \
 
-		GENERATE_ANIM_TRAIT_IMPLEMENTATION(FInertializationDeadBlendTrait, TRAIT_INTERFACE_ENUMERATOR, TRAIT_EVENT_ENUMERATOR)
+		GENERATE_ANIM_TRAIT_IMPLEMENTATION(FInertializationDeadBlendTrait, TRAIT_INTERFACE_ENUMERATOR, NULL_ANIM_TRAIT_INTERFACE_ENUMERATOR, TRAIT_EVENT_ENUMERATOR)
 #undef TRAIT_INTERFACE_ENUMERATOR
 #undef TRAIT_EVENT_ENUMERATOR
 
diff --git a/Engine/Plugins/Experimental/AnimNextAnimGraph/Source/AnimNextAnimGraph/Private/Traits/Passthrough.cpp b/Engine/Plugins/Experimental/AnimNextAnimGraph/Source/AnimNextAnimGraph/Private/Traits/Passthrough.cpp
--- a/Engine/Plugins/Experimental/AnimNextAnimGraph/Source/AnimNextAnimGraph/Private/Traits/Passthrough.cpp
+++ b/Engine/Plugins/Experimental/AnimNextAnimGraph/Source/AnimNextAnimGraph/Private/Traits/Passthrough.cpp
@@ -12,7 +12,7 @@ namespace UE::AnimNext
 #define TRAIT_INTERFACE_ENUMERATOR(GeneratorMacro) \
 		GeneratorMacro(IHierarchy) \
 
-		GENERATE_ANIM_TRAIT_IMPLEMENTATION(FPassthroughTrait, TRAIT_INTERFACE_ENUMERATOR, NULL_ANIM_TRAIT_EVENT_ENUMERATOR)
+		GENERATE_ANIM_TRAIT_IMPLEMENTATION(FPassthroughTrait, TRAIT_INTERFACE_ENUMERATOR, NULL_ANIM_TRAIT_INTERFACE_ENUMERATOR, NULL_ANIM_TRAIT_EVENT_ENUMERATOR)
 #undef TRAIT_INTERFACE_ENUMERATOR
 
 	void FPassthroughTrait::FInstanceData::Construct(const FExecutionContext& Context, const FTraitBinding& Binding)
diff --git a/Engine/Plugins/Experimental/AnimNextAnimGraph/Source/AnimNextAnimGraph/Private/Traits/ReferencePoseTrait.cpp b/Engine/Plugins/Experimental/AnimNextAnimGraph/Source/AnimNextAnimGraph/Private/Traits/ReferencePoseTrait.cpp
--- a/Engine/Plugins/Experimental/AnimNextAnimGraph/Source/AnimNextAnimGraph/Private/Traits/ReferencePoseTrait.cpp
+++ b/Engine/Plugins/Experimental/AnimNextAnimGraph/Source/AnimNextAnimGraph/Private/Traits/ReferencePoseTrait.cpp
@@ -13,7 +13,7 @@ namespace UE::AnimNext
 	#define TRAIT_INTERFACE_ENUMERATOR(GeneratorMacro) \
 		GeneratorMacro(IEvaluate) \
 
-	GENERATE_ANIM_TRAIT_IMPLEMENTATION(FReferencePoseTrait, TRAIT_INTERFACE_ENUMERATOR, NULL_ANIM_TRAIT_EVENT_ENUMERATOR)
+	GENERATE_ANIM_TRAIT_IMPLEMENTATION(FReferencePoseTrait, TRAIT_INTERFACE_ENUMERATOR, NULL_ANIM_TRAIT_INTERFACE_ENUMERATOR, NULL_ANIM_TRAIT_EVENT_ENUMERATOR)
 	#undef TRAIT_INTERFACE_ENUMERATOR
 
 	void FReferencePoseTrait::PreEvaluate(FEvaluateTraversalContext& Context, const TTraitBinding<IEvaluate>& Binding) const
diff --git a/Engine/Plugins/Experimental/AnimNextAnimGraph/Source/AnimNextAnimGraph/Private/Traits/SequencePlayer.cpp b/Engine/Plugins/Experimental/AnimNextAnimGraph/Source/AnimNextAnimGraph/Private/Traits/SequencePlayer.cpp
--- a/Engine/Plugins/Experimental/AnimNextAnimGraph/Source/AnimNextAnimGraph/Private/Traits/SequencePlayer.cpp
+++ b/Engine/Plugins/Experimental/AnimNextAnimGraph/Source/AnimNextAnimGraph/Private/Traits/SequencePlayer.cpp
@@ -17,7 +17,7 @@ namespace UE::AnimNext
 		GeneratorMacro(IUpdate) \
 		GeneratorMacro(IGarbageCollection) \
 
-	GENERATE_ANIM_TRAIT_IMPLEMENTATION(FSequencePlayerTrait, TRAIT_INTERFACE_ENUMERATOR, NULL_ANIM_TRAIT_EVENT_ENUMERATOR)
+	GENERATE_ANIM_TRAIT_IMPLEMENTATION(FSequencePlayerTrait, TRAIT_INTERFACE_ENUMERATOR, NULL_ANIM_TRAIT_INTERFACE_ENUMERATOR, NULL_ANIM_TRAIT_EVENT_ENUMERATOR)
 	#undef TRAIT_INTERFACE_ENUMERATOR
 
 	void FSequencePlayerTrait::FInstanceData::Construct(const FExecutionContext& Context, const FTraitBinding& Binding)
diff --git a/Engine/Plugins/Experimental/AnimNextAnimGraph/Source/AnimNextAnimGraph/Private/Traits/SubGraphHost.cpp b/Engine/Plugins/Experimental/AnimNextAnimGraph/Source/AnimNextAnimGraph/Private/Traits/SubGraphHost.cpp
--- a/Engine/Plugins/Experimental/AnimNextAnimGraph/Source/AnimNextAnimGraph/Private/Traits/SubGraphHost.cpp
+++ b/Engine/Plugins/Experimental/AnimNextAnimGraph/Source/AnimNextAnimGraph/Private/Traits/SubGraphHost.cpp
@@ -18,7 +18,7 @@ namespace UE::AnimNext
 		GeneratorMacro(IUpdate) \
 		GeneratorMacro(IUpdateTraversal) \
 
-	GENERATE_ANIM_TRAIT_IMPLEMENTATION(FSubGraphHostTrait, TRAIT_INTERFACE_ENUMERATOR, NULL_ANIM_TRAIT_EVENT_ENUMERATOR)
+	GENERATE_ANIM_TRAIT_IMPLEMENTATION(FSubGraphHostTrait, TRAIT_INTERFACE_ENUMERATOR, NULL_ANIM_TRAIT_INTERFACE_ENUMERATOR, NULL_ANIM_TRAIT_EVENT_ENUMERATOR)
 	#undef TRAIT_INTERFACE_ENUMERATOR
 
 	void FSubGraphHostTrait::FInstanceData::Construct(const FExecutionContext& Context, const FTraitBinding& Binding)
diff --git a/Engine/Plugins/Experimental/AnimNextAnimGraph/Source/AnimNextAnimGraph/Private/Traits/SynchronizeUsingGroups.cpp b/Engine/Plugins/Experimental/AnimNextAnimGraph/Source/AnimNextAnimGraph/Private/Traits/SynchronizeUsingGroups.cpp
--- a/Engine/Plugins/Experimental/AnimNextAnimGraph/Source/AnimNextAnimGraph/Private/Traits/SynchronizeUsingGroups.cpp
+++ b/Engine/Plugins/Experimental/AnimNextAnimGraph/Source/AnimNextAnimGraph/Private/Traits/SynchronizeUsingGroups.cpp
@@ -13,7 +13,7 @@ namespace UE::AnimNext
 		GeneratorMacro(ITimeline) \
 		GeneratorMacro(IUpdate) \
 
-	GENERATE_ANIM_TRAIT_IMPLEMENTATION(FSynchronizeUsingGroupsTrait, TRAIT_INTERFACE_ENUMERATOR, NULL_ANIM_TRAIT_EVENT_ENUMERATOR)
+	GENERATE_ANIM_TRAIT_IMPLEMENTATION(FSynchronizeUsingGroupsTrait, TRAIT_INTERFACE_ENUMERATOR, NULL_ANIM_TRAIT_INTERFACE_ENUMERATOR, NULL_ANIM_TRAIT_EVENT_ENUMERATOR)
 	#undef TRAIT_INTERFACE_ENUMERATOR
 
 	void FSynchronizeUsingGroupsTrait::PreUpdate(FUpdateTraversalContext& Context, const TTraitBinding<IUpdate>& Binding, const FTraitUpdateState& TraitState) const
diff --git a/Engine/Plugins/Experimental/AnimNextAnimGraph/Source/AnimNextAnimGraph/Public/TraitInterfaces/IContinuousBlend.h b/Engine/Plugins/Experimental/AnimNextAnimGraph/Source/AnimNextAnimGraph/Public/TraitInterfaces/IContinuousBlend.h
--- a/Engine/Plugins/Experimental/AnimNextAnimGraph/Source/AnimNextAnimGraph/Public/TraitInterfaces/IContinuousBlend.h
+++ b/Engine/Plugins/Experimental/AnimNextAnimGraph/Source/AnimNextAnimGraph/Public/TraitInterfaces/IContinuousBlend.h
@@ -15,7 +15,7 @@ namespace UE::AnimNext
 	 */
 	struct ANIMNEXTANIMGRAPH_API IContinuousBlend : ITraitInterface
 	{
-		DECLARE_ANIM_TRAIT_INTERFACE(IContinuousBlend, 0xe7d79186)
+		DECLARE_ANIM_TRAIT_INTERFACE(IContinuousBlend, CNB, 0xe7d79186)
 
 		// Returns the blend weight for the specified child
 		// Multiple children can have non-zero weight but their sum must be 1.0
diff --git a/Engine/Plugins/Experimental/AnimNextAnimGraph/Source/AnimNextAnimGraph/Public/TraitInterfaces/IDiscreteBlend.h b/Engine/Plugins/Experimental/AnimNextAnimGraph/Source/AnimNextAnimGraph/Public/TraitInterfaces/IDiscreteBlend.h
--- a/Engine/Plugins/Experimental/AnimNextAnimGraph/Source/AnimNextAnimGraph/Public/TraitInterfaces/IDiscreteBlend.h
+++ b/Engine/Plugins/Experimental/AnimNextAnimGraph/Source/AnimNextAnimGraph/Public/TraitInterfaces/IDiscreteBlend.h
@@ -17,7 +17,7 @@ namespace UE::AnimNext
 	 */
 	struct ANIMNEXTANIMGRAPH_API IDiscreteBlend : ITraitInterface
 	{
-		DECLARE_ANIM_TRAIT_INTERFACE(IDiscreteBlend, 0x2d395d56)
+		DECLARE_ANIM_TRAIT_INTERFACE(IDiscreteBlend, DSB, 0x2d395d56)
 
 		// Returns the blend weight for the specified child
 		// Multiple children can have non-zero weight but their sum must be 1.0
diff --git a/Engine/Plugins/Experimental/AnimNextAnimGraph/Source/AnimNextAnimGraph/Public/TraitInterfaces/IEvaluate.h b/Engine/Plugins/Experimental/AnimNextAnimGraph/Source/AnimNextAnimGraph/Public/TraitInterfaces/IEvaluate.h
--- a/Engine/Plugins/Experimental/AnimNextAnimGraph/Source/AnimNextAnimGraph/Public/TraitInterfaces/IEvaluate.h
+++ b/Engine/Plugins/Experimental/AnimNextAnimGraph/Source/AnimNextAnimGraph/Public/TraitInterfaces/IEvaluate.h
@@ -49,7 +49,7 @@ namespace UE::AnimNext
 	 */
 	struct ANIMNEXTANIMGRAPH_API IEvaluate : ITraitInterface
 	{
-		DECLARE_ANIM_TRAIT_INTERFACE(IEvaluate, 0xa303e9e7)
+		DECLARE_ANIM_TRAIT_INTERFACE(IEvaluate, EVA, 0xa303e9e7)
 
 		// Called before a traits children are evaluated
 		virtual void PreEvaluate(FEvaluateTraversalContext& Context, const TTraitBinding<IEvaluate>& Binding) const;
diff --git a/Engine/Plugins/Experimental/AnimNextAnimGraph/Source/AnimNextAnimGraph/Public/TraitInterfaces/IGroupSynchronization.h b/Engine/Plugins/Experimental/AnimNextAnimGraph/Source/AnimNextAnimGraph/Public/TraitInterfaces/IGroupSynchronization.h
--- a/Engine/Plugins/Experimental/AnimNextAnimGraph/Source/AnimNextAnimGraph/Public/TraitInterfaces/IGroupSynchronization.h
+++ b/Engine/Plugins/Experimental/AnimNextAnimGraph/Source/AnimNextAnimGraph/Public/TraitInterfaces/IGroupSynchronization.h
@@ -16,7 +16,7 @@ namespace UE::AnimNext
 	 */
 	struct ANIMNEXTANIMGRAPH_API IGroupSynchronization : ITraitInterface
 	{
-		DECLARE_ANIM_TRAIT_INTERFACE(IGroupSynchronization, 0xf607d0fd)
+		DECLARE_ANIM_TRAIT_INTERFACE(IGroupSynchronization, GRS, 0xf607d0fd)
 
 		// Returns the group name used for synchronization
 		virtual FName GetGroupName(const FExecutionContext& Context, const TTraitBinding<IGroupSynchronization>& Binding) const;
diff --git a/Engine/Plugins/Experimental/AnimNextAnimGraph/Source/AnimNextAnimGraph/Public/TraitInterfaces/ISmoothBlend.h b/Engine/Plugins/Experimental/AnimNextAnimGraph/Source/AnimNextAnimGraph/Public/TraitInterfaces/ISmoothBlend.h
--- a/Engine/Plugins/Experimental/AnimNextAnimGraph/Source/AnimNextAnimGraph/Public/TraitInterfaces/ISmoothBlend.h
+++ b/Engine/Plugins/Experimental/AnimNextAnimGraph/Source/AnimNextAnimGraph/Public/TraitInterfaces/ISmoothBlend.h
@@ -15,7 +15,7 @@ namespace UE::AnimNext
 	 */
 	struct ANIMNEXTANIMGRAPH_API ISmoothBlend : ITraitInterface
 	{
-		DECLARE_ANIM_TRAIT_INTERFACE(ISmoothBlend, 0x1c2c1739)
+		DECLARE_ANIM_TRAIT_INTERFACE(ISmoothBlend, SMB, 0x1c2c1739)
 
 		// Returns the desired blend time for the specified child
 		virtual float GetBlendTime(const FExecutionContext& Context, const TTraitBinding<ISmoothBlend>& Binding, int32 ChildIndex) const;
diff --git a/Engine/Plugins/Experimental/AnimNextAnimGraph/Source/AnimNextAnimGraph/Public/TraitInterfaces/ITimeline.h b/Engine/Plugins/Experimental/AnimNextAnimGraph/Source/AnimNextAnimGraph/Public/TraitInterfaces/ITimeline.h
--- a/Engine/Plugins/Experimental/AnimNextAnimGraph/Source/AnimNextAnimGraph/Public/TraitInterfaces/ITimeline.h
+++ b/Engine/Plugins/Experimental/AnimNextAnimGraph/Source/AnimNextAnimGraph/Public/TraitInterfaces/ITimeline.h
@@ -15,7 +15,7 @@ namespace UE::AnimNext
 	 */
 	struct ANIMNEXTANIMGRAPH_API ITimeline : ITraitInterface
 	{
-		DECLARE_ANIM_TRAIT_INTERFACE(ITimeline, 0x53760727)
+		DECLARE_ANIM_TRAIT_INTERFACE(ITimeline, TIM, 0x53760727)
 
 		// Returns the play rate of this timeline
 		virtual float GetPlayRate(const FExecutionContext& Context, const TTraitBinding<ITimeline>& Binding) const;
diff --git a/Engine/Plugins/Experimental/AnimNextAnimGraph/Source/AnimNextAnimGraph/Public/TraitInterfaces/IUpdate.h b/Engine/Plugins/Experimental/AnimNextAnimGraph/Source/AnimNextAnimGraph/Public/TraitInterfaces/IUpdate.h
--- a/Engine/Plugins/Experimental/AnimNextAnimGraph/Source/AnimNextAnimGraph/Public/TraitInterfaces/IUpdate.h
+++ b/Engine/Plugins/Experimental/AnimNextAnimGraph/Source/AnimNextAnimGraph/Public/TraitInterfaces/IUpdate.h
@@ -229,7 +229,7 @@ namespace UE::AnimNext
 	 */
 	struct ANIMNEXTANIMGRAPH_API IUpdate : ITraitInterface
 	{
-		DECLARE_ANIM_TRAIT_INTERFACE(IUpdate, 0x59d24dc5)
+		DECLARE_ANIM_TRAIT_INTERFACE(IUpdate, UPD, 0x59d24dc5)
 
 		// Called before the first update when a trait stack becomes relevant
 		virtual void OnBecomeRelevant(FUpdateTraversalContext& Context, const TTraitBinding<IUpdate>& Binding, const FTraitUpdateState& TraitState) const;
@@ -255,7 +255,7 @@ namespace UE::AnimNext
 	 */
 	struct ANIMNEXTANIMGRAPH_API IUpdateTraversal : ITraitInterface
 	{
-		DECLARE_ANIM_TRAIT_INTERFACE(IUpdateTraversal, 0x256c21b1)
+		DECLARE_ANIM_TRAIT_INTERFACE(IUpdateTraversal, TRA, 0x256c21b1)
 
 		// Called after PreUpdate to request that children be queued with the provided context
 		virtual void QueueChildrenForTraversal(FUpdateTraversalContext& Context, const TTraitBinding<IUpdateTraversal>& Binding, const FTraitUpdateState& TraitState, FUpdateTraversalQueue& TraversalQueue) const;
diff --git a/Engine/Plugins/Experimental/AnimNextAnimGraph/Source/AnimNextAnimGraphTestSuite/Private/AnimNextAnimGraphTraitEventTest.cpp b/Engine/Plugins/Experimental/AnimNextAnimGraph/Source/AnimNextAnimGraphTestSuite/Private/AnimNextAnimGraphTraitEventTest.cpp
--- a/Engine/Plugins/Experimental/AnimNextAnimGraph/Source/AnimNextAnimGraphTestSuite/Private/AnimNextAnimGraphTraitEventTest.cpp
+++ b/Engine/Plugins/Experimental/AnimNextAnimGraph/Source/AnimNextAnimGraphTestSuite/Private/AnimNextAnimGraphTraitEventTest.cpp
@@ -48,7 +48,7 @@ namespace UE::AnimNext
 		GeneratorMacro(FTraitGraphTest_EventAB_NoChildren::OnEventA) \
 		GeneratorMacro(FTraitGraphTest_EventAB_NoChildren::OnEventB) \
 
-	GENERATE_ANIM_TRAIT_IMPLEMENTATION(FTraitGraphTest_EventAB_NoChildren, NULL_ANIM_TRAIT_INTERFACE_ENUMERATOR, TRAIT_EVENT_ENUMERATOR)
+	GENERATE_ANIM_TRAIT_IMPLEMENTATION(FTraitGraphTest_EventAB_NoChildren, NULL_ANIM_TRAIT_INTERFACE_ENUMERATOR, NULL_ANIM_TRAIT_INTERFACE_ENUMERATOR, TRAIT_EVENT_ENUMERATOR)
 	#undef TRAIT_EVENT_ENUMERATOR
 
 	struct FTraitGraphTest_EventAB_OneChild : FBaseTrait, IHierarchy
@@ -122,7 +122,7 @@ namespace UE::AnimNext
 		GeneratorMacro(FTraitGraphTest_EventAB_OneChild::OnEventA) \
 		GeneratorMacro(FTraitGraphTest_EventAB_OneChild::OnEventB) \
 
-	GENERATE_ANIM_TRAIT_IMPLEMENTATION(FTraitGraphTest_EventAB_OneChild, TRAIT_INTERFACE_ENUMERATOR, TRAIT_EVENT_ENUMERATOR)
+	GENERATE_ANIM_TRAIT_IMPLEMENTATION(FTraitGraphTest_EventAB_OneChild, TRAIT_INTERFACE_ENUMERATOR, NULL_ANIM_TRAIT_INTERFACE_ENUMERATOR, TRAIT_EVENT_ENUMERATOR)
 	#undef TRAIT_INTERFACE_ENUMERATOR
 	#undef TRAIT_EVENT_ENUMERATOR
 
@@ -178,7 +178,7 @@ namespace UE::AnimNext
 		GeneratorMacro(FTraitGraphTest_EventAB_TwoChildren::OnEventA) \
 		GeneratorMacro(FTraitGraphTest_EventAB_TwoChildren::OnEventB) \
 
-	GENERATE_ANIM_TRAIT_IMPLEMENTATION(FTraitGraphTest_EventAB_TwoChildren, TRAIT_INTERFACE_ENUMERATOR, TRAIT_EVENT_ENUMERATOR)
+	GENERATE_ANIM_TRAIT_IMPLEMENTATION(FTraitGraphTest_EventAB_TwoChildren, TRAIT_INTERFACE_ENUMERATOR, NULL_ANIM_TRAIT_INTERFACE_ENUMERATOR, TRAIT_EVENT_ENUMERATOR)
 	#undef TRAIT_INTERFACE_ENUMERATOR
 	#undef TRAIT_EVENT_ENUMERATOR
 }
diff --git a/Engine/Plugins/Experimental/AnimNextAnimGraph/Source/AnimNextAnimGraphTestSuite/Private/AnimNextAnimGraphTraitGraphTest.cpp b/Engine/Plugins/Experimental/AnimNextAnimGraph/Source/AnimNextAnimGraphTestSuite/Private/AnimNextAnimGraphTraitGraphTest.cpp
--- a/Engine/Plugins/Experimental/AnimNextAnimGraph/Source/AnimNextAnimGraphTestSuite/Private/AnimNextAnimGraphTraitGraphTest.cpp
+++ b/Engine/Plugins/Experimental/AnimNextAnimGraph/Source/AnimNextAnimGraphTestSuite/Private/AnimNextAnimGraphTraitGraphTest.cpp
@@ -75,13 +75,25 @@ namespace UE::AnimNext
 		GeneratorMacro(IEvaluate) \
 		GeneratorMacro(IUpdate) \
 
-	GENERATE_ANIM_TRAIT_IMPLEMENTATION(FTestTrait, TRAIT_INTERFACE_ENUMERATOR, NULL_ANIM_TRAIT_EVENT_ENUMERATOR)
+	GENERATE_ANIM_TRAIT_IMPLEMENTATION(FTestTrait, TRAIT_INTERFACE_ENUMERATOR, NULL_ANIM_TRAIT_INTERFACE_ENUMERATOR, NULL_ANIM_TRAIT_EVENT_ENUMERATOR)
+	#undef TRAIT_INTERFACE_ENUMERATOR
+
+	// --- FTestBasicTrait ---
+
+	struct FTestBasicTrait final : FBaseTrait
+	{
+		DECLARE_ANIM_TRAIT(FTestBasicTrait, 0x24ce4372, FBaseTrait)
+
+		using FSharedData = FTestTraitSharedData;
+	};
+
+	GENERATE_ANIM_TRAIT_IMPLEMENTATION(FTestBasicTrait, NULL_ANIM_TRAIT_INTERFACE_ENUMERATOR, NULL_ANIM_TRAIT_INTERFACE_ENUMERATOR, NULL_ANIM_TRAIT_EVENT_ENUMERATOR)
 	#undef TRAIT_INTERFACE_ENUMERATOR
 }
 
-IMPLEMENT_SIMPLE_AUTOMATION_TEST(FAnimationAnimNextRuntimeTest_GraphAddTrait, "Animation.AnimNext.Runtime.Graph.AddTrait", EAutomationTestFlags::EditorContext | EAutomationTestFlags::EngineFilter)
+IMPLEMENT_SIMPLE_AUTOMATION_TEST(FAnimationAnimNextEditorTest_GraphAddTrait, "Animation.AnimNext.Editor.Graph.AddTrait", EAutomationTestFlags::EditorContext | EAutomationTestFlags::EngineFilter)
 
-bool FAnimationAnimNextRuntimeTest_GraphAddTrait::RunTest(const FString& InParameters)
+bool FAnimationAnimNextEditorTest_GraphAddTrait::RunTest(const FString& InParameters)
 {
 	using namespace UE::AnimNext;
 
@@ -92,15 +104,15 @@ bool FAnimationAnimNextRuntimeTest_GraphAddTrait::RunTest(const FString& InParam
 
 		UFactory* GraphFactory = NewObject<UAnimNextGraphFactory>();
 		UAnimNextGraph* AnimNextGraph = CastChecked<UAnimNextGraph>(GraphFactory->FactoryCreateNew(UAnimNextGraph::StaticClass(), GetTransientPackage(), TEXT("TestAnimNextGraph"), RF_Transient, nullptr, nullptr, NAME_None));
-		UE_RETURN_ON_ERROR(AnimNextGraph != nullptr, "FAnimationAnimNextRuntimeTest_GraphAddTrait -> Failed to create animation graph");
+		UE_RETURN_ON_ERROR(AnimNextGraph != nullptr, "FAnimationAnimNextEditorTest_GraphAddTrait -> Failed to create animation graph");
 
 		UAnimNextGraph_EditorData* EditorData = UncookedOnly::FUtils::GetEditorData(AnimNextGraph);
-		UE_RETURN_ON_ERROR(EditorData != nullptr, "FAnimationAnimNextRuntimeTest_GraphAddTrait -> Failed to find animation graph editor data");
+		UE_RETURN_ON_ERROR(EditorData != nullptr, "FAnimationAnimNextEditorTest_GraphAddTrait -> Failed to find animation graph editor data");
 
-		UE_RETURN_ON_ERROR(EditorData->AddAnimationGraph(FRigUnit_AnimNextGraphRoot::DefaultEntryPoint, false) != nullptr, "FAnimationAnimNextRuntimeTest_GraphAddTrait -> Failed to add animation graph");
+		UE_RETURN_ON_ERROR(EditorData->AddAnimationGraph(FRigUnit_AnimNextGraphRoot::DefaultEntryPoint, false) != nullptr, "FAnimationAnimNextEditorTest_GraphAddTrait -> Failed to add animation graph");
 
 		URigVMController* Controller = EditorData->GetRigVMClient()->GetController(EditorData->GetRigVMClient()->GetDefaultModel());
-		UE_RETURN_ON_ERROR(Controller != nullptr, "FAnimationAnimNextRuntimeTest_GraphAddTrait -> Failed to get RigVM controller");
+		UE_RETURN_ON_ERROR(Controller != nullptr, "FAnimationAnimNextEditorTest_GraphAddTrait -> Failed to get RigVM controller");
 
 		// Create an empty trait stack node
 		URigVMUnitNode* DecoratorStackNode = Controller->AddUnitNode(FRigUnit_AnimNextTraitStack::StaticStruct(), FRigVMStruct::ExecuteName, FVector2D(0.0f, 0.0f), FString(), false);
@@ -111,10 +123,10 @@ bool FAnimationAnimNextRuntimeTest_GraphAddTrait::RunTest(const FString& InParam
 		UE_RETURN_ON_ERROR(CppDecoratorStruct != nullptr, "FAnimationAnimNextRuntimeTest_GraphAddDecorator -> Failed to get find Cpp decorator static struct");
 
 		const FTrait* Trait = FTraitRegistry::Get().Find(FTestTrait::TraitUID);
-		UE_RETURN_ON_ERROR(Trait != nullptr, "FAnimationAnimNextRuntimeTest_GraphAddTrait -> Failed to find test trait");
+		UE_RETURN_ON_ERROR(Trait != nullptr, "FAnimationAnimNextEditorTest_GraphAddTrait -> Failed to find test trait");
 
 		UScriptStruct* ScriptStruct = Trait->GetTraitSharedDataStruct();
-		UE_RETURN_ON_ERROR(ScriptStruct != nullptr, "FAnimationAnimNextRuntimeTest_GraphAddTrait -> Failed to find trait shared data struct");
+		UE_RETURN_ON_ERROR(ScriptStruct != nullptr, "FAnimationAnimNextEditorTest_GraphAddTrait -> Failed to find trait shared data struct");
 
 		FString DefaultValue;
 		{
@@ -122,10 +134,10 @@ bool FAnimationAnimNextRuntimeTest_GraphAddTrait::RunTest(const FString& InParam
 			FRigDecorator_AnimNextCppDecorator CppDecoratorStructInstance;
 			CppDecoratorStructInstance.DecoratorSharedDataStruct = ScriptStruct;
 
-			UE_RETURN_ON_ERROR(CppDecoratorStructInstance.CanBeAddedToNode(DecoratorStackNode, nullptr), "FAnimationAnimNextRuntimeTest_GraphAddTrait -> Trait cannot be added to decorator stack node");
+			UE_RETURN_ON_ERROR(CppDecoratorStructInstance.CanBeAddedToNode(DecoratorStackNode, nullptr), "FAnimationAnimNextEditorTest_GraphAddTrait -> Trait cannot be added to decorator stack node");
 
 			const FProperty* Prop = FAnimNextCppDecoratorWrapper::StaticStruct()->FindPropertyByName(GET_MEMBER_NAME_STRING_CHECKED(FAnimNextCppDecoratorWrapper, CppDecorator));
-			UE_RETURN_ON_ERROR(Prop != nullptr, "FAnimationAnimNextRuntimeTest_GraphAddTrait -> Failed to find wrapper property");
+			UE_RETURN_ON_ERROR(Prop != nullptr, "FAnimationAnimNextEditorTest_GraphAddTrait -> Failed to find wrapper property");
 
 			Prop->ExportText_Direct(DefaultValue, &CppDecoratorStructInstance, &DefaultCppDecoratorStructInstance, nullptr, PPF_None);
 		}
@@ -139,44 +151,207 @@ bool FAnimationAnimNextRuntimeTest_GraphAddTrait::RunTest(const FString& InParam
 			*CppDecoratorStruct->GetPathName(),
 			*DisplayName,
 			DefaultValue, INDEX_NONE, true, true);
-		UE_RETURN_ON_ERROR(DecoratorName == DisplayName, TEXT("FAnimationAnimNextRuntimeTest_GraphAddTrait -> Unexpected decorator name"));
+		UE_RETURN_ON_ERROR(DecoratorName == DisplayName, TEXT("FAnimationAnimNextEditorTest_GraphAddTrait -> Unexpected decorator name"));
 
 		URigVMPin* DecoratorPin = DecoratorStackNode->FindPin(*DisplayName);
-		UE_RETURN_ON_ERROR(DecoratorPin != nullptr, TEXT("FAnimationAnimNextRuntimeTest_GraphAddTrait -> Failed to find decorator pin"));
+		UE_RETURN_ON_ERROR(DecoratorPin != nullptr, TEXT("FAnimationAnimNextEditorTest_GraphAddTrait -> Failed to find decorator pin"));
 
 		// Our first pin is the hard coded output result, decorator pins follow
-		UE_RETURN_ON_ERROR(DecoratorStackNode->GetPins().Num() == 2, TEXT("FAnimationAnimNextRuntimeTest_GraphAddTrait -> Unexpected number of pins"));
-		UE_RETURN_ON_ERROR(DecoratorPin->IsDecoratorPin() == true, TEXT("FAnimationAnimNextRuntimeTest_GraphAddTrait -> Unexpected pin type"));
-		UE_RETURN_ON_ERROR(DecoratorPin->GetFName() == DecoratorName, TEXT("FAnimationAnimNextRuntimeTest_GraphAddTrait -> Unexpected pin name"));
-		UE_RETURN_ON_ERROR(DecoratorPin->GetCPPTypeObject() == FRigDecorator_AnimNextCppDecorator::StaticStruct(), TEXT("FAnimationAnimNextRuntimeTest_GraphAddTrait -> Unexpected pin type"));
+		UE_RETURN_ON_ERROR(DecoratorStackNode->GetPins().Num() == 2, TEXT("FAnimationAnimNextEditorTest_GraphAddTrait -> Unexpected number of pins"));
+		UE_RETURN_ON_ERROR(DecoratorPin->IsDecoratorPin() == true, TEXT("FAnimationAnimNextEditorTest_GraphAddTrait -> Unexpected pin type"));
+		UE_RETURN_ON_ERROR(DecoratorPin->GetFName() == DecoratorName, TEXT("FAnimationAnimNextEditorTest_GraphAddTrait -> Unexpected pin name"));
+		UE_RETURN_ON_ERROR(DecoratorPin->GetCPPTypeObject() == FRigDecorator_AnimNextCppDecorator::StaticStruct(), TEXT("FAnimationAnimNextEditorTest_GraphAddTrait -> Unexpected pin type"));
 
 		// Our first sub-pin is the hard coded script struct member that parametrizes the decorator, dynamic decorator sub-pins follow
-		UE_RETURN_ON_ERROR(DecoratorPin->GetSubPins().Num() == 6, TEXT("FAnimationAnimNextRuntimeTest_GraphAddTrait -> Unexpected decorator sub pins"));
+		UE_RETURN_ON_ERROR(DecoratorPin->GetSubPins().Num() == 6, TEXT("FAnimationAnimNextEditorTest_GraphAddTrait -> Unexpected decorator sub pins"));
 
 		// SomeInt32
-		UE_RETURN_ON_ERROR(DecoratorPin->GetSubPins()[1]->GetCPPType() == TEXT("int32"), TEXT("FAnimationAnimNextRuntimeTest_GraphAddTrait -> Unexpected decorator pin type"));
-		UE_RETURN_ON_ERROR(DecoratorPin->GetSubPins()[1]->GetDefaultValue() == TEXT("3"), TEXT("FAnimationAnimNextRuntimeTest_GraphAddTrait -> Unexpected decorator pin value"));
-		UE_RETURN_ON_ERROR(!DecoratorPin->GetSubPins()[1]->IsLazy(), TEXT("FAnimationAnimNextRuntimeTest_GraphAddTrait -> Expected non-lazy decorator pin"));
+		UE_RETURN_ON_ERROR(DecoratorPin->GetSubPins()[1]->GetCPPType() == TEXT("int32"), TEXT("FAnimationAnimNextEditorTest_GraphAddTrait -> Unexpected decorator pin type"));
+		UE_RETURN_ON_ERROR(DecoratorPin->GetSubPins()[1]->GetDefaultValue() == TEXT("3"), TEXT("FAnimationAnimNextEditorTest_GraphAddTrait -> Unexpected decorator pin value"));
+		UE_RETURN_ON_ERROR(!DecoratorPin->GetSubPins()[1]->IsLazy(), TEXT("FAnimationAnimNextEditorTest_GraphAddTrait -> Expected non-lazy decorator pin"));
 
 		// SomeFloat
-		UE_RETURN_ON_ERROR(DecoratorPin->GetSubPins()[2]->GetCPPType() == TEXT("float"), TEXT("FAnimationAnimNextRuntimeTest_GraphAddTrait -> Unexpected decorator pin type"));
-		UE_RETURN_ON_ERROR(DecoratorPin->GetSubPins()[2]->GetDefaultValue() == TEXT("34.000000"), TEXT("FAnimationAnimNextRuntimeTest_GraphAddTrait -> Unexpected decorator pin value"));
-		UE_RETURN_ON_ERROR(!DecoratorPin->GetSubPins()[2]->IsLazy(), TEXT("FAnimationAnimNextRuntimeTest_GraphAddTrait -> Expected non-lazy decorator pin"));
+		UE_RETURN_ON_ERROR(DecoratorPin->GetSubPins()[2]->GetCPPType() == TEXT("float"), TEXT("FAnimationAnimNextEditorTest_GraphAddTrait -> Unexpected decorator pin type"));
+		UE_RETURN_ON_ERROR(DecoratorPin->GetSubPins()[2]->GetDefaultValue() == TEXT("34.000000"), TEXT("FAnimationAnimNextEditorTest_GraphAddTrait -> Unexpected decorator pin value"));
+		UE_RETURN_ON_ERROR(!DecoratorPin->GetSubPins()[2]->IsLazy(), TEXT("FAnimationAnimNextEditorTest_GraphAddTrait -> Expected non-lazy decorator pin"));
 
 		// SomeLatentInt32
-		UE_RETURN_ON_ERROR(DecoratorPin->GetSubPins()[3]->GetCPPType() == TEXT("int32"), TEXT("FAnimationAnimNextRuntimeTest_GraphAddTrait -> Unexpected decorator pin type"));
-		UE_RETURN_ON_ERROR(DecoratorPin->GetSubPins()[3]->GetDefaultValue() == TEXT("3"), TEXT("FAnimationAnimNextRuntimeTest_GraphAddTrait -> Unexpected decorator pin value"));
-		UE_RETURN_ON_ERROR(DecoratorPin->GetSubPins()[3]->IsLazy(), TEXT("FAnimationAnimNextRuntimeTest_GraphAddTrait -> Expected lazy decorator pin"));
+		UE_RETURN_ON_ERROR(DecoratorPin->GetSubPins()[3]->GetCPPType() == TEXT("int32"), TEXT("FAnimationAnimNextEditorTest_GraphAddTrait -> Unexpected decorator pin type"));
+		UE_RETURN_ON_ERROR(DecoratorPin->GetSubPins()[3]->GetDefaultValue() == TEXT("3"), TEXT("FAnimationAnimNextEditorTest_GraphAddTrait -> Unexpected decorator pin value"));
+		UE_RETURN_ON_ERROR(DecoratorPin->GetSubPins()[3]->IsLazy(), TEXT("FAnimationAnimNextEditorTest_GraphAddTrait -> Expected lazy decorator pin"));
 
 		// SomeOtherLatentInt32
-		UE_RETURN_ON_ERROR(DecoratorPin->GetSubPins()[4]->GetCPPType() == TEXT("int32"), TEXT("FAnimationAnimNextRuntimeTest_GraphAddTrait -> Unexpected decorator pin type"));
-		UE_RETURN_ON_ERROR(DecoratorPin->GetSubPins()[4]->GetDefaultValue() == TEXT("3"), TEXT("FAnimationAnimNextRuntimeTest_GraphAddTrait -> Unexpected decorator pin value"));
-		UE_RETURN_ON_ERROR(DecoratorPin->GetSubPins()[4]->IsLazy(), TEXT("FAnimationAnimNextRuntimeTest_GraphAddTrait -> Expected lazy decorator pin"));
+		UE_RETURN_ON_ERROR(DecoratorPin->GetSubPins()[4]->GetCPPType() == TEXT("int32"), TEXT("FAnimationAnimNextEditorTest_GraphAddTrait -> Unexpected decorator pin type"));
+		UE_RETURN_ON_ERROR(DecoratorPin->GetSubPins()[4]->GetDefaultValue() == TEXT("3"), TEXT("FAnimationAnimNextEditorTest_GraphAddTrait -> Unexpected decorator pin value"));
+		UE_RETURN_ON_ERROR(DecoratorPin->GetSubPins()[4]->IsLazy(), TEXT("FAnimationAnimNextEditorTest_GraphAddTrait -> Expected lazy decorator pin"));
 
 		// SomeLatentFloat
-		UE_RETURN_ON_ERROR(DecoratorPin->GetSubPins()[5]->GetCPPType() == TEXT("float"), TEXT("FAnimationAnimNextRuntimeTest_GraphAddTrait -> Unexpected decorator pin type"));
-		UE_RETURN_ON_ERROR(DecoratorPin->GetSubPins()[5]->GetDefaultValue() == TEXT("34.000000"), TEXT("FAnimationAnimNextRuntimeTest_GraphAddTrait -> Unexpected decorator pin value"));
-		UE_RETURN_ON_ERROR(DecoratorPin->GetSubPins()[5]->IsLazy(), TEXT("FAnimationAnimNextRuntimeTest_GraphAddTrait -> Expected lazy decorator pin"));
+		UE_RETURN_ON_ERROR(DecoratorPin->GetSubPins()[5]->GetCPPType() == TEXT("float"), TEXT("FAnimationAnimNextEditorTest_GraphAddTrait -> Unexpected decorator pin type"));
+		UE_RETURN_ON_ERROR(DecoratorPin->GetSubPins()[5]->GetDefaultValue() == TEXT("34.000000"), TEXT("FAnimationAnimNextEditorTest_GraphAddTrait -> Unexpected decorator pin value"));
+		UE_RETURN_ON_ERROR(DecoratorPin->GetSubPins()[5]->IsLazy(), TEXT("FAnimationAnimNextEditorTest_GraphAddTrait -> Expected lazy decorator pin"));
+	}
+
+	Tests::FUtils::CleanupAfterTests();
+
+	return true;
+}
+
+IMPLEMENT_SIMPLE_AUTOMATION_TEST(FAnimationAnimNextEditorTest_GraphTraitOperations, "Animation.AnimNext.Editor.Graph.TraitOperations", EAutomationTestFlags::EditorContext | EAutomationTestFlags::EngineFilter)
+
+bool FAnimationAnimNextEditorTest_GraphTraitOperations::RunTest(const FString& InParameters)
+{
+	using namespace UE::AnimNext;
+
+	{
+		AUTO_REGISTER_ANIM_TRAIT(FTestTrait)
+		AUTO_REGISTER_ANIM_TRAIT(FTestBasicTrait)
+
+		FScopedClearNodeTemplateRegistry ScopedClearNodeTemplateRegistry;
+
+		UFactory* GraphFactory = NewObject<UAnimNextGraphFactory>();
+		UAnimNextGraph* AnimNextGraph = CastChecked<UAnimNextGraph>(GraphFactory->FactoryCreateNew(UAnimNextGraph::StaticClass(), GetTransientPackage(), TEXT("TestAnimNextGraph"), RF_Transient, nullptr, nullptr, NAME_None));
+		UE_RETURN_ON_ERROR(AnimNextGraph != nullptr, "FAnimationAnimNextEditorTest_GraphTraitOperations -> Failed to create animation graph");
+
+		UAnimNextGraph_EditorData* EditorData = UncookedOnly::FUtils::GetEditorData(AnimNextGraph);
+		UE_RETURN_ON_ERROR(EditorData != nullptr, "FAnimationAnimNextEditorTest_GraphTraitOperations -> Failed to find animation graph editor data");
+
+		UE_RETURN_ON_ERROR(EditorData->AddAnimationGraph(FRigUnit_AnimNextGraphRoot::DefaultEntryPoint, false) != nullptr, "FAnimationAnimNextEditorTest_GraphTraitOperations -> Failed to add animation graph");
+
+		UAnimNextGraph_Controller* Controller = Cast<UAnimNextGraph_Controller>(EditorData->GetRigVMClient()->GetController(EditorData->GetRigVMClient()->GetDefaultModel()));
+		UE_RETURN_ON_ERROR(Controller != nullptr, "FAnimationAnimNextEditorTest_GraphTraitOperations -> Failed to get RigVM controller");
+
+		// Create an empty trait stack node
+		URigVMUnitNode* TraitStackNode = Controller->AddUnitNode(FRigUnit_AnimNextTraitStack::StaticStruct(), FRigVMStruct::ExecuteName, FVector2D(0.0f, 0.0f), FString(), false);
+		UE_RETURN_ON_ERROR(TraitStackNode != nullptr, "FAnimationAnimNextEditorTest_GraphTraitOperations -> Failed to create decorator stack node");
+
+		const TArray<URigVMPin*>& NodePins = TraitStackNode->GetPins();
+
+		FName TraitInstanceName = NAME_None;
+
+		// --- Add a trait ---
+		{
+			const FTrait* Trait = FTraitRegistry::Get().Find(FTestTrait::TraitUID);
+			UE_RETURN_ON_ERROR(Trait != nullptr, "FAnimationAnimNextEditorTest_GraphTraitOperations -> Failed to find test trait");
+
+			const FName TraitTypeName = *Trait->GetTraitName();
+
+			TraitInstanceName = Controller->AddTrait(TraitStackNode->GetFName(), TraitTypeName, INDEX_NONE);
+			UE_RETURN_ON_ERROR(TraitInstanceName == TraitTypeName, TEXT("FAnimationAnimNextEditorTest_GraphTraitOperations -> Unexpected Trait name"));
+
+			URigVMPin* TraitPin = TraitStackNode->FindPin(TraitInstanceName.ToString());
+			UE_RETURN_ON_ERROR(TraitPin != nullptr, TEXT("FAnimationAnimNextEditorTest_GraphTraitOperations -> Failed to find Trait pin"));
+
+			// Our first pin is the hard coded output result, decorator pins follow
+			UE_RETURN_ON_ERROR(NodePins.Num() == 2, TEXT("FAnimationAnimNextEditorTest_GraphTraitOperations -> Unexpected number of pins"));
+			UE_RETURN_ON_ERROR(TraitPin->IsDecoratorPin() == true, TEXT("FAnimationAnimNextEditorTest_GraphTraitOperations -> Unexpected pin type"));
+			UE_RETURN_ON_ERROR(TraitPin->GetFName() == TraitInstanceName, TEXT("FAnimationAnimNextEditorTest_GraphTraitOperations -> Unexpected pin name"));
+			UE_RETURN_ON_ERROR(TraitPin->GetCPPTypeObject() == FRigDecorator_AnimNextCppDecorator::StaticStruct(), TEXT("FAnimationAnimNextEditorTest_GraphTraitOperations -> Unexpected pin type"));
+
+			// Our first sub-pin is the hard coded script struct member that parametrizes the decorator, dynamic decorator sub-pins follow
+			UE_RETURN_ON_ERROR(TraitPin->GetSubPins().Num() == 6, TEXT("FAnimationAnimNextEditorTest_GraphTraitOperations -> Unexpected decorator sub pins"));
+
+			// SomeInt32
+			UE_RETURN_ON_ERROR(TraitPin->GetSubPins()[1]->GetCPPType() == TEXT("int32"), TEXT("FAnimationAnimNextEditorTest_GraphTraitOperations -> Unexpected decorator pin type"));
+			UE_RETURN_ON_ERROR(TraitPin->GetSubPins()[1]->GetDefaultValue() == TEXT("3"), TEXT("FAnimationAnimNextEditorTest_GraphTraitOperations -> Unexpected decorator pin value"));
+			UE_RETURN_ON_ERROR(!TraitPin->GetSubPins()[1]->IsLazy(), TEXT("FAnimationAnimNextEditorTest_GraphTraitOperations -> Expected non-lazy decorator pin"));
+
+			// SomeFloat
+			UE_RETURN_ON_ERROR(TraitPin->GetSubPins()[2]->GetCPPType() == TEXT("float"), TEXT("FAnimationAnimNextEditorTest_GraphTraitOperations -> Unexpected decorator pin type"));
+			UE_RETURN_ON_ERROR(TraitPin->GetSubPins()[2]->GetDefaultValue() == TEXT("34.000000"), TEXT("FAnimationAnimNextEditorTest_GraphTraitOperations -> Unexpected decorator pin value"));
+			UE_RETURN_ON_ERROR(!TraitPin->GetSubPins()[2]->IsLazy(), TEXT("FAnimationAnimNextEditorTest_GraphTraitOperations -> Expected non-lazy decorator pin"));
+
+			// SomeLatentInt32
+			UE_RETURN_ON_ERROR(TraitPin->GetSubPins()[3]->GetCPPType() == TEXT("int32"), TEXT("FAnimationAnimNextEditorTest_GraphTraitOperations -> Unexpected decorator pin type"));
+			UE_RETURN_ON_ERROR(TraitPin->GetSubPins()[3]->GetDefaultValue() == TEXT("3"), TEXT("FAnimationAnimNextEditorTest_GraphTraitOperations -> Unexpected decorator pin value"));
+			UE_RETURN_ON_ERROR(TraitPin->GetSubPins()[3]->IsLazy(), TEXT("FAnimationAnimNextEditorTest_GraphTraitOperations -> Expected lazy decorator pin"));
+
+			// SomeOtherLatentInt32
+			UE_RETURN_ON_ERROR(TraitPin->GetSubPins()[4]->GetCPPType() == TEXT("int32"), TEXT("FAnimationAnimNextEditorTest_GraphTraitOperations -> Unexpected decorator pin type"));
+			UE_RETURN_ON_ERROR(TraitPin->GetSubPins()[4]->GetDefaultValue() == TEXT("3"), TEXT("FAnimationAnimNextEditorTest_GraphTraitOperations -> Unexpected decorator pin value"));
+			UE_RETURN_ON_ERROR(TraitPin->GetSubPins()[4]->IsLazy(), TEXT("FAnimationAnimNextEditorTest_GraphTraitOperations -> Expected lazy decorator pin"));
+
+			// SomeLatentFloat
+			UE_RETURN_ON_ERROR(TraitPin->GetSubPins()[5]->GetCPPType() == TEXT("float"), TEXT("FAnimationAnimNextEditorTest_GraphTraitOperations -> Unexpected decorator pin type"));
+			UE_RETURN_ON_ERROR(TraitPin->GetSubPins()[5]->GetDefaultValue() == TEXT("34.000000"), TEXT("FAnimationAnimNextEditorTest_GraphTraitOperations -> Unexpected decorator pin value"));
+			UE_RETURN_ON_ERROR(TraitPin->GetSubPins()[5]->IsLazy(), TEXT("FAnimationAnimNextEditorTest_GraphTraitOperations -> Expected lazy decorator pin"));
+		}
+
+		// --- Undo Add Trait ---
+		{
+			Controller->Undo();
+
+			URigVMPin* TraitPin = TraitStackNode->FindPin(TraitInstanceName.ToString());
+			UE_RETURN_ON_ERROR(TraitPin == nullptr, TEXT("FAnimationAnimNextEditorTest_GraphTraitOperations -> Undo AddTrait failed, Trait pin is still present"));
+
+			// Our first pin is the hard coded output result, decorator pins follow
+			const URigVMPin* FirstPin = NodePins[0];
+			UE_RETURN_ON_ERROR(NodePins.Num() == 1, TEXT("FAnimationAnimNextEditorTest_GraphTraitOperations -> Unexpected number of pins"));
+			UE_RETURN_ON_ERROR(FirstPin->IsDecoratorPin() == false, TEXT("FAnimationAnimNextEditorTest_GraphTraitOperations -> Unexpected pin type"));
+		}
+
+		// --- Redo Add Trait ---
+		{
+			Controller->Redo();
+
+			URigVMPin* TraitPin = TraitStackNode->FindPin(TraitInstanceName.ToString());
+			UE_RETURN_ON_ERROR(TraitPin != nullptr, TEXT("FAnimationAnimNextEditorTest_GraphTraitOperations -> Redo AddTrait failed, can not find Trait pin"));
+
+			// Our first pin is the hard coded output result, trait pins follow
+			UE_RETURN_ON_ERROR(NodePins.Num() == 2, TEXT("FAnimationAnimNextEditorTest_GraphTraitOperations -> Unexpected number of pins"));
+			UE_RETURN_ON_ERROR(TraitPin->IsDecoratorPin() == true, TEXT("FAnimationAnimNextEditorTest_GraphTraitOperations -> Unexpected pin type"));
+			UE_RETURN_ON_ERROR(TraitPin->GetFName() == TraitInstanceName, TEXT("FAnimationAnimNextEditorTest_GraphTraitOperations -> Unexpected pin name"));
+			UE_RETURN_ON_ERROR(TraitPin->GetCPPTypeObject() == FRigDecorator_AnimNextCppDecorator::StaticStruct(), TEXT("FAnimationAnimNextEditorTest_GraphTraitOperations -> Unexpected pin type"));
+		}
+
+		// --- Remove the created trait ---
+		{
+			Controller->RemoveTrait(TraitStackNode->GetFName(), TraitInstanceName);
+
+			// Our first pin is the hard coded output result, decorator pins follow
+			UE_RETURN_ON_ERROR(NodePins.Num() == 1, TEXT("FAnimationAnimNextEditorTest_GraphTraitOperations -> Unexpected number of pins"));
+
+			URigVMPin* DeletedTraitPin = TraitStackNode->FindPin(TraitInstanceName.ToString());
+			UE_RETURN_ON_ERROR(DeletedTraitPin == nullptr, TEXT("FAnimationAnimNextEditorTest_GraphTraitOperations -> Failed to remove Trait pin"));
+
+			// Only the output result pin should be in the pin array
+			const URigVMPin* FirstPin = NodePins[0];
+			UE_RETURN_ON_ERROR(FirstPin->IsDecoratorPin() == false, TEXT("FAnimationAnimNextEditorTest_GraphTraitOperations -> Unexpected pin type"));
+			UE_RETURN_ON_ERROR(FirstPin->GetFName() != TraitInstanceName, TEXT("FAnimationAnimNextEditorTest_GraphTraitOperations -> Unexpected pin name"));
+		}
+
+		// --- Undo Remove Trait ---
+		{
+			Controller->Undo();
+
+			URigVMPin* TraitPin = TraitStackNode->FindPin(TraitInstanceName.ToString());
+			UE_RETURN_ON_ERROR(TraitPin != nullptr, TEXT("FAnimationAnimNextEditorTest_GraphTraitOperations -> Undo failed, unable to find Trait pin"));
+
+			// Our first pin is the hard coded output result, decorator pins follow
+			UE_RETURN_ON_ERROR(NodePins.Num() == 2, TEXT("FAnimationAnimNextEditorTest_GraphTraitOperations -> Unexpected number of pins"));
+			UE_RETURN_ON_ERROR(TraitPin->IsDecoratorPin() == true, TEXT("FAnimationAnimNextEditorTest_GraphTraitOperations -> Unexpected pin type"));
+			UE_RETURN_ON_ERROR(TraitPin->GetFName() == TraitInstanceName, TEXT("FAnimationAnimNextEditorTest_GraphTraitOperations -> Unexpected pin name"));
+			UE_RETURN_ON_ERROR(TraitPin->GetCPPTypeObject() == FRigDecorator_AnimNextCppDecorator::StaticStruct(), TEXT("FAnimationAnimNextEditorTest_GraphTraitOperations -> Unexpected pin type"));
+		}
+
+		// --- Swap the FTestTrait with FTestBasicTrait ---
+		{
+			const FTrait* BasicTrait = FTraitRegistry::Get().Find(FTestBasicTrait::TraitUID);
+			UE_RETURN_ON_ERROR(BasicTrait != nullptr, "FAnimationAnimNextEditorTest_GraphTraitOperations -> Failed to find test basic trait");
+
+			const FName BasicTraitTypeName = *BasicTrait->GetTraitName();
+
+			TraitInstanceName = Controller->SwapTrait(TraitStackNode->GetFName(), TraitInstanceName, 1, BasicTraitTypeName);
+			UE_RETURN_ON_ERROR(TraitInstanceName == BasicTraitTypeName, TEXT("FAnimationAnimNextEditorTest_GraphTraitOperations -> Unexpected Trait name"));
+
+			URigVMPin* TraitPin = TraitStackNode->FindPin(TraitInstanceName.ToString());
+			UE_RETURN_ON_ERROR(TraitPin != nullptr, TEXT("FAnimationAnimNextEditorTest_GraphTraitOperations -> Failed to find FTestBasicTrait pin"));
+
+			// Our first pin is the hard coded output result, decorator pins follow
+			UE_RETURN_ON_ERROR(NodePins.Num() == 2, TEXT("FAnimationAnimNextEditorTest_GraphTraitOperations -> Unexpected number of pins"));
+			UE_RETURN_ON_ERROR(TraitPin->IsDecoratorPin() == true, TEXT("FAnimationAnimNextEditorTest_GraphTraitOperations -> Unexpected pin type"));
+			UE_RETURN_ON_ERROR(TraitPin->GetFName() == TraitInstanceName, TEXT("FAnimationAnimNextEditorTest_GraphTraitOperations -> Unexpected pin name"));
+			UE_RETURN_ON_ERROR(TraitPin->GetCPPTypeObject() == FRigDecorator_AnimNextCppDecorator::StaticStruct(), TEXT("FAnimationAnimNextEditorTest_GraphTraitOperations -> Unexpected pin type"));
+		}
 	}
 
 	Tests::FUtils::CleanupAfterTests();
diff --git a/Engine/Plugins/Experimental/AnimNextAnimGraph/Source/AnimNextAnimGraphTestSuite/Private/AnimNextAnimGraphTraitInterfacesTest.cpp b/Engine/Plugins/Experimental/AnimNextAnimGraph/Source/AnimNextAnimGraphTestSuite/Private/AnimNextAnimGraphTraitInterfacesTest.cpp
--- a/Engine/Plugins/Experimental/AnimNextAnimGraph/Source/AnimNextAnimGraphTestSuite/Private/AnimNextAnimGraphTraitInterfacesTest.cpp
+++ b/Engine/Plugins/Experimental/AnimNextAnimGraph/Source/AnimNextAnimGraphTestSuite/Private/AnimNextAnimGraphTraitInterfacesTest.cpp
@@ -94,7 +94,7 @@ namespace UE::AnimNext
 		GeneratorMacro(IEvaluate) \
 		GeneratorMacro(IUpdate) \
 
-	GENERATE_ANIM_TRAIT_IMPLEMENTATION(FTraitWithNoChildren, TRAIT_INTERFACE_ENUMERATOR, NULL_ANIM_TRAIT_EVENT_ENUMERATOR)
+	GENERATE_ANIM_TRAIT_IMPLEMENTATION(FTraitWithNoChildren, TRAIT_INTERFACE_ENUMERATOR, NULL_ANIM_TRAIT_INTERFACE_ENUMERATOR, NULL_ANIM_TRAIT_EVENT_ENUMERATOR)
 	#undef TRAIT_INTERFACE_ENUMERATOR
 
 	// This trait does not update or evaluate
@@ -132,7 +132,7 @@ namespace UE::AnimNext
 	#define TRAIT_INTERFACE_ENUMERATOR(GeneratorMacro) \
 		GeneratorMacro(IHierarchy) \
 
-	GENERATE_ANIM_TRAIT_IMPLEMENTATION(FTraitWithOneChild, TRAIT_INTERFACE_ENUMERATOR, NULL_ANIM_TRAIT_EVENT_ENUMERATOR)
+	GENERATE_ANIM_TRAIT_IMPLEMENTATION(FTraitWithOneChild, TRAIT_INTERFACE_ENUMERATOR, NULL_ANIM_TRAIT_INTERFACE_ENUMERATOR, NULL_ANIM_TRAIT_EVENT_ENUMERATOR)
 	#undef TRAIT_INTERFACE_ENUMERATOR
 
 	struct FTraitWithChildren : FBaseTrait, IHierarchy, IUpdate, IUpdateTraversal, IEvaluate
@@ -225,12 +225,12 @@ namespace UE::AnimNext
 		GeneratorMacro(IUpdate) \
 		GeneratorMacro(IUpdateTraversal) \
 
-	GENERATE_ANIM_TRAIT_IMPLEMENTATION(FTraitWithChildren, TRAIT_INTERFACE_ENUMERATOR, NULL_ANIM_TRAIT_EVENT_ENUMERATOR)
+	GENERATE_ANIM_TRAIT_IMPLEMENTATION(FTraitWithChildren, TRAIT_INTERFACE_ENUMERATOR, NULL_ANIM_TRAIT_INTERFACE_ENUMERATOR, NULL_ANIM_TRAIT_EVENT_ENUMERATOR)
 	#undef TRAIT_INTERFACE_ENUMERATOR
 
 	struct IScopedTagInterface : IScopedTraitInterface
 	{
-		DECLARE_ANIM_TRAIT_INTERFACE(IScopedTagInterface, 0xb8f4b397)
+		DECLARE_ANIM_TRAIT_INTERFACE(IScopedTagInterface, STG, 0xb8f4b397)
 
 		virtual FName GetTag(const FExecutionContext& Context, const TTraitBinding<IScopedTagInterface>& Binding) const;
 
@@ -315,7 +315,7 @@ namespace UE::AnimNext
 		GeneratorMacro(IScopedTagInterface) \
 		GeneratorMacro(IUpdate) \
 
-	GENERATE_ANIM_TRAIT_IMPLEMENTATION(FScopedTagTrait, TRAIT_INTERFACE_ENUMERATOR, NULL_ANIM_TRAIT_EVENT_ENUMERATOR)
+	GENERATE_ANIM_TRAIT_IMPLEMENTATION(FScopedTagTrait, TRAIT_INTERFACE_ENUMERATOR, NULL_ANIM_TRAIT_INTERFACE_ENUMERATOR, NULL_ANIM_TRAIT_EVENT_ENUMERATOR)
 	#undef TRAIT_INTERFACE_ENUMERATOR
 
 	// Tests if we have a scoped tag
@@ -349,7 +349,7 @@ namespace UE::AnimNext
 	#define TRAIT_INTERFACE_ENUMERATOR(GeneratorMacro) \
 		GeneratorMacro(IUpdate) \
 
-	GENERATE_ANIM_TRAIT_IMPLEMENTATION(FTestScopedTagTrait, TRAIT_INTERFACE_ENUMERATOR, NULL_ANIM_TRAIT_EVENT_ENUMERATOR)
+	GENERATE_ANIM_TRAIT_IMPLEMENTATION(FTestScopedTagTrait, TRAIT_INTERFACE_ENUMERATOR, NULL_ANIM_TRAIT_INTERFACE_ENUMERATOR, NULL_ANIM_TRAIT_EVENT_ENUMERATOR)
 	#undef TRAIT_INTERFACE_ENUMERATOR
 }
 
diff --git a/Engine/Plugins/Experimental/Workspace/Source/WorkspaceEditor/Private/WorkspaceEditor.cpp b/Engine/Plugins/Experimental/Workspace/Source/WorkspaceEditor/Private/WorkspaceEditor.cpp
--- a/Engine/Plugins/Experimental/Workspace/Source/WorkspaceEditor/Private/WorkspaceEditor.cpp
+++ b/Engine/Plugins/Experimental/Workspace/Source/WorkspaceEditor/Private/WorkspaceEditor.cpp
@@ -30,9 +30,11 @@ namespace WorkspaceTabs
 {
 	const FName Details("DetailsTab");
 	const FName WorkspaceView("WorkspaceView");
+	const FName TopDocumentArea("TopDocumentArea");
 	const FName LeftDocumentArea("LeftDocumentArea");
 	const FName MiddleDocumentArea("MiddleDocumentArea");
 	const FName RightDocumentArea("RightDocumentArea");
+	const FName BottomDocumentArea("BottomDocumentArea");
 }
 
 const FName WorkspaceAppIdentifier("WorkspaceEditor");
@@ -66,7 +68,7 @@ void FWorkspaceEditor::CreateWidgets()
 	DocumentManager->RegisterDocumentFactory(RightAssetDocumentSummoner);
 	
 	check(DetailsView.IsValid());
-	WorkspaceEditorModule.ApplyWorkspaceDetailsCustomization(DetailsView);
+	WorkspaceEditorModule.ApplyWorkspaceDetailsCustomization(StaticCastWeakPtr<IWorkspaceEditor>(this->AsWeak()), DetailsView);
 
 	StandaloneDefaultLayout = FTabManager::NewLayout("Standalone_WorkspaceEditor_Layout_v1.1")
 	->AddArea
@@ -74,9 +76,16 @@ void FWorkspaceEditor::CreateWidgets()
 		FTabManager::NewPrimaryArea()
 		->SetOrientation(Orient_Vertical)
 		->Split
+		(
+			FTabManager::NewStack()
+			->SetSizeCoefficient(0.15f)
+			->SetHideTabWell(false)
+			->AddTab(WorkspaceTabs::TopDocumentArea, ETabState::ClosedTab)
+		)
+		->Split
 		(
 			FTabManager::NewSplitter()
-			->SetSizeCoefficient(1.0f)
+			->SetSizeCoefficient(0.6f)
 			->SetOrientation(Orient_Horizontal)
 			->Split
 			(
@@ -102,6 +111,13 @@ void FWorkspaceEditor::CreateWidgets()
 				->AddTab(FBaseAssetToolkit::DetailsTabID, ETabState::OpenedTab)
 			)
 		)
+		->Split
+		(
+			FTabManager::NewStack()
+			->SetSizeCoefficient(0.15f)
+			->SetHideTabWell(false)
+			->AddTab(WorkspaceTabs::BottomDocumentArea, ETabState::ClosedTab)
+		)
 	);
 
 	WorkspaceView = SNew(SWorkspaceView, Workspace, StaticCastSharedRef<UE::Workspace::IWorkspaceEditor>(AsShared()));
@@ -207,7 +223,7 @@ void FWorkspaceEditor::CloseObjects(TConstArrayView<UObject*> InObjects)
 
 void FWorkspaceEditor::SetDetailsObjects(const TArray<UObject*>& InObjects)
 {
-	if(DetailsView.IsValid())
+ 	if(DetailsView.IsValid())
 	{
 		DetailsView->SetObjects(InObjects);
 	}
@@ -221,6 +237,11 @@ void FWorkspaceEditor::RefreshDetails()
 	}
 }
 
+UWorkspaceSchema* FWorkspaceEditor::GetSchema() const
+{
+	return Workspace.Get() != nullptr ? Workspace->GetSchema() : nullptr;
+}
+
 void FWorkspaceEditor::BindCommands()
 {	
 	FWorkspaceAssetEditorCommands::Register();
@@ -268,6 +289,12 @@ void FWorkspaceEditor::RegisterTabSpawners(const TSharedRef<FTabManager>& InTabM
 	.SetTooltipText(LOCTEXT("WorkspaceTabToolTip", "Shows the workspace outliner tab."));
 
 	DocumentManager->SetTabManager(InTabManager);
+
+	Workspace::IWorkspaceEditorModule& WorkspaceEditorModule = FModuleManager::Get().LoadModuleChecked<Workspace::IWorkspaceEditorModule>("WorkspaceEditor");
+	if (WorkspaceEditorModule.OnRegisterTabsForEditor().IsBound())
+	{
+		WorkspaceEditorModule.OnRegisterTabsForEditor().Broadcast(TabFactories, InTabManager, StaticCastSharedPtr<IWorkspaceEditor>(this->AsShared().ToSharedPtr()));
+	}
 }
 
 void FWorkspaceEditor::UnregisterTabSpawners(const TSharedRef<FTabManager>& InTabManager)
diff --git a/Engine/Plugins/Experimental/Workspace/Source/WorkspaceEditor/Private/WorkspaceEditor.h b/Engine/Plugins/Experimental/Workspace/Source/WorkspaceEditor/Private/WorkspaceEditor.h
--- a/Engine/Plugins/Experimental/Workspace/Source/WorkspaceEditor/Private/WorkspaceEditor.h
+++ b/Engine/Plugins/Experimental/Workspace/Source/WorkspaceEditor/Private/WorkspaceEditor.h
@@ -88,6 +88,7 @@ private:
 	virtual void CloseObjects(TConstArrayView<UObject*> InObjects) override;
 	virtual void SetDetailsObjects(const TArray<UObject*>& InObjects) override;
 	virtual void RefreshDetails() override;
+	virtual UWorkspaceSchema* GetSchema() const override;
 
 	void BindCommands();
 
@@ -118,6 +119,9 @@ private:
 
 	TSharedPtr<SWorkspaceView> WorkspaceView;
 
+	/** Tabs to be registered into the Workspace */
+	FWorkflowAllowedTabSet TabFactories;
+
 	bool bSavingWorkspaceOnly = false;
 	bool bClosingDown = false;
 };
diff --git a/Engine/Plugins/Experimental/Workspace/Source/WorkspaceEditor/Private/WorkspaceEditorModule.cpp b/Engine/Plugins/Experimental/Workspace/Source/WorkspaceEditor/Private/WorkspaceEditorModule.cpp
--- a/Engine/Plugins/Experimental/Workspace/Source/WorkspaceEditor/Private/WorkspaceEditorModule.cpp
+++ b/Engine/Plugins/Experimental/Workspace/Source/WorkspaceEditor/Private/WorkspaceEditorModule.cpp
@@ -256,11 +256,11 @@ IWorkspaceEditorModule::FOnRegisterDetailCustomizations& FWorkspaceEditorModule:
 	return OnRegisterDetailCustomizations;
 }
 
-void FWorkspaceEditorModule::ApplyWorkspaceDetailsCustomization(TSharedPtr<IDetailsView>& DetailsView) const
+void FWorkspaceEditorModule::ApplyWorkspaceDetailsCustomization(const TWeakPtr<IWorkspaceEditor>& InWorkspaceEditor, TSharedPtr<IDetailsView>& DetailsView) const
 {
 	if (OnRegisterDetailCustomizations.IsBound())
 	{
-		OnRegisterDetailCustomizations.Broadcast(DetailsView);
+		OnRegisterDetailCustomizations.Broadcast(InWorkspaceEditor, DetailsView);
 	}
 }
 
diff --git a/Engine/Plugins/Experimental/Workspace/Source/WorkspaceEditor/Private/WorkspaceEditorModule.h b/Engine/Plugins/Experimental/Workspace/Source/WorkspaceEditor/Private/WorkspaceEditorModule.h
--- a/Engine/Plugins/Experimental/Workspace/Source/WorkspaceEditor/Private/WorkspaceEditorModule.h
+++ b/Engine/Plugins/Experimental/Workspace/Source/WorkspaceEditor/Private/WorkspaceEditorModule.h
@@ -5,6 +5,7 @@
 #include "IWorkspaceEditorModule.h"
 #include "UObject/TopLevelAssetPath.h"
 #include "IWorkspaceOutlinerItemDetails.h"
+#include "WorkflowOrientedApp/WorkflowTabManager.h"
 
 struct FWorkspaceAssetRegistryExports;
 class SWorkspaceTabWrapper;
@@ -30,6 +31,7 @@ private:
 	virtual FOnRegisterDetailCustomizations& OnRegisterWorkspaceDetailsCustomization() override;
 	virtual void RegisterWorkspaceItemDetails(const FOutlinerItemDetailsId& InItemDetailsId, TSharedPtr<IWorkspaceOutlinerItemDetails> InItemDetails) override;	
 	virtual void UnregisterWorkspaceItemDetails(const FOutlinerItemDetailsId& InItemDetails) override;
+	virtual FOnRegisterTabs& OnRegisterTabsForEditor() override { return RegisterTabsForEditor; }
 
 	// Find an existing registered object document type
 	const FObjectDocumentArgs* FindObjectDocumentType(const FTopLevelAssetPath& InClassPath) const;
@@ -41,7 +43,7 @@ private:
 	static bool GetExportedAssetsForWorkspace(const FAssetData& InWorkspaceAsset, FWorkspaceAssetRegistryExports& OutExports);
 
 	// Applies any previously registered details-view customizations
-	void ApplyWorkspaceDetailsCustomization(TSharedPtr<IDetailsView>& DetailsView) const;
+	void ApplyWorkspaceDetailsCustomization(const TWeakPtr<IWorkspaceEditor>& Editor, TSharedPtr<IDetailsView>& DetailsView) const;
 
 	TMap<FTopLevelAssetPath, FObjectDocumentArgs> ObjectDocumentArgs;
 
@@ -53,6 +55,9 @@ private:
 	/** Event called to allow external clients to register details customizations */
 	FOnRegisterDetailCustomizations OnRegisterDetailCustomizations;
 
+	/** Event called to allow external clients to register additional tabs for the specified editor */
+	FOnRegisterTabs RegisterTabsForEditor;
+
 	friend struct FAssetDocumentSummoner;
 	friend struct FWorkspaceOutlinerTreeItem;
 	friend class SWorkspaceOutlinerTreelabel;
diff --git a/Engine/Plugins/Experimental/Workspace/Source/WorkspaceEditor/Public/IWorkspaceEditor.h b/Engine/Plugins/Experimental/Workspace/Source/WorkspaceEditor/Public/IWorkspaceEditor.h
--- a/Engine/Plugins/Experimental/Workspace/Source/WorkspaceEditor/Public/IWorkspaceEditor.h
+++ b/Engine/Plugins/Experimental/Workspace/Source/WorkspaceEditor/Public/IWorkspaceEditor.h
@@ -1,10 +1,11 @@
-﻿// Copyright Epic Games, Inc. All Rights Reserved.
+// Copyright Epic Games, Inc. All Rights Reserved.
 
 #pragma once
 
 #include "Tools/BaseAssetToolkit.h"
 
 class UWorkspaceAssetEditor;
+class UWorkspaceSchema;
 
 namespace UE::Workspace
 {
@@ -28,6 +29,9 @@ public:
 
 	// Refresh the workspace editor details panel
 	virtual void RefreshDetails() = 0;
+
+	// Exposes the editor WorkspaceSchema
+	virtual UWorkspaceSchema* GetSchema() const = 0;
 };
 
 }
\ No newline at end of file
diff --git a/Engine/Plugins/Experimental/Workspace/Source/WorkspaceEditor/Public/IWorkspaceEditorModule.h b/Engine/Plugins/Experimental/Workspace/Source/WorkspaceEditor/Public/IWorkspaceEditorModule.h
--- a/Engine/Plugins/Experimental/Workspace/Source/WorkspaceEditor/Public/IWorkspaceEditorModule.h
+++ b/Engine/Plugins/Experimental/Workspace/Source/WorkspaceEditor/Public/IWorkspaceEditorModule.h
@@ -13,10 +13,14 @@
 class SWidget;
 class UEdGraph;
 class IDetailsView;
+class FTabManager;
+class FWorkflowAllowedTabSet;
+class IDetailCustomization;
 
 struct FSlateBrush;
 struct FTopLevelAssetPath;
 struct FWorkspaceDocumentState;
+struct FRegisterCustomClassLayoutParams;
 
 namespace UE::Workspace
 {
@@ -116,6 +120,9 @@ using FOnPasteNodes = TDelegate<void(const FWorkspaceEditorContext&, const FVect
 
 using FOnDuplicateSelectedNodes = TDelegate<void(const FWorkspaceEditorContext&, const FVector2D&, const FGraphPanelSelectionSet&)>;
 
+using FOnGetWorkspaceDetailCustomizationInstance = TDelegate<TSharedRef<IDetailCustomization>(TWeakPtr<IWorkspaceEditor>)>;
+
+
 // Arguments used to make document widgets for graphs
 struct FGraphDocumentWidgetArgs
 {
@@ -175,11 +182,15 @@ public:
 	virtual FObjectDocumentArgs CreateGraphDocumentArgs(const FGraphDocumentWidgetArgs& InArgs) = 0;
 
 	// Event to allow registering details customizations
-	DECLARE_EVENT_OneParam(IWorkspaceEditorModule, FOnRegisterDetailCustomizations, TSharedPtr<IDetailsView>&);
+	DECLARE_EVENT_TwoParams(IWorkspaceEditorModule, FOnRegisterDetailCustomizations, const TWeakPtr<IWorkspaceEditor>&, TSharedPtr<IDetailsView>&);
 	virtual FOnRegisterDetailCustomizations& OnRegisterWorkspaceDetailsCustomization() = 0;
 
 	virtual void RegisterWorkspaceItemDetails(const FOutlinerItemDetailsId& InDetailsId, TSharedPtr<IWorkspaceOutlinerItemDetails> InDetails) = 0;
 	virtual void UnregisterWorkspaceItemDetails(const FOutlinerItemDetailsId& InDetailsId) = 0;
+
+	// Event to allow registering tabs to other elements
+	DECLARE_EVENT_ThreeParams(IWorkspaceEditorModule, FOnRegisterTabs, FWorkflowAllowedTabSet& TabFactories, const TSharedRef<FTabManager>&, TSharedPtr<IWorkspaceEditor>);
+	virtual FOnRegisterTabs& OnRegisterTabsForEditor() = 0;
 };
 
 }
diff --git a/Engine/Plugins/Runtime/RigVM/Source/RigVMDeveloper/Private/RigVMModel/RigVMController.cpp b/Engine/Plugins/Runtime/RigVM/Source/RigVMDeveloper/Private/RigVMModel/RigVMController.cpp
--- a/Engine/Plugins/Runtime/RigVM/Source/RigVMDeveloper/Private/RigVMModel/RigVMController.cpp
+++ b/Engine/Plugins/Runtime/RigVM/Source/RigVMDeveloper/Private/RigVMModel/RigVMController.cpp
@@ -14288,7 +14288,7 @@ FName URigVMController::AddDecorator(URigVMNode* InNode, UScriptStruct* InDecora
 	{
 		Action.SetTitle(FString::Printf(TEXT("Add Decorator")));
 		GetActionStack()->BeginAction(Action);
-		GetActionStack()->AddAction(FRigVMAddDecoratorAction(this, InNode, ValidDecoratorName, InDecoratorScriptStruct, InDefaultValue, InNode->GetPins().Num()));
+		GetActionStack()->AddAction(FRigVMAddDecoratorAction(this, InNode, ValidDecoratorName, InDecoratorScriptStruct, InDefaultValue, InPinIndex == INDEX_NONE ? InNode->GetPins().Num() : InPinIndex));
 	}
 
 	InNode->DecoratorRootPinNames.Add(ValidDecoratorName.ToString());
