Remove const qualifier from execution context

diff --git a/Engine/Plugins/Experimental/AnimNextAnimGraph/Source/AnimNextAnimGraph/Private/TraitInterfaces/IDiscreteBlend.cpp b/Engine/Plugins/Experimental/AnimNextAnimGraph/Source/AnimNextAnimGraph/Private/TraitInterfaces/IDiscreteBlend.cpp
--- a/Engine/Plugins/Experimental/AnimNextAnimGraph/Source/AnimNextAnimGraph/Private/TraitInterfaces/IDiscreteBlend.cpp
+++ b/Engine/Plugins/Experimental/AnimNextAnimGraph/Source/AnimNextAnimGraph/Private/TraitInterfaces/IDiscreteBlend.cpp
@@ -6,7 +6,7 @@
 
 namespace UE::AnimNext
 {
-	float IDiscreteBlend::GetBlendWeight(const FExecutionContext& Context, const TTraitBinding<IDiscreteBlend>& Binding, int32 ChildIndex) const
+	float IDiscreteBlend::GetBlendWeight(FExecutionContext& Context, const TTraitBinding<IDiscreteBlend>& Binding, int32 ChildIndex) const
 	{
 		TTraitBinding<IDiscreteBlend> SuperBinding;
 		if (Binding.GetStackInterfaceSuper(SuperBinding))
@@ -17,7 +17,7 @@ namespace UE::AnimNext
 		return -1.0f;
 	}
 
-	const FAlphaBlend* IDiscreteBlend::GetBlendState(const FExecutionContext& Context, const TTraitBinding<IDiscreteBlend>& Binding, int32 ChildIndex) const
+	const FAlphaBlend* IDiscreteBlend::GetBlendState(FExecutionContext& Context, const TTraitBinding<IDiscreteBlend>& Binding, int32 ChildIndex) const
 	{
 		TTraitBinding<IDiscreteBlend> SuperBinding;
 		if (Binding.GetStackInterfaceSuper(SuperBinding))
@@ -28,7 +28,7 @@ namespace UE::AnimNext
 		return nullptr;
 	}
 
-	int32 IDiscreteBlend::GetBlendDestinationChildIndex(const FExecutionContext& Context, const TTraitBinding<IDiscreteBlend>& Binding) const
+	int32 IDiscreteBlend::GetBlendDestinationChildIndex(FExecutionContext& Context, const TTraitBinding<IDiscreteBlend>& Binding) const
 	{
 		TTraitBinding<IDiscreteBlend> SuperBinding;
 		if (Binding.GetStackInterfaceSuper(SuperBinding))
diff --git a/Engine/Plugins/Experimental/AnimNextAnimGraph/Source/AnimNextAnimGraph/Private/TraitInterfaces/IGroupSynchronization.cpp b/Engine/Plugins/Experimental/AnimNextAnimGraph/Source/AnimNextAnimGraph/Private/TraitInterfaces/IGroupSynchronization.cpp
--- a/Engine/Plugins/Experimental/AnimNextAnimGraph/Source/AnimNextAnimGraph/Private/TraitInterfaces/IGroupSynchronization.cpp
+++ b/Engine/Plugins/Experimental/AnimNextAnimGraph/Source/AnimNextAnimGraph/Private/TraitInterfaces/IGroupSynchronization.cpp
@@ -6,7 +6,7 @@
 
 namespace UE::AnimNext
 {
-	FName IGroupSynchronization::GetGroupName(const FExecutionContext& Context, const TTraitBinding<IGroupSynchronization>& Binding) const
+	FName IGroupSynchronization::GetGroupName(FExecutionContext& Context, const TTraitBinding<IGroupSynchronization>& Binding) const
 	{
 		TTraitBinding<IGroupSynchronization> SuperBinding;
 		if (Binding.GetStackInterfaceSuper(SuperBinding))
@@ -17,7 +17,7 @@ namespace UE::AnimNext
 		return NAME_None;
 	}
 
-	EAnimGroupRole::Type IGroupSynchronization::GetGroupRole(const FExecutionContext& Context, const TTraitBinding<IGroupSynchronization>& Binding) const
+	EAnimGroupRole::Type IGroupSynchronization::GetGroupRole(FExecutionContext& Context, const TTraitBinding<IGroupSynchronization>& Binding) const
 	{
 		TTraitBinding<IGroupSynchronization> SuperBinding;
 		if (Binding.GetStackInterfaceSuper(SuperBinding))
@@ -28,7 +28,7 @@ namespace UE::AnimNext
 		return EAnimGroupRole::CanBeLeader;
 	}
 
-	float IGroupSynchronization::AdvanceBy(const FExecutionContext& Context, const TTraitBinding<IGroupSynchronization>& Binding, float DeltaTime) const
+	float IGroupSynchronization::AdvanceBy(FExecutionContext& Context, const TTraitBinding<IGroupSynchronization>& Binding, float DeltaTime) const
 	{
 		TTraitBinding<IGroupSynchronization> SuperBinding;
 		if (Binding.GetStackInterfaceSuper(SuperBinding))
@@ -39,7 +39,7 @@ namespace UE::AnimNext
 		return 0.0f;
 	}
 
-	void IGroupSynchronization::AdvanceToRatio(const FExecutionContext& Context, const TTraitBinding<IGroupSynchronization>& Binding, float ProgressRatio) const
+	void IGroupSynchronization::AdvanceToRatio(FExecutionContext& Context, const TTraitBinding<IGroupSynchronization>& Binding, float ProgressRatio) const
 	{
 		TTraitBinding<IGroupSynchronization> SuperBinding;
 		if (Binding.GetStackInterfaceSuper(SuperBinding))
diff --git a/Engine/Plugins/Experimental/AnimNextAnimGraph/Source/AnimNextAnimGraph/Private/TraitInterfaces/ISmoothBlend.cpp b/Engine/Plugins/Experimental/AnimNextAnimGraph/Source/AnimNextAnimGraph/Private/TraitInterfaces/ISmoothBlend.cpp
--- a/Engine/Plugins/Experimental/AnimNextAnimGraph/Source/AnimNextAnimGraph/Private/TraitInterfaces/ISmoothBlend.cpp
+++ b/Engine/Plugins/Experimental/AnimNextAnimGraph/Source/AnimNextAnimGraph/Private/TraitInterfaces/ISmoothBlend.cpp
@@ -6,7 +6,7 @@
 
 namespace UE::AnimNext
 {
-	float ISmoothBlend::GetBlendTime(const FExecutionContext& Context, const TTraitBinding<ISmoothBlend>& Binding, int32 ChildIndex) const
+	float ISmoothBlend::GetBlendTime(FExecutionContext& Context, const TTraitBinding<ISmoothBlend>& Binding, int32 ChildIndex) const
 	{
 		TTraitBinding<ISmoothBlend> SuperBinding;
 		if (Binding.GetStackInterfaceSuper(SuperBinding))
diff --git a/Engine/Plugins/Experimental/AnimNextAnimGraph/Source/AnimNextAnimGraph/Private/TraitInterfaces/ITimeline.cpp b/Engine/Plugins/Experimental/AnimNextAnimGraph/Source/AnimNextAnimGraph/Private/TraitInterfaces/ITimeline.cpp
--- a/Engine/Plugins/Experimental/AnimNextAnimGraph/Source/AnimNextAnimGraph/Private/TraitInterfaces/ITimeline.cpp
+++ b/Engine/Plugins/Experimental/AnimNextAnimGraph/Source/AnimNextAnimGraph/Private/TraitInterfaces/ITimeline.cpp
@@ -6,7 +6,7 @@
 
 namespace UE::AnimNext
 {
-	float ITimeline::GetPlayRate(const FExecutionContext& Context, const TTraitBinding<ITimeline>& Binding) const
+	float ITimeline::GetPlayRate(FExecutionContext& Context, const TTraitBinding<ITimeline>& Binding) const
 	{
 		TTraitBinding<ITimeline> SuperBinding;
 		if (Binding.GetStackInterfaceSuper(SuperBinding))
@@ -17,7 +17,7 @@ namespace UE::AnimNext
 		return 1.0f;
 	}
 
-	float ITimeline::AdvanceBy(const FExecutionContext& Context, const TTraitBinding<ITimeline>& Binding, float DeltaTime) const
+	float ITimeline::AdvanceBy(FExecutionContext& Context, const TTraitBinding<ITimeline>& Binding, float DeltaTime) const
 	{
 		TTraitBinding<ITimeline> SuperBinding;
 		if (Binding.GetStackInterfaceSuper(SuperBinding))
@@ -28,7 +28,7 @@ namespace UE::AnimNext
 		return 0.0f;
 	}
 
-	void ITimeline::AdvanceToRatio(const FExecutionContext& Context, const TTraitBinding<ITimeline>& Binding, float ProgressRatio) const
+	void ITimeline::AdvanceToRatio(FExecutionContext& Context, const TTraitBinding<ITimeline>& Binding, float ProgressRatio) const
 	{
 		TTraitBinding<ITimeline> SuperBinding;
 		if (Binding.GetStackInterfaceSuper(SuperBinding))
diff --git a/Engine/Plugins/Experimental/AnimNextAnimGraph/Source/AnimNextAnimGraph/Private/Traits/BlendByBool.cpp b/Engine/Plugins/Experimental/AnimNextAnimGraph/Source/AnimNextAnimGraph/Private/Traits/BlendByBool.cpp
--- a/Engine/Plugins/Experimental/AnimNextAnimGraph/Source/AnimNextAnimGraph/Private/Traits/BlendByBool.cpp
+++ b/Engine/Plugins/Experimental/AnimNextAnimGraph/Source/AnimNextAnimGraph/Private/Traits/BlendByBool.cpp
@@ -100,7 +100,7 @@ namespace UE::AnimNext
 		Children.Add(InstanceData->FalseChild);
 	}
 
-	float FBlendByBoolTrait::GetBlendWeight(const FExecutionContext& Context, const TTraitBinding<IDiscreteBlend>& Binding, int32 ChildIndex) const
+	float FBlendByBoolTrait::GetBlendWeight(FExecutionContext& Context, const TTraitBinding<IDiscreteBlend>& Binding, int32 ChildIndex) const
 	{
 		TTraitBinding<IDiscreteBlend> DiscreteBlendTrait;
 		Binding.GetStackInterface(DiscreteBlendTrait);
@@ -122,7 +122,7 @@ namespace UE::AnimNext
 		}
 	}
 
-	int32 FBlendByBoolTrait::GetBlendDestinationChildIndex(const FExecutionContext& Context, const TTraitBinding<IDiscreteBlend>& Binding) const
+	int32 FBlendByBoolTrait::GetBlendDestinationChildIndex(FExecutionContext& Context, const TTraitBinding<IDiscreteBlend>& Binding) const
 	{
 		const FSharedData* SharedData = Binding.GetSharedData<FSharedData>();
 
diff --git a/Engine/Plugins/Experimental/AnimNextAnimGraph/Source/AnimNextAnimGraph/Private/Traits/BlendByBool.h b/Engine/Plugins/Experimental/AnimNextAnimGraph/Source/AnimNextAnimGraph/Private/Traits/BlendByBool.h
--- a/Engine/Plugins/Experimental/AnimNextAnimGraph/Source/AnimNextAnimGraph/Private/Traits/BlendByBool.h
+++ b/Engine/Plugins/Experimental/AnimNextAnimGraph/Source/AnimNextAnimGraph/Private/Traits/BlendByBool.h
@@ -70,8 +70,8 @@ namespace UE::AnimNext
 		virtual void GetChildren(const FExecutionContext& Context, const TTraitBinding<IHierarchy>& Binding, FChildrenArray& Children) const override;
 
 		// IDiscreteBlend impl
-		virtual float GetBlendWeight(const FExecutionContext& Context, const TTraitBinding<IDiscreteBlend>& Binding, int32 ChildIndex) const override;
-		virtual int32 GetBlendDestinationChildIndex(const FExecutionContext& Context, const TTraitBinding<IDiscreteBlend>& Binding) const override;
+		virtual float GetBlendWeight(FExecutionContext& Context, const TTraitBinding<IDiscreteBlend>& Binding, int32 ChildIndex) const override;
+		virtual int32 GetBlendDestinationChildIndex(FExecutionContext& Context, const TTraitBinding<IDiscreteBlend>& Binding) const override;
 		virtual void OnBlendTransition(FExecutionContext& Context, const TTraitBinding<IDiscreteBlend>& Binding, int32 OldChildIndex, int32 NewChildIndex) const override;
 		virtual void OnBlendInitiated(FExecutionContext& Context, const TTraitBinding<IDiscreteBlend>& Binding, int32 ChildIndex) const override;
 		virtual void OnBlendTerminated(FExecutionContext& Context, const TTraitBinding<IDiscreteBlend>& Binding, int32 ChildIndex) const override;
diff --git a/Engine/Plugins/Experimental/AnimNextAnimGraph/Source/AnimNextAnimGraph/Private/Traits/BlendInertializer.cpp b/Engine/Plugins/Experimental/AnimNextAnimGraph/Source/AnimNextAnimGraph/Private/Traits/BlendInertializer.cpp
--- a/Engine/Plugins/Experimental/AnimNextAnimGraph/Source/AnimNextAnimGraph/Private/Traits/BlendInertializer.cpp
+++ b/Engine/Plugins/Experimental/AnimNextAnimGraph/Source/AnimNextAnimGraph/Private/Traits/BlendInertializer.cpp
@@ -37,7 +37,7 @@ namespace UE::AnimNext
 		Context.RaiseOutputTraitEvent(Event);
 	}
 
-	float FBlendInertializerTrait::GetBlendTime(const FExecutionContext& Context, const TTraitBinding<ISmoothBlend>& Binding, int32 ChildIndex) const
+	float FBlendInertializerTrait::GetBlendTime(FExecutionContext& Context, const TTraitBinding<ISmoothBlend>& Binding, int32 ChildIndex) const
 	{
 		// We hijack the blend time and always transition instantaneously
 		return 0.0f;
diff --git a/Engine/Plugins/Experimental/AnimNextAnimGraph/Source/AnimNextAnimGraph/Private/Traits/BlendInertializer.h b/Engine/Plugins/Experimental/AnimNextAnimGraph/Source/AnimNextAnimGraph/Private/Traits/BlendInertializer.h
--- a/Engine/Plugins/Experimental/AnimNextAnimGraph/Source/AnimNextAnimGraph/Private/Traits/BlendInertializer.h
+++ b/Engine/Plugins/Experimental/AnimNextAnimGraph/Source/AnimNextAnimGraph/Private/Traits/BlendInertializer.h
@@ -34,12 +34,11 @@ namespace UE::AnimNext
 		DECLARE_ANIM_TRAIT(FBlendInertializerTrait, 0x7ea0bdee, FAdditiveTrait)
 
 		using FSharedData = FAnimNextBlendInertializerTraitSharedData;
-		using FInstanceData = FTrait::FInstanceData;
 
 		// IDiscreteBlend impl
 		virtual void OnBlendTransition(FExecutionContext& Context, const TTraitBinding<IDiscreteBlend>& Binding, int32 OldChildIndex, int32 NewChildIndex) const override;
 
 		// ISmoothBlend impl
-		virtual float GetBlendTime(const FExecutionContext& Context, const TTraitBinding<ISmoothBlend>& Binding, int32 ChildIndex) const override;
+		virtual float GetBlendTime(FExecutionContext& Context, const TTraitBinding<ISmoothBlend>& Binding, int32 ChildIndex) const override;
 	};
 }
diff --git a/Engine/Plugins/Experimental/AnimNextAnimGraph/Source/AnimNextAnimGraph/Private/Traits/BlendSmoother.cpp b/Engine/Plugins/Experimental/AnimNextAnimGraph/Source/AnimNextAnimGraph/Private/Traits/BlendSmoother.cpp
--- a/Engine/Plugins/Experimental/AnimNextAnimGraph/Source/AnimNextAnimGraph/Private/Traits/BlendSmoother.cpp
+++ b/Engine/Plugins/Experimental/AnimNextAnimGraph/Source/AnimNextAnimGraph/Private/Traits/BlendSmoother.cpp
@@ -158,13 +158,13 @@ namespace UE::AnimNext
 		}
 	}
 
-	float FBlendSmootherTrait::GetBlendWeight(const FExecutionContext& Context, const TTraitBinding<IDiscreteBlend>& Binding, int32 ChildIndex) const
+	float FBlendSmootherTrait::GetBlendWeight(FExecutionContext& Context, const TTraitBinding<IDiscreteBlend>& Binding, int32 ChildIndex) const
 	{
 		const FInstanceData* InstanceData = Binding.GetInstanceData<FInstanceData>();
 		return InstanceData->PerChildBlendData.IsValidIndex(ChildIndex) ? InstanceData->PerChildBlendData[ChildIndex].Weight : -1.0f;
 	}
 
-	const FAlphaBlend* FBlendSmootherTrait::GetBlendState(const FExecutionContext& Context, const TTraitBinding<IDiscreteBlend>& Binding, int32 ChildIndex) const
+	const FAlphaBlend* FBlendSmootherTrait::GetBlendState(FExecutionContext& Context, const TTraitBinding<IDiscreteBlend>& Binding, int32 ChildIndex) const
 	{
 		const FInstanceData* InstanceData = Binding.GetInstanceData<FInstanceData>();
 		return InstanceData->PerChildBlendData.IsValidIndex(ChildIndex) ? &InstanceData->PerChildBlendData[ChildIndex].Blend : nullptr;
@@ -231,7 +231,7 @@ namespace UE::AnimNext
 		DiscreteBlendTrait.OnBlendInitiated(Context, NewChildIndex);
 	}
 
-	float FBlendSmootherTrait::GetBlendTime(const FExecutionContext& Context, const TTraitBinding<ISmoothBlend>& Binding, int32 ChildIndex) const
+	float FBlendSmootherTrait::GetBlendTime(FExecutionContext& Context, const TTraitBinding<ISmoothBlend>& Binding, int32 ChildIndex) const
 	{
 		const FSharedData* SharedData = Binding.GetSharedData<FSharedData>();
 
@@ -252,7 +252,7 @@ namespace UE::AnimNext
 		}
 	}
 
-	void FBlendSmootherTrait::InitializeInstanceData(const FExecutionContext& Context, const FTraitBinding& Binding, const FSharedData* SharedData, FInstanceData* InstanceData)
+	void FBlendSmootherTrait::InitializeInstanceData(FExecutionContext& Context, const FTraitBinding& Binding, const FSharedData* SharedData, FInstanceData* InstanceData)
 	{
 		check(InstanceData->PerChildBlendData.IsEmpty());
 
diff --git a/Engine/Plugins/Experimental/AnimNextAnimGraph/Source/AnimNextAnimGraph/Private/Traits/BlendSmoother.h b/Engine/Plugins/Experimental/AnimNextAnimGraph/Source/AnimNextAnimGraph/Private/Traits/BlendSmoother.h
--- a/Engine/Plugins/Experimental/AnimNextAnimGraph/Source/AnimNextAnimGraph/Private/Traits/BlendSmoother.h
+++ b/Engine/Plugins/Experimental/AnimNextAnimGraph/Source/AnimNextAnimGraph/Private/Traits/BlendSmoother.h
@@ -72,14 +72,14 @@ namespace UE::AnimNext
 		virtual void PreUpdate(FUpdateTraversalContext& Context, const TTraitBinding<IUpdate>& Binding, const FTraitUpdateState& TraitState) const override;
 
 		// IDiscreteBlend impl
-		virtual float GetBlendWeight(const FExecutionContext& Context, const TTraitBinding<IDiscreteBlend>& Binding, int32 ChildIndex) const override;
-		virtual const FAlphaBlend* GetBlendState(const FExecutionContext& Context, const TTraitBinding<IDiscreteBlend>& Binding, int32 ChildIndex) const override;
+		virtual float GetBlendWeight(FExecutionContext& Context, const TTraitBinding<IDiscreteBlend>& Binding, int32 ChildIndex) const override;
+		virtual const FAlphaBlend* GetBlendState(FExecutionContext& Context, const TTraitBinding<IDiscreteBlend>& Binding, int32 ChildIndex) const override;
 		virtual void OnBlendTransition(FExecutionContext& Context, const TTraitBinding<IDiscreteBlend>& Binding, int32 OldChildIndex, int32 NewChildIndex) const override;
 
 		// ISmoothBlend impl
-		virtual float GetBlendTime(const FExecutionContext& Context, const TTraitBinding<ISmoothBlend>& Binding, int32 ChildIndex) const override;
+		virtual float GetBlendTime(FExecutionContext& Context, const TTraitBinding<ISmoothBlend>& Binding, int32 ChildIndex) const override;
 
 		// Internal impl
-		static void InitializeInstanceData(const FExecutionContext& Context, const FTraitBinding& Binding, const FSharedData* SharedData, FInstanceData* InstanceData);
+		static void InitializeInstanceData(FExecutionContext& Context, const FTraitBinding& Binding, const FSharedData* SharedData, FInstanceData* InstanceData);
 	};
 }
diff --git a/Engine/Plugins/Experimental/AnimNextAnimGraph/Source/AnimNextAnimGraph/Private/Traits/SequencePlayer.cpp b/Engine/Plugins/Experimental/AnimNextAnimGraph/Source/AnimNextAnimGraph/Private/Traits/SequencePlayer.cpp
--- a/Engine/Plugins/Experimental/AnimNextAnimGraph/Source/AnimNextAnimGraph/Private/Traits/SequencePlayer.cpp
+++ b/Engine/Plugins/Experimental/AnimNextAnimGraph/Source/AnimNextAnimGraph/Private/Traits/SequencePlayer.cpp
@@ -47,13 +47,13 @@ namespace UE::AnimNext
 		Context.AppendTask(Task);
 	}
 
-	float FSequencePlayerTrait::GetPlayRate(const FExecutionContext& Context, const TTraitBinding<ITimeline>& Binding) const
+	float FSequencePlayerTrait::GetPlayRate(FExecutionContext& Context, const TTraitBinding<ITimeline>& Binding) const
 	{
 		const FSharedData* SharedData = Binding.GetSharedData<FSharedData>();
 		return SharedData->GetPlayRate(Binding);
 	}
 
-	float FSequencePlayerTrait::AdvanceBy(const FExecutionContext& Context, const TTraitBinding<ITimeline>& Binding, float DeltaTime) const
+	float FSequencePlayerTrait::AdvanceBy(FExecutionContext& Context, const TTraitBinding<ITimeline>& Binding, float DeltaTime) const
 	{
 		FInstanceData* InstanceData = Binding.GetInstanceData<FInstanceData>();
 
@@ -76,7 +76,7 @@ namespace UE::AnimNext
 		return 0.0f;
 	}
 
-	void FSequencePlayerTrait::AdvanceToRatio(const FExecutionContext& Context, const TTraitBinding<ITimeline>& Binding, float ProgressRatio) const
+	void FSequencePlayerTrait::AdvanceToRatio(FExecutionContext& Context, const TTraitBinding<ITimeline>& Binding, float ProgressRatio) const
 	{
 		FInstanceData* InstanceData = Binding.GetInstanceData<FInstanceData>();
 
diff --git a/Engine/Plugins/Experimental/AnimNextAnimGraph/Source/AnimNextAnimGraph/Private/Traits/SequencePlayer.h b/Engine/Plugins/Experimental/AnimNextAnimGraph/Source/AnimNextAnimGraph/Private/Traits/SequencePlayer.h
--- a/Engine/Plugins/Experimental/AnimNextAnimGraph/Source/AnimNextAnimGraph/Private/Traits/SequencePlayer.h
+++ b/Engine/Plugins/Experimental/AnimNextAnimGraph/Source/AnimNextAnimGraph/Private/Traits/SequencePlayer.h
@@ -72,9 +72,9 @@ namespace UE::AnimNext
 		virtual void PreEvaluate(FEvaluateTraversalContext& Context, const TTraitBinding<IEvaluate>& Binding) const override;
 
 		// ITimeline impl
-		virtual float GetPlayRate(const FExecutionContext& Context, const TTraitBinding<ITimeline>& Binding) const override;
-		virtual float AdvanceBy(const FExecutionContext& Context, const TTraitBinding<ITimeline>& Binding, float DeltaTime) const override;
-		virtual void AdvanceToRatio(const FExecutionContext& Context, const TTraitBinding<ITimeline>& Binding, float ProgressRatio) const override;
+		virtual float GetPlayRate(FExecutionContext& Context, const TTraitBinding<ITimeline>& Binding) const override;
+		virtual float AdvanceBy(FExecutionContext& Context, const TTraitBinding<ITimeline>& Binding, float DeltaTime) const override;
+		virtual void AdvanceToRatio(FExecutionContext& Context, const TTraitBinding<ITimeline>& Binding, float ProgressRatio) const override;
 
 		// IUpdate impl
 		virtual void OnBecomeRelevant(FUpdateTraversalContext& Context, const TTraitBinding<IUpdate>& Binding, const FTraitUpdateState& TraitState) const override;
diff --git a/Engine/Plugins/Experimental/AnimNextAnimGraph/Source/AnimNextAnimGraph/Private/Traits/SubGraphHost.cpp b/Engine/Plugins/Experimental/AnimNextAnimGraph/Source/AnimNextAnimGraph/Private/Traits/SubGraphHost.cpp
--- a/Engine/Plugins/Experimental/AnimNextAnimGraph/Source/AnimNextAnimGraph/Private/Traits/SubGraphHost.cpp
+++ b/Engine/Plugins/Experimental/AnimNextAnimGraph/Source/AnimNextAnimGraph/Private/Traits/SubGraphHost.cpp
@@ -161,7 +161,7 @@ namespace UE::AnimNext
 		}
 	}
 
-	float FSubGraphHostTrait::GetBlendWeight(const FExecutionContext& Context, const TTraitBinding<IDiscreteBlend>& Binding, int32 ChildIndex) const
+	float FSubGraphHostTrait::GetBlendWeight(FExecutionContext& Context, const TTraitBinding<IDiscreteBlend>& Binding, int32 ChildIndex) const
 	{
 		const FInstanceData* InstanceData = Binding.GetInstanceData<FInstanceData>();
 
@@ -180,7 +180,7 @@ namespace UE::AnimNext
 		}
 	}
 
-	int32 FSubGraphHostTrait::GetBlendDestinationChildIndex(const FExecutionContext& Context, const TTraitBinding<IDiscreteBlend>& Binding) const
+	int32 FSubGraphHostTrait::GetBlendDestinationChildIndex(FExecutionContext& Context, const TTraitBinding<IDiscreteBlend>& Binding) const
 	{
 		const FInstanceData* InstanceData = Binding.GetInstanceData<FInstanceData>();
 
@@ -226,7 +226,7 @@ namespace UE::AnimNext
 
 			if (SubGraphEntry.State == ESlotState::ActiveWithGraph)
 			{
-				InstanceData->SubGraphSlots[ChildIndex].GraphInstance.Release();
+				SubGraphEntry.GraphInstance.Release();
 			}
 
 			SubGraphEntry.State = ESlotState::Inactive;
diff --git a/Engine/Plugins/Experimental/AnimNextAnimGraph/Source/AnimNextAnimGraph/Private/Traits/SubGraphHost.h b/Engine/Plugins/Experimental/AnimNextAnimGraph/Source/AnimNextAnimGraph/Private/Traits/SubGraphHost.h
--- a/Engine/Plugins/Experimental/AnimNextAnimGraph/Source/AnimNextAnimGraph/Private/Traits/SubGraphHost.h
+++ b/Engine/Plugins/Experimental/AnimNextAnimGraph/Source/AnimNextAnimGraph/Private/Traits/SubGraphHost.h
@@ -105,8 +105,8 @@ namespace UE::AnimNext
 		virtual void GetChildren(const FExecutionContext& Context, const TTraitBinding<IHierarchy>& Binding, FChildrenArray& Children) const override;
 
 		// IDiscreteBlend impl
-		virtual float GetBlendWeight(const FExecutionContext& Context, const TTraitBinding<IDiscreteBlend>& Binding, int32 ChildIndex) const override;
-		virtual int32 GetBlendDestinationChildIndex(const FExecutionContext& Context, const TTraitBinding<IDiscreteBlend>& Binding) const override;
+		virtual float GetBlendWeight(FExecutionContext& Context, const TTraitBinding<IDiscreteBlend>& Binding, int32 ChildIndex) const override;
+		virtual int32 GetBlendDestinationChildIndex(FExecutionContext& Context, const TTraitBinding<IDiscreteBlend>& Binding) const override;
 		virtual void OnBlendTransition(FExecutionContext& Context, const TTraitBinding<IDiscreteBlend>& Binding, int32 OldChildIndex, int32 NewChildIndex) const override;
 		virtual void OnBlendInitiated(FExecutionContext& Context, const TTraitBinding<IDiscreteBlend>& Binding, int32 ChildIndex) const override;
 		virtual void OnBlendTerminated(FExecutionContext& Context, const TTraitBinding<IDiscreteBlend>& Binding, int32 ChildIndex) const override;
diff --git a/Engine/Plugins/Experimental/AnimNextAnimGraph/Source/AnimNextAnimGraph/Private/Traits/SynchronizeUsingGroups.cpp b/Engine/Plugins/Experimental/AnimNextAnimGraph/Source/AnimNextAnimGraph/Private/Traits/SynchronizeUsingGroups.cpp
--- a/Engine/Plugins/Experimental/AnimNextAnimGraph/Source/AnimNextAnimGraph/Private/Traits/SynchronizeUsingGroups.cpp
+++ b/Engine/Plugins/Experimental/AnimNextAnimGraph/Source/AnimNextAnimGraph/Private/Traits/SynchronizeUsingGroups.cpp
@@ -46,19 +46,19 @@ namespace UE::AnimNext
 		Component.RegisterWithGroup(GroupName, GroupRole, Binding.GetTraitPtr(), TraitState);
 	}
 
-	FName FSynchronizeUsingGroupsTrait::GetGroupName(const FExecutionContext& Context, const TTraitBinding<IGroupSynchronization>& Binding) const
+	FName FSynchronizeUsingGroupsTrait::GetGroupName(FExecutionContext& Context, const TTraitBinding<IGroupSynchronization>& Binding) const
 	{
 		const FSharedData* SharedData = Binding.GetSharedData<FSharedData>();
 		return SharedData->GroupName;
 	}
 
-	EAnimGroupRole::Type FSynchronizeUsingGroupsTrait::GetGroupRole(const FExecutionContext& Context, const TTraitBinding<IGroupSynchronization>& Binding) const
+	EAnimGroupRole::Type FSynchronizeUsingGroupsTrait::GetGroupRole(FExecutionContext& Context, const TTraitBinding<IGroupSynchronization>& Binding) const
 	{
 		const FSharedData* SharedData = Binding.GetSharedData<FSharedData>();
 		return SharedData->GroupRole;
 	}
 
-	float FSynchronizeUsingGroupsTrait::AdvanceBy(const FExecutionContext& Context, const TTraitBinding<IGroupSynchronization>& Binding, float DeltaTime) const
+	float FSynchronizeUsingGroupsTrait::AdvanceBy(FExecutionContext& Context, const TTraitBinding<IGroupSynchronization>& Binding, float DeltaTime) const
 	{
 		FInstanceData* InstanceData = Binding.GetInstanceData<FInstanceData>();
 
@@ -75,7 +75,7 @@ namespace UE::AnimNext
 		return ProgressRatio;
 	}
 
-	void FSynchronizeUsingGroupsTrait::AdvanceToRatio(const FExecutionContext& Context, const TTraitBinding<IGroupSynchronization>& Binding, float ProgressRatio) const
+	void FSynchronizeUsingGroupsTrait::AdvanceToRatio(FExecutionContext& Context, const TTraitBinding<IGroupSynchronization>& Binding, float ProgressRatio) const
 	{
 		FInstanceData* InstanceData = Binding.GetInstanceData<FInstanceData>();
 
@@ -90,7 +90,7 @@ namespace UE::AnimNext
 		InstanceData->bFreezeTimeline = true;
 	}
 
-	float FSynchronizeUsingGroupsTrait::AdvanceBy(const FExecutionContext& Context, const TTraitBinding<ITimeline>& Binding, float DeltaTime) const
+	float FSynchronizeUsingGroupsTrait::AdvanceBy(FExecutionContext& Context, const TTraitBinding<ITimeline>& Binding, float DeltaTime) const
 	{
 		const FInstanceData* InstanceData = Binding.GetInstanceData<FInstanceData>();
 		if (InstanceData->bFreezeTimeline)
@@ -101,7 +101,7 @@ namespace UE::AnimNext
 		return ITimeline::AdvanceBy(Context, Binding, DeltaTime);
 	}
 
-	void FSynchronizeUsingGroupsTrait::AdvanceToRatio(const FExecutionContext& Context, const TTraitBinding<ITimeline>& Binding, float ProgressRatio) const
+	void FSynchronizeUsingGroupsTrait::AdvanceToRatio(FExecutionContext& Context, const TTraitBinding<ITimeline>& Binding, float ProgressRatio) const
 	{
 		const FInstanceData* InstanceData = Binding.GetInstanceData<FInstanceData>();
 		if (InstanceData->bFreezeTimeline)
diff --git a/Engine/Plugins/Experimental/AnimNextAnimGraph/Source/AnimNextAnimGraph/Private/Traits/SynchronizeUsingGroups.h b/Engine/Plugins/Experimental/AnimNextAnimGraph/Source/AnimNextAnimGraph/Private/Traits/SynchronizeUsingGroups.h
--- a/Engine/Plugins/Experimental/AnimNextAnimGraph/Source/AnimNextAnimGraph/Private/Traits/SynchronizeUsingGroups.h
+++ b/Engine/Plugins/Experimental/AnimNextAnimGraph/Source/AnimNextAnimGraph/Private/Traits/SynchronizeUsingGroups.h
@@ -48,13 +48,13 @@ namespace UE::AnimNext
 		virtual void PreUpdate(FUpdateTraversalContext& Context, const TTraitBinding<IUpdate>& Binding, const FTraitUpdateState& TraitState) const override;
 
 		// IGroupSynchronization impl
-		virtual FName GetGroupName(const FExecutionContext& Context, const TTraitBinding<IGroupSynchronization>& Binding) const override;
-		virtual EAnimGroupRole::Type GetGroupRole(const FExecutionContext& Context, const TTraitBinding<IGroupSynchronization>& Binding) const override;
-		virtual float AdvanceBy(const FExecutionContext& Context, const TTraitBinding<IGroupSynchronization>& Binding, float DeltaTime) const override;
-		virtual void AdvanceToRatio(const FExecutionContext& Context, const TTraitBinding<IGroupSynchronization>& Binding, float ProgressRatio) const override;
+		virtual FName GetGroupName(FExecutionContext& Context, const TTraitBinding<IGroupSynchronization>& Binding) const override;
+		virtual EAnimGroupRole::Type GetGroupRole(FExecutionContext& Context, const TTraitBinding<IGroupSynchronization>& Binding) const override;
+		virtual float AdvanceBy(FExecutionContext& Context, const TTraitBinding<IGroupSynchronization>& Binding, float DeltaTime) const override;
+		virtual void AdvanceToRatio(FExecutionContext& Context, const TTraitBinding<IGroupSynchronization>& Binding, float ProgressRatio) const override;
 
 		// ITimeline impl
-		virtual float AdvanceBy(const FExecutionContext& Context, const TTraitBinding<ITimeline>& Binding, float DeltaTime) const override;
-		virtual void AdvanceToRatio(const FExecutionContext& Context, const TTraitBinding<ITimeline>& Binding, float ProgressRatio) const override;
+		virtual float AdvanceBy(FExecutionContext& Context, const TTraitBinding<ITimeline>& Binding, float DeltaTime) const override;
+		virtual void AdvanceToRatio(FExecutionContext& Context, const TTraitBinding<ITimeline>& Binding, float ProgressRatio) const override;
 	};
 }
diff --git a/Engine/Plugins/Experimental/AnimNextAnimGraph/Source/AnimNextAnimGraph/Public/TraitInterfaces/IDiscreteBlend.h b/Engine/Plugins/Experimental/AnimNextAnimGraph/Source/AnimNextAnimGraph/Public/TraitInterfaces/IDiscreteBlend.h
--- a/Engine/Plugins/Experimental/AnimNextAnimGraph/Source/AnimNextAnimGraph/Public/TraitInterfaces/IDiscreteBlend.h
+++ b/Engine/Plugins/Experimental/AnimNextAnimGraph/Source/AnimNextAnimGraph/Public/TraitInterfaces/IDiscreteBlend.h
@@ -22,16 +22,16 @@ namespace UE::AnimNext
 		// Returns the blend weight for the specified child
 		// Multiple children can have non-zero weight but their sum must be 1.0
 		// Returns -1.0 if the child index is invalid
-		virtual float GetBlendWeight(const FExecutionContext& Context, const TTraitBinding<IDiscreteBlend>& Binding, int32 ChildIndex) const;
+		virtual float GetBlendWeight(FExecutionContext& Context, const TTraitBinding<IDiscreteBlend>& Binding, int32 ChildIndex) const;
 
 		// Returns the blend alpha for the specified child
 		// Returns nullptr if the child index is invalid
 		// Allows additive traits to query the internal state of the base trait
-		virtual const FAlphaBlend* GetBlendState(const FExecutionContext& Context, const TTraitBinding<IDiscreteBlend>& Binding, int32 ChildIndex) const;
+		virtual const FAlphaBlend* GetBlendState(FExecutionContext& Context, const TTraitBinding<IDiscreteBlend>& Binding, int32 ChildIndex) const;
 
 		// Returns the blend destination child index (aka the active child index)
 		// Returns INDEX_NONE if no child is active
-		virtual int32 GetBlendDestinationChildIndex(const FExecutionContext& Context, const TTraitBinding<IDiscreteBlend>& Binding) const;
+		virtual int32 GetBlendDestinationChildIndex(FExecutionContext& Context, const TTraitBinding<IDiscreteBlend>& Binding) const;
 
 		// Called when a blend transition between children occurs
 		// OldChildIndex can be INDEX_NONE if there was no previously active child
@@ -54,19 +54,19 @@ namespace UE::AnimNext
 	struct TTraitBinding<IDiscreteBlend> : FTraitBinding
 	{
 		// @see IDiscreteBlend::GetBlendWeight
-		float GetBlendWeight(const FExecutionContext& Context, int32 ChildIndex) const
+		float GetBlendWeight(FExecutionContext& Context, int32 ChildIndex) const
 		{
 			return GetInterface()->GetBlendWeight(Context, *this, ChildIndex);
 		}
 
 		// @see IDiscreteBlend::GetBlendState
-		const FAlphaBlend* GetBlendState(const FExecutionContext& Context, int32 ChildIndex) const
+		const FAlphaBlend* GetBlendState(FExecutionContext& Context, int32 ChildIndex) const
 		{
 			return GetInterface()->GetBlendState(Context, *this, ChildIndex);
 		}
 
 		// @see IDiscreteBlend::GetBlendDestinationChildIndex
-		int32 GetBlendDestinationChildIndex(const FExecutionContext& Context) const
+		int32 GetBlendDestinationChildIndex(FExecutionContext& Context) const
 		{
 			return GetInterface()->GetBlendDestinationChildIndex(Context, *this);
 		}
diff --git a/Engine/Plugins/Experimental/AnimNextAnimGraph/Source/AnimNextAnimGraph/Public/TraitInterfaces/IGroupSynchronization.h b/Engine/Plugins/Experimental/AnimNextAnimGraph/Source/AnimNextAnimGraph/Public/TraitInterfaces/IGroupSynchronization.h
--- a/Engine/Plugins/Experimental/AnimNextAnimGraph/Source/AnimNextAnimGraph/Public/TraitInterfaces/IGroupSynchronization.h
+++ b/Engine/Plugins/Experimental/AnimNextAnimGraph/Source/AnimNextAnimGraph/Public/TraitInterfaces/IGroupSynchronization.h
@@ -19,18 +19,18 @@ namespace UE::AnimNext
 		DECLARE_ANIM_TRAIT_INTERFACE(IGroupSynchronization, GRS, 0xf607d0fd)
 
 		// Returns the group name used for synchronization
-		virtual FName GetGroupName(const FExecutionContext& Context, const TTraitBinding<IGroupSynchronization>& Binding) const;
+		virtual FName GetGroupName(FExecutionContext& Context, const TTraitBinding<IGroupSynchronization>& Binding) const;
 
 		// Returns the group role used for synchronization
-		virtual EAnimGroupRole::Type GetGroupRole(const FExecutionContext& Context, const TTraitBinding<IGroupSynchronization>& Binding) const;
+		virtual EAnimGroupRole::Type GetGroupRole(FExecutionContext& Context, const TTraitBinding<IGroupSynchronization>& Binding) const;
 
 		// Called by the sync group graph instance component once a group has been synchronized to advance time on the leader
 		// Returns the progress ratio of playback: 0.0 = start of animation, 1.0 = end of animation
-		virtual float AdvanceBy(const FExecutionContext& Context, const TTraitBinding<IGroupSynchronization>& Binding, float DeltaTime) const;
+		virtual float AdvanceBy(FExecutionContext& Context, const TTraitBinding<IGroupSynchronization>& Binding, float DeltaTime) const;
 
 		// Called by the sync group graph instance component once a group has been synchronized to advance time on each follower
 		// Progress ratio must be between [0.0, 1.0]
-		virtual void AdvanceToRatio(const FExecutionContext& Context, const TTraitBinding<IGroupSynchronization>& Binding, float ProgressRatio) const;
+		virtual void AdvanceToRatio(FExecutionContext& Context, const TTraitBinding<IGroupSynchronization>& Binding, float ProgressRatio) const;
 	};
 
 	/**
@@ -40,25 +40,25 @@ namespace UE::AnimNext
 	struct TTraitBinding<IGroupSynchronization> : FTraitBinding
 	{
 		// @see IGroupSynchronization::GetGroupName
-		FName GetGroupName(const FExecutionContext& Context) const
+		FName GetGroupName(FExecutionContext& Context) const
 		{
 			return GetInterface()->GetGroupName(Context, *this);
 		}
 
 		// @see IGroupSynchronization::GetGroupRole
-		EAnimGroupRole::Type GetGroupRole(const FExecutionContext& Context) const
+		EAnimGroupRole::Type GetGroupRole(FExecutionContext& Context) const
 		{
 			return GetInterface()->GetGroupRole(Context, *this);
 		}
 
 		// @see IGroupSynchronization::AdvanceBy
-		float AdvanceBy(const FExecutionContext& Context, float DeltaTime) const
+		float AdvanceBy(FExecutionContext& Context, float DeltaTime) const
 		{
 			return GetInterface()->AdvanceBy(Context, *this, DeltaTime);
 		}
 
 		// @see IGroupSynchronization::AdvanceToRatio
-		void AdvanceToRatio(const FExecutionContext& Context, float ProgressRatio) const
+		void AdvanceToRatio(FExecutionContext& Context, float ProgressRatio) const
 		{
 			GetInterface()->AdvanceToRatio(Context, *this, ProgressRatio);
 		}
diff --git a/Engine/Plugins/Experimental/AnimNextAnimGraph/Source/AnimNextAnimGraph/Public/TraitInterfaces/ISmoothBlend.h b/Engine/Plugins/Experimental/AnimNextAnimGraph/Source/AnimNextAnimGraph/Public/TraitInterfaces/ISmoothBlend.h
--- a/Engine/Plugins/Experimental/AnimNextAnimGraph/Source/AnimNextAnimGraph/Public/TraitInterfaces/ISmoothBlend.h
+++ b/Engine/Plugins/Experimental/AnimNextAnimGraph/Source/AnimNextAnimGraph/Public/TraitInterfaces/ISmoothBlend.h
@@ -18,7 +18,7 @@ namespace UE::AnimNext
 		DECLARE_ANIM_TRAIT_INTERFACE(ISmoothBlend, SMB, 0x1c2c1739)
 
 		// Returns the desired blend time for the specified child
-		virtual float GetBlendTime(const FExecutionContext& Context, const TTraitBinding<ISmoothBlend>& Binding, int32 ChildIndex) const;
+		virtual float GetBlendTime(FExecutionContext& Context, const TTraitBinding<ISmoothBlend>& Binding, int32 ChildIndex) const;
 	};
 
 	/**
@@ -28,7 +28,7 @@ namespace UE::AnimNext
 	struct TTraitBinding<ISmoothBlend> : FTraitBinding
 	{
 		// @see ISmoothBlend::GetBlendTime
-		float GetBlendTime(const FExecutionContext& Context, int32 ChildIndex) const
+		float GetBlendTime(FExecutionContext& Context, int32 ChildIndex) const
 		{
 			return GetInterface()->GetBlendTime(Context, *this, ChildIndex);
 		}
diff --git a/Engine/Plugins/Experimental/AnimNextAnimGraph/Source/AnimNextAnimGraph/Public/TraitInterfaces/ITimeline.h b/Engine/Plugins/Experimental/AnimNextAnimGraph/Source/AnimNextAnimGraph/Public/TraitInterfaces/ITimeline.h
--- a/Engine/Plugins/Experimental/AnimNextAnimGraph/Source/AnimNextAnimGraph/Public/TraitInterfaces/ITimeline.h
+++ b/Engine/Plugins/Experimental/AnimNextAnimGraph/Source/AnimNextAnimGraph/Public/TraitInterfaces/ITimeline.h
@@ -18,15 +18,15 @@ namespace UE::AnimNext
 		DECLARE_ANIM_TRAIT_INTERFACE(ITimeline, TIM, 0x53760727)
 
 		// Returns the play rate of this timeline
-		virtual float GetPlayRate(const FExecutionContext& Context, const TTraitBinding<ITimeline>& Binding) const;
+		virtual float GetPlayRate(FExecutionContext& Context, const TTraitBinding<ITimeline>& Binding) const;
 
 		// Advances time by the provided delta time (positive or negative) on this timeline
-		// Returns the progress ratio of playback: 0.0 = start of animation, 1.0 = end of animation
-		virtual float AdvanceBy(const FExecutionContext& Context, const TTraitBinding<ITimeline>& Binding, float DeltaTime) const;
+		// Returns the progress of playback
+		virtual float AdvanceBy(FExecutionContext& Context, const TTraitBinding<ITimeline>& Binding, float DeltaTime) const;
 
 		// Advances time to the specified progress ratio on this timeline
 		// Progress ratio must be between [0.0, 1.0]
-		virtual void AdvanceToRatio(const FExecutionContext& Context, const TTraitBinding<ITimeline>& Binding, float ProgressRatio) const;
+		virtual void AdvanceToRatio(FExecutionContext& Context, const TTraitBinding<ITimeline>& Binding, float ProgressRatio) const;
 	};
 
 	/**
@@ -36,19 +36,19 @@ namespace UE::AnimNext
 	struct TTraitBinding<ITimeline> : FTraitBinding
 	{
 		// @see ITimeline::GetPlayRate
-		float GetPlayRate(const FExecutionContext& Context) const
+		float GetPlayRate(FExecutionContext& Context) const
 		{
 			return GetInterface()->GetPlayRate(Context, *this);
 		}
 
 		// @see ITimeline::AdvanceBy
-		float AdvanceBy(const FExecutionContext& Context, float DeltaTime) const
+		float AdvanceBy(FExecutionContext& Context, float DeltaTime) const
 		{
 			return GetInterface()->AdvanceBy(Context, *this, DeltaTime);
 		}
 
 		// @see ITimeline::AdvanceToRatio
-		void AdvanceToRatio(const FExecutionContext& Context, float ProgressRatio) const
+		void AdvanceToRatio(FExecutionContext& Context, float ProgressRatio) const
 		{
 			GetInterface()->AdvanceToRatio(Context, *this, ProgressRatio);
 		}
