Enable latent property freezing on blend out

diff --git a/Engine/Plugins/Experimental/AnimNextAnimGraph/Source/AnimNextAnimGraph/Private/TraitInterfaces/IUpdate.cpp b/Engine/Plugins/Experimental/AnimNextAnimGraph/Source/AnimNextAnimGraph/Private/TraitInterfaces/IUpdate.cpp
--- a/Engine/Plugins/Experimental/AnimNextAnimGraph/Source/AnimNextAnimGraph/Private/TraitInterfaces/IUpdate.cpp
+++ b/Engine/Plugins/Experimental/AnimNextAnimGraph/Source/AnimNextAnimGraph/Private/TraitInterfaces/IUpdate.cpp
@@ -514,7 +514,9 @@ namespace UE::AnimNext
 				// First, if it has latent pins, we must execute and cache their results
 				// This will ensure that other calls into this node will have a consistent view of
 				// what the node saw when it started to update. We thus take a snapshot.
-				const bool bIsFrozen = false;	// Not yet supported
+				// When a trait stack is blending out, its properties are frozen by default
+				// unless a property opts to always update regardless.
+				const bool bIsFrozen = Entry->TraitState.IsBlendingOut();
 				Entry->TraitStack.SnapshotLatentProperties(bIsFrozen);
 
 				const bool bImplementsIUpdate = Entry->TraitStack.GetInterface(Entry->UpdateTrait);
