Split smooth blend trait into two: behavior and arguments
Abstract arguments into an interface

diff --git a/Engine/Plugins/Experimental/AnimNextAnimGraph/Source/AnimNextAnimGraph/Private/TraitInterfaces/ISmoothBlend.cpp b/Engine/Plugins/Experimental/AnimNextAnimGraph/Source/AnimNextAnimGraph/Private/TraitInterfaces/ISmoothBlend.cpp
--- a/Engine/Plugins/Experimental/AnimNextAnimGraph/Source/AnimNextAnimGraph/Private/TraitInterfaces/ISmoothBlend.cpp
+++ b/Engine/Plugins/Experimental/AnimNextAnimGraph/Source/AnimNextAnimGraph/Private/TraitInterfaces/ISmoothBlend.cpp
@@ -16,4 +16,26 @@ namespace UE::AnimNext
 
 		return 0.0f;
 	}
+
+	EAlphaBlendOption ISmoothBlend::GetBlendType(FExecutionContext& Context, const TTraitBinding<ISmoothBlend>& Binding, int32 ChildIndex) const
+	{
+		TTraitBinding<ISmoothBlend> SuperBinding;
+		if (Binding.GetStackInterfaceSuper(SuperBinding))
+		{
+			return SuperBinding.GetBlendType(Context, ChildIndex);
+		}
+
+		return EAlphaBlendOption::Linear;
+	}
+
+	UCurveFloat* ISmoothBlend::GetCustomBlendCurve(FExecutionContext& Context, const TTraitBinding<ISmoothBlend>& Binding, int32 ChildIndex) const
+	{
+		TTraitBinding<ISmoothBlend> SuperBinding;
+		if (Binding.GetStackInterfaceSuper(SuperBinding))
+		{
+			return SuperBinding.GetCustomBlendCurve(Context, ChildIndex);
+		}
+
+		return nullptr;
+	}
 }
diff --git a/Engine/Plugins/Experimental/AnimNextAnimGraph/Source/AnimNextAnimGraph/Private/Traits/BlendSmoother.cpp b/Engine/Plugins/Experimental/AnimNextAnimGraph/Source/AnimNextAnimGraph/Private/Traits/BlendSmoother.cpp
--- a/Engine/Plugins/Experimental/AnimNextAnimGraph/Source/AnimNextAnimGraph/Private/Traits/BlendSmoother.cpp
+++ b/Engine/Plugins/Experimental/AnimNextAnimGraph/Source/AnimNextAnimGraph/Private/Traits/BlendSmoother.cpp
@@ -10,19 +10,26 @@
 
 namespace UE::AnimNext
 {
+	AUTO_REGISTER_ANIM_TRAIT(FBlendSmootherCoreTrait)
 	AUTO_REGISTER_ANIM_TRAIT(FBlendSmootherTrait)
 
 	// Trait implementation boilerplate
 	#define TRAIT_INTERFACE_ENUMERATOR(GeneratorMacro) \
 		GeneratorMacro(IDiscreteBlend) \
 		GeneratorMacro(IEvaluate) \
-		GeneratorMacro(ISmoothBlend) \
 		GeneratorMacro(IUpdate) \
 
+	GENERATE_ANIM_TRAIT_IMPLEMENTATION(FBlendSmootherCoreTrait, TRAIT_INTERFACE_ENUMERATOR, NULL_ANIM_TRAIT_INTERFACE_ENUMERATOR, NULL_ANIM_TRAIT_EVENT_ENUMERATOR)
+	#undef TRAIT_INTERFACE_ENUMERATOR
+
+	// Trait implementation boilerplate
+	#define TRAIT_INTERFACE_ENUMERATOR(GeneratorMacro) \
+		GeneratorMacro(ISmoothBlend) \
+
 	GENERATE_ANIM_TRAIT_IMPLEMENTATION(FBlendSmootherTrait, TRAIT_INTERFACE_ENUMERATOR, NULL_ANIM_TRAIT_INTERFACE_ENUMERATOR, NULL_ANIM_TRAIT_EVENT_ENUMERATOR)
 	#undef TRAIT_INTERFACE_ENUMERATOR
 
-	void FBlendSmootherTrait::PostEvaluate(FEvaluateTraversalContext& Context, const TTraitBinding<IEvaluate>& Binding) const
+	void FBlendSmootherCoreTrait::PostEvaluate(FEvaluateTraversalContext& Context, const TTraitBinding<IEvaluate>& Binding) const
 	{
 		const FInstanceData* InstanceData = Binding.GetInstanceData<FInstanceData>();
 
@@ -79,7 +86,7 @@ namespace UE::AnimNext
 		Context.AppendTask(FAnimNextNormalizeKeyframeRotationsTask());
 	}
 
-	void FBlendSmootherTrait::PreUpdate(FUpdateTraversalContext& Context, const TTraitBinding<IUpdate>& Binding, const FTraitUpdateState& TraitState) const
+	void FBlendSmootherCoreTrait::PreUpdate(FUpdateTraversalContext& Context, const TTraitBinding<IUpdate>& Binding, const FTraitUpdateState& TraitState) const
 	{
 		const FSharedData* SharedData = Binding.GetSharedData<FSharedData>();
 		FInstanceData* InstanceData = Binding.GetInstanceData<FInstanceData>();
@@ -158,33 +165,39 @@ namespace UE::AnimNext
 		}
 	}
 
-	float FBlendSmootherTrait::GetBlendWeight(FExecutionContext& Context, const TTraitBinding<IDiscreteBlend>& Binding, int32 ChildIndex) const
+	float FBlendSmootherCoreTrait::GetBlendWeight(FExecutionContext& Context, const TTraitBinding<IDiscreteBlend>& Binding, int32 ChildIndex) const
 	{
 		const FInstanceData* InstanceData = Binding.GetInstanceData<FInstanceData>();
 		return InstanceData->PerChildBlendData.IsValidIndex(ChildIndex) ? InstanceData->PerChildBlendData[ChildIndex].Weight : -1.0f;
 	}
 
-	const FAlphaBlend* FBlendSmootherTrait::GetBlendState(FExecutionContext& Context, const TTraitBinding<IDiscreteBlend>& Binding, int32 ChildIndex) const
+	const FAlphaBlend* FBlendSmootherCoreTrait::GetBlendState(FExecutionContext& Context, const TTraitBinding<IDiscreteBlend>& Binding, int32 ChildIndex) const
 	{
 		const FInstanceData* InstanceData = Binding.GetInstanceData<FInstanceData>();
 		return InstanceData->PerChildBlendData.IsValidIndex(ChildIndex) ? &InstanceData->PerChildBlendData[ChildIndex].Blend : nullptr;
 	}
 
-	void FBlendSmootherTrait::OnBlendTransition(FExecutionContext& Context, const TTraitBinding<IDiscreteBlend>& Binding, int32 OldChildIndex, int32 NewChildIndex) const
+	void FBlendSmootherCoreTrait::OnBlendTransition(FExecutionContext& Context, const TTraitBinding<IDiscreteBlend>& Binding, int32 OldChildIndex, int32 NewChildIndex) const
 	{
 		const FSharedData* SharedData = Binding.GetSharedData<FSharedData>();
 		FInstanceData* InstanceData = Binding.GetInstanceData<FInstanceData>();
 
+		TTraitBinding<ISmoothBlend> SmoothBlendTrait;
+		Binding.GetStackInterface(SmoothBlendTrait);
+
 		const int32 NumChildren = InstanceData->PerChildBlendData.Num();
 		if (NewChildIndex >= NumChildren)
 		{
 			// We have a new child
 			check(NewChildIndex == NumChildren);
 
+			const EAlphaBlendOption BlendType = SmoothBlendTrait.GetBlendType(Context, NewChildIndex);
+			UCurveFloat* CustomBlendCurve = SmoothBlendTrait.GetCustomBlendCurve(Context, NewChildIndex);
+
 			FBlendData& ChildBlendData = InstanceData->PerChildBlendData.AddDefaulted_GetRef();
 
-			ChildBlendData.Blend.SetBlendOption(SharedData->BlendType);
-			ChildBlendData.Blend.SetCustomCurve(SharedData->CustomBlendCurve);
+			ChildBlendData.Blend.SetBlendOption(BlendType);
+			ChildBlendData.Blend.SetCustomCurve(CustomBlendCurve);
 		}
 
 		// scale by the weight difference since we want consistency:
@@ -194,9 +207,6 @@ namespace UE::AnimNext
 		const float NewChildDesiredWeight = 1.0f;
 		const float WeightDifference = FMath::Clamp(FMath::Abs(NewChildDesiredWeight - NewChildCurrentWeight), 0.0f, 1.0f);
 
-		TTraitBinding<ISmoothBlend> SmoothBlendTrait;
-		Binding.GetStackInterface(SmoothBlendTrait);
-
 		const float BlendTime = SmoothBlendTrait.GetBlendTime(Context, NewChildIndex);
 		const float RemainingBlendTime = OldChildIndex != INDEX_NONE ? (BlendTime * WeightDifference) : 0.0f;
 
@@ -231,6 +241,28 @@ namespace UE::AnimNext
 		DiscreteBlendTrait.OnBlendInitiated(Context, NewChildIndex);
 	}
 
+	void FBlendSmootherCoreTrait::InitializeInstanceData(FExecutionContext& Context, const FTraitBinding& Binding, const FSharedData* SharedData, FInstanceData* InstanceData)
+	{
+		check(InstanceData->PerChildBlendData.IsEmpty());
+
+		TTraitBinding<ISmoothBlend> SmoothBlendTrait;
+		Binding.GetStackInterface(SmoothBlendTrait);
+
+		const uint32 NumChildren = IHierarchy::GetNumStackChildren(Context, Binding);
+
+		InstanceData->PerChildBlendData.SetNum(NumChildren);
+
+		for (uint32 ChildIndex = 0; ChildIndex < NumChildren; ++ChildIndex)
+		{
+			const EAlphaBlendOption BlendType = SmoothBlendTrait.GetBlendType(Context, ChildIndex);
+			UCurveFloat* CustomBlendCurve = SmoothBlendTrait.GetCustomBlendCurve(Context, ChildIndex);
+
+			FBlendData& ChildBlendData = InstanceData->PerChildBlendData[ChildIndex];
+			ChildBlendData.Blend.SetBlendOption(BlendType);
+			ChildBlendData.Blend.SetCustomCurve(CustomBlendCurve);
+		}
+	}
+
 	float FBlendSmootherTrait::GetBlendTime(FExecutionContext& Context, const TTraitBinding<ISmoothBlend>& Binding, int32 ChildIndex) const
 	{
 		const FSharedData* SharedData = Binding.GetSharedData<FSharedData>();
@@ -247,23 +279,20 @@ namespace UE::AnimNext
 		}
 		else
 		{
-			// No blend time has been specified, we snap
-			return 0.0f;
+			// No blend time has been specified, forward below us on the stack, maybe someone can provide one
+			return ISmoothBlend::GetBlendTime(Context, Binding, ChildIndex);
 		}
 	}
 
-	void FBlendSmootherTrait::InitializeInstanceData(FExecutionContext& Context, const FTraitBinding& Binding, const FSharedData* SharedData, FInstanceData* InstanceData)
+	EAlphaBlendOption FBlendSmootherTrait::GetBlendType(FExecutionContext& Context, const TTraitBinding<ISmoothBlend>& Binding, int32 ChildIndex) const
 	{
-		check(InstanceData->PerChildBlendData.IsEmpty());
-
-		const uint32 NumChildren = IHierarchy::GetNumStackChildren(Context, Binding);
-
-		InstanceData->PerChildBlendData.SetNum(NumChildren);
+		const FSharedData* SharedData = Binding.GetSharedData<FSharedData>();
+		return SharedData->BlendType;
+	}
 
-		for (FBlendData& ChildBlendData : InstanceData->PerChildBlendData)
-		{
-			ChildBlendData.Blend.SetBlendOption(SharedData->BlendType);
-			ChildBlendData.Blend.SetCustomCurve(SharedData->CustomBlendCurve);
-		}
+	UCurveFloat* FBlendSmootherTrait::GetCustomBlendCurve(FExecutionContext& Context, const TTraitBinding<ISmoothBlend>& Binding, int32 ChildIndex) const
+	{
+		const FSharedData* SharedData = Binding.GetSharedData<FSharedData>();
+		return SharedData->CustomBlendCurve;
 	}
 }
diff --git a/Engine/Plugins/Experimental/AnimNextAnimGraph/Source/AnimNextAnimGraph/Private/Traits/BlendSmoother.h b/Engine/Plugins/Experimental/AnimNextAnimGraph/Source/AnimNextAnimGraph/Private/Traits/BlendSmoother.h
--- a/Engine/Plugins/Experimental/AnimNextAnimGraph/Source/AnimNextAnimGraph/Private/Traits/BlendSmoother.h
+++ b/Engine/Plugins/Experimental/AnimNextAnimGraph/Source/AnimNextAnimGraph/Private/Traits/BlendSmoother.h
@@ -15,8 +15,16 @@
 
 class UCurveFloat;
 
+USTRUCT(meta = (DisplayName = "Blend Smoother Core"))
+struct FAnimNextBlendSmootherCoreTraitSharedData : public FAnimNextTraitSharedData
+{
+	GENERATED_BODY()
+
+	// This struct is empty but required so that we can give a nice display name to the trait
+};
+
 USTRUCT(meta = (DisplayName = "Blend Smoother"))
-struct FAnimNextBlendSmootherTraitSharedData : public FAnimNextTraitSharedData
+struct FAnimNextBlendSmootherTraitSharedData : public FAnimNextBlendSmootherCoreTraitSharedData
 {
 	GENERATED_BODY()
 
@@ -36,15 +44,16 @@ struct FAnimNextBlendSmootherTraitSharedData : public FAnimNextTraitSharedData
 namespace UE::AnimNext
 {
 	/**
-	 * FBlendSmootherTrait
-	 * 
+	 * FBlendSmootherCoreTrait
+	 *
 	 * A trait that smoothly blends between discrete states over time.
+	 * It only implements the smoothing logic, it queries its required arguments using the ISmoothBlend interface.
 	 */
-	struct FBlendSmootherTrait : FAdditiveTrait, IEvaluate, IUpdate, IDiscreteBlend, ISmoothBlend
+	struct FBlendSmootherCoreTrait : FAdditiveTrait, IEvaluate, IUpdate, IDiscreteBlend
 	{
-		DECLARE_ANIM_TRAIT(FBlendSmootherTrait, 0xbaeb537b, FAdditiveTrait)
+		DECLARE_ANIM_TRAIT(FBlendSmootherCoreTrait, 0x440a43d2, FAdditiveTrait)
 
-		using FSharedData = FAnimNextBlendSmootherTraitSharedData;
+		using FSharedData = FAnimNextBlendSmootherCoreTraitSharedData;
 
 		// Struct for tracking blends for each pose
 		struct FBlendData
@@ -76,10 +85,25 @@ namespace UE::AnimNext
 		virtual const FAlphaBlend* GetBlendState(FExecutionContext& Context, const TTraitBinding<IDiscreteBlend>& Binding, int32 ChildIndex) const override;
 		virtual void OnBlendTransition(FExecutionContext& Context, const TTraitBinding<IDiscreteBlend>& Binding, int32 OldChildIndex, int32 NewChildIndex) const override;
 
-		// ISmoothBlend impl
-		virtual float GetBlendTime(FExecutionContext& Context, const TTraitBinding<ISmoothBlend>& Binding, int32 ChildIndex) const override;
-
 		// Internal impl
 		static void InitializeInstanceData(FExecutionContext& Context, const FTraitBinding& Binding, const FSharedData* SharedData, FInstanceData* InstanceData);
 	};
+
+	/**
+	 * FBlendSmootherTrait
+	 * 
+	 * A trait that smoothly blends between discrete states over time.
+	 * This trait implements both the logic and contains the arguments necessary.
+	 */
+	struct FBlendSmootherTrait : FBlendSmootherCoreTrait, ISmoothBlend
+	{
+		DECLARE_ANIM_TRAIT(FBlendSmootherTrait, 0xbaeb537b, FBlendSmootherCoreTrait)
+
+		using FSharedData = FAnimNextBlendSmootherTraitSharedData;
+
+		// ISmoothBlend impl
+		virtual float GetBlendTime(FExecutionContext& Context, const TTraitBinding<ISmoothBlend>& Binding, int32 ChildIndex) const override;
+		virtual EAlphaBlendOption GetBlendType(FExecutionContext& Context, const TTraitBinding<ISmoothBlend>& Binding, int32 ChildIndex) const override;
+		virtual UCurveFloat* GetCustomBlendCurve(FExecutionContext& Context, const TTraitBinding<ISmoothBlend>& Binding, int32 ChildIndex) const override;
+	};
 }
diff --git a/Engine/Plugins/Experimental/AnimNextAnimGraph/Source/AnimNextAnimGraph/Public/TraitInterfaces/ISmoothBlend.h b/Engine/Plugins/Experimental/AnimNextAnimGraph/Source/AnimNextAnimGraph/Public/TraitInterfaces/ISmoothBlend.h
--- a/Engine/Plugins/Experimental/AnimNextAnimGraph/Source/AnimNextAnimGraph/Public/TraitInterfaces/ISmoothBlend.h
+++ b/Engine/Plugins/Experimental/AnimNextAnimGraph/Source/AnimNextAnimGraph/Public/TraitInterfaces/ISmoothBlend.h
@@ -3,9 +3,12 @@
 #pragma once
 
 #include "CoreMinimal.h"
+#include "AlphaBlend.h"
 #include "TraitCore/ITraitInterface.h"
 #include "TraitCore/TraitBinding.h"
 
+class UCurveFloat;
+
 namespace UE::AnimNext
 {
 	/**
@@ -19,6 +22,12 @@ namespace UE::AnimNext
 
 		// Returns the desired blend time for the specified child
 		virtual float GetBlendTime(FExecutionContext& Context, const TTraitBinding<ISmoothBlend>& Binding, int32 ChildIndex) const;
+
+		// Returns the desired blend type for the specified child
+		virtual EAlphaBlendOption GetBlendType(FExecutionContext& Context, const TTraitBinding<ISmoothBlend>& Binding, int32 ChildIndex) const;
+
+		// Returns the desired blend curve for the specified child
+		virtual UCurveFloat* GetCustomBlendCurve(FExecutionContext& Context, const TTraitBinding<ISmoothBlend>& Binding, int32 ChildIndex) const;
 	};
 
 	/**
@@ -33,6 +42,18 @@ namespace UE::AnimNext
 			return GetInterface()->GetBlendTime(Context, *this, ChildIndex);
 		}
 
+		// @see ISmoothBlend::GetBlendType
+		EAlphaBlendOption GetBlendType(FExecutionContext& Context, int32 ChildIndex) const
+		{
+			return GetInterface()->GetBlendType(Context, *this, ChildIndex);
+		}
+
+		// @see ISmoothBlend::GetCustomBlendCurve
+		UCurveFloat* GetCustomBlendCurve(FExecutionContext& Context, int32 ChildIndex) const
+		{
+			return GetInterface()->GetCustomBlendCurve(Context, *this, ChildIndex);
+		}
+
 	protected:
 		const ISmoothBlend* GetInterface() const { return GetInterfaceTyped<ISmoothBlend>(); }
 	};
