Split inertializer trait into two: behavior and arguments
Abstract arguments into an interface

diff --git a/Engine/Plugins/Experimental/AnimNextAnimGraph/Source/AnimNextAnimGraph/Private/TraitInterfaces/IInertializerBlend.cpp b/Engine/Plugins/Experimental/AnimNextAnimGraph/Source/AnimNextAnimGraph/Private/TraitInterfaces/IInertializerBlend.cpp
new file mode 100644
--- /dev/null
+++ b/Engine/Plugins/Experimental/AnimNextAnimGraph/Source/AnimNextAnimGraph/Private/TraitInterfaces/IInertializerBlend.cpp
@@ -0,0 +1,19 @@
+// Copyright Epic Games, Inc. All Rights Reserved.
+
+#include "TraitInterfaces/IInertializerBlend.h"
+
+#include "TraitCore/ExecutionContext.h"
+
+namespace UE::AnimNext
+{
+	float IInertializerBlend::GetBlendTime(FExecutionContext& Context, const TTraitBinding<IInertializerBlend>& Binding, int32 ChildIndex) const
+	{
+		TTraitBinding<IInertializerBlend> SuperBinding;
+		if (Binding.GetStackInterfaceSuper(SuperBinding))
+		{
+			return SuperBinding.GetBlendTime(Context, ChildIndex);
+		}
+
+		return 0.0f;
+	}
+}
diff --git a/Engine/Plugins/Experimental/AnimNextAnimGraph/Source/AnimNextAnimGraph/Private/Traits/BlendInertializer.cpp b/Engine/Plugins/Experimental/AnimNextAnimGraph/Source/AnimNextAnimGraph/Private/Traits/BlendInertializer.cpp
--- a/Engine/Plugins/Experimental/AnimNextAnimGraph/Source/AnimNextAnimGraph/Private/Traits/BlendInertializer.cpp
+++ b/Engine/Plugins/Experimental/AnimNextAnimGraph/Source/AnimNextAnimGraph/Private/Traits/BlendInertializer.cpp
@@ -9,6 +9,7 @@
 
 namespace UE::AnimNext
 {
+	AUTO_REGISTER_ANIM_TRAIT(FBlendInertializerCoreTrait)
 	AUTO_REGISTER_ANIM_TRAIT(FBlendInertializerTrait)
 
 	// Trait implementation boilerplate
@@ -16,30 +17,55 @@ namespace UE::AnimNext
 		GeneratorMacro(IDiscreteBlend) \
 		GeneratorMacro(ISmoothBlend) \
 
+	GENERATE_ANIM_TRAIT_IMPLEMENTATION(FBlendInertializerCoreTrait, TRAIT_INTERFACE_ENUMERATOR, NULL_ANIM_TRAIT_INTERFACE_ENUMERATOR, NULL_ANIM_TRAIT_EVENT_ENUMERATOR)
+	#undef TRAIT_INTERFACE_ENUMERATOR
+
 	// Trait implementation boilerplate
-	#define TRAIT_REQUIRED_INTERFACE_ENUMERATOR(GeneratorMacro) \
-		GeneratorMacro(IDiscreteBlend) \
+	#define TRAIT_INTERFACE_ENUMERATOR(GeneratorMacro) \
+		GeneratorMacro(IInertializerBlend) \
 
-	GENERATE_ANIM_TRAIT_IMPLEMENTATION(FBlendInertializerTrait, TRAIT_INTERFACE_ENUMERATOR, TRAIT_REQUIRED_INTERFACE_ENUMERATOR, NULL_ANIM_TRAIT_EVENT_ENUMERATOR)
+	GENERATE_ANIM_TRAIT_IMPLEMENTATION(FBlendInertializerTrait, TRAIT_INTERFACE_ENUMERATOR, NULL_ANIM_TRAIT_INTERFACE_ENUMERATOR, NULL_ANIM_TRAIT_EVENT_ENUMERATOR)
 	#undef TRAIT_INTERFACE_ENUMERATOR
-	#undef TRAIT_REQUIRED_INTERFACE_ENUMERATOR
 
-	void FBlendInertializerTrait::OnBlendTransition(FExecutionContext& Context, const TTraitBinding<IDiscreteBlend>& Binding, int32 OldChildIndex, int32 NewChildIndex) const
+	void FBlendInertializerCoreTrait::OnBlendTransition(FExecutionContext& Context, const TTraitBinding<IDiscreteBlend>& Binding, int32 OldChildIndex, int32 NewChildIndex) const
 	{
 		// Trigger the new transition
 		IDiscreteBlend::OnBlendTransition(Context, Binding, OldChildIndex, NewChildIndex);
 
-		const FSharedData* SharedData = Binding.GetSharedData<FSharedData>();
+		TTraitBinding<IInertializerBlend> InertializerBlendTrait;
+		Binding.GetStackInterface(InertializerBlendTrait);
+
+		const float BlendTime = InertializerBlendTrait.GetBlendTime(Context, NewChildIndex);
+		if (BlendTime <= 0.0f)
+		{
+			return;	// No blend time means we are disabled
+		}
 
 		// Make Request Event
 		TSharedPtr<FAnimNextInertializationRequestEvent> Event = MakeTraitEvent<FAnimNextInertializationRequestEvent>();
-		Event->Request.BlendTime = SharedData->BlendTime;
+		Event->Request.BlendTime = BlendTime;
 		Context.RaiseOutputTraitEvent(Event);
 	}
 
-	float FBlendInertializerTrait::GetBlendTime(FExecutionContext& Context, const TTraitBinding<ISmoothBlend>& Binding, int32 ChildIndex) const
+	float FBlendInertializerCoreTrait::GetBlendTime(FExecutionContext& Context, const TTraitBinding<ISmoothBlend>& Binding, int32 ChildIndex) const
 	{
-		// We hijack the blend time and always transition instantaneously
-		return 0.0f;
+		TTraitBinding<IInertializerBlend> InertializerBlendTrait;
+		Binding.GetStackInterface(InertializerBlendTrait);
+
+		const float BlendTime = InertializerBlendTrait.GetBlendTime(Context, ChildIndex);
+		if (BlendTime > 0.0f)
+		{
+			// We hijack the blend time and always transition instantaneously
+			return 0.0f;
+		}
+
+		// We are disabled
+		return ISmoothBlend::GetBlendTime(Context, Binding, ChildIndex);
+	}
+
+	float FBlendInertializerTrait::GetBlendTime(FExecutionContext& Context, const TTraitBinding<IInertializerBlend>& Binding, int32 ChildIndex) const
+	{
+		const FSharedData* SharedData = Binding.GetSharedData<FSharedData>();
+		return SharedData->BlendTime;
 	}
 }
diff --git a/Engine/Plugins/Experimental/AnimNextAnimGraph/Source/AnimNextAnimGraph/Private/Traits/BlendInertializer.h b/Engine/Plugins/Experimental/AnimNextAnimGraph/Source/AnimNextAnimGraph/Private/Traits/BlendInertializer.h
--- a/Engine/Plugins/Experimental/AnimNextAnimGraph/Source/AnimNextAnimGraph/Private/Traits/BlendInertializer.h
+++ b/Engine/Plugins/Experimental/AnimNextAnimGraph/Source/AnimNextAnimGraph/Private/Traits/BlendInertializer.h
@@ -6,14 +6,21 @@
 
 #include "TraitCore/Trait.h"
 #include "TraitInterfaces/IDiscreteBlend.h"
+#include "TraitInterfaces/IInertializerBlend.h"
 #include "TraitInterfaces/ISmoothBlend.h"
 
 #include "BlendInertializer.generated.h"
 
-class UBlendProfile;
+USTRUCT(meta = (DisplayName = "Blend Inertializer Core"))
+struct FAnimNextBlendInertializerCoreTraitSharedData : public FAnimNextTraitSharedData
+{
+	GENERATED_BODY()
+
+	// This struct is empty but required so that we can give a nice display name to the trait
+};
 
 USTRUCT(meta = (DisplayName = "Blend Inertializer", Category = "Default"))
-struct FAnimNextBlendInertializerTraitSharedData : public FAnimNextTraitSharedData
+struct FAnimNextBlendInertializerTraitSharedData : public FAnimNextBlendInertializerCoreTraitSharedData
 {
 	GENERATED_BODY()
 
@@ -25,15 +32,16 @@ struct FAnimNextBlendInertializerTraitSharedData : public FAnimNextTraitSharedDa
 namespace UE::AnimNext
 {
 	/**
-	 * FBlendInertializerTrait
-	 * 
+	 * FBlendInertializerCoreTrait
+	 *
 	 * A trait that converts a normal smooth blend into an inertializing blend.
+	 * It only implements the inertializing logic, it queries its required arguments using the IInertializerBlend interface.
 	 */
-	struct FBlendInertializerTrait : FAdditiveTrait, IDiscreteBlend, ISmoothBlend
+	struct FBlendInertializerCoreTrait : FAdditiveTrait, IDiscreteBlend, ISmoothBlend
 	{
-		DECLARE_ANIM_TRAIT(FBlendInertializerTrait, 0x7ea0bdee, FAdditiveTrait)
+		DECLARE_ANIM_TRAIT(FBlendInertializerCoreTrait, 0x137a22ab, FAdditiveTrait)
 
-		using FSharedData = FAnimNextBlendInertializerTraitSharedData;
+		using FSharedData = FAnimNextBlendInertializerCoreTraitSharedData;
 
 		// IDiscreteBlend impl
 		virtual void OnBlendTransition(FExecutionContext& Context, const TTraitBinding<IDiscreteBlend>& Binding, int32 OldChildIndex, int32 NewChildIndex) const override;
@@ -41,4 +49,23 @@ namespace UE::AnimNext
 		// ISmoothBlend impl
 		virtual float GetBlendTime(FExecutionContext& Context, const TTraitBinding<ISmoothBlend>& Binding, int32 ChildIndex) const override;
 	};
+
+	/**
+	 * FBlendInertializerTrait
+	 * 
+	 * A trait that converts a normal smooth blend into an inertializing blend.
+	 * This trait implements both the logic and contains the arguments necessary.
+	 */
+	struct FBlendInertializerTrait : FBlendInertializerCoreTrait, IInertializerBlend
+	{
+		DECLARE_ANIM_TRAIT(FBlendInertializerTrait, 0x7ea0bdee, FBlendInertializerCoreTrait)
+
+		using FSharedData = FAnimNextBlendInertializerTraitSharedData;
+
+		// Unhide base function to silence override hiding warning with clang
+		using FBlendInertializerCoreTrait::GetBlendTime;
+
+		// IInertializerBlend impl
+		virtual float GetBlendTime(FExecutionContext& Context, const TTraitBinding<IInertializerBlend>& Binding, int32 ChildIndex) const override;
+	};
 }
diff --git a/Engine/Plugins/Experimental/AnimNextAnimGraph/Source/AnimNextAnimGraph/Public/TraitInterfaces/IInertializerBlend.h b/Engine/Plugins/Experimental/AnimNextAnimGraph/Source/AnimNextAnimGraph/Public/TraitInterfaces/IInertializerBlend.h
new file mode 100644
--- /dev/null
+++ b/Engine/Plugins/Experimental/AnimNextAnimGraph/Source/AnimNextAnimGraph/Public/TraitInterfaces/IInertializerBlend.h
@@ -0,0 +1,39 @@
+// Copyright Epic Games, Inc. All Rights Reserved.
+
+#pragma once
+
+#include "CoreMinimal.h"
+#include "TraitCore/ITraitInterface.h"
+#include "TraitCore/TraitBinding.h"
+
+namespace UE::AnimNext
+{
+	/**
+	 * IInertializerBlend
+	 *
+	 * This interface exposes inertializing blend related information.
+	 */
+	struct ANIMNEXTANIMGRAPH_API IInertializerBlend : ITraitInterface
+	{
+		DECLARE_ANIM_TRAIT_INTERFACE(IInertializerBlend, IB, 0x3856b8e9)
+
+		// Returns the desired blend time for the specified child
+		virtual float GetBlendTime(FExecutionContext& Context, const TTraitBinding<IInertializerBlend>& Binding, int32 ChildIndex) const;
+	};
+
+	/**
+	 * Specialization for trait binding.
+	 */
+	template<>
+	struct TTraitBinding<IInertializerBlend> : FTraitBinding
+	{
+		// @see IInertializerBlend::GetBlendTime
+		float GetBlendTime(FExecutionContext& Context, int32 ChildIndex) const
+		{
+			return GetInterface()->GetBlendTime(Context, *this, ChildIndex);
+		}
+
+	protected:
+		const IInertializerBlend* GetInterface() const { return GetInterfaceTyped<IInertializerBlend>(); }
+	};
+}
