Introduce timeline progress struct to encapsulate position and ratio
Add SimulateAdvanceBy on timeline interface

diff --git a/Engine/Plugins/Experimental/AnimNextAnimGraph/Source/AnimNextAnimGraph/Private/Graph/SyncGroup_GraphInstanceComponent.cpp b/Engine/Plugins/Experimental/AnimNextAnimGraph/Source/AnimNextAnimGraph/Private/Graph/SyncGroup_GraphInstanceComponent.cpp
--- a/Engine/Plugins/Experimental/AnimNextAnimGraph/Source/AnimNextAnimGraph/Private/Graph/SyncGroup_GraphInstanceComponent.cpp
+++ b/Engine/Plugins/Experimental/AnimNextAnimGraph/Source/AnimNextAnimGraph/Private/Graph/SyncGroup_GraphInstanceComponent.cpp
@@ -75,7 +75,7 @@ namespace UE::AnimNext
 			}
 
 			// Now that we found our leader, advance it by the delta time
-			float LeaderProgressRatio;
+			FTimelineProgress LeaderProgress;
 			{
 				const FSyncGroupMember& GroupLeader = GroupState.Members[LeaderIndex];
 
@@ -83,9 +83,11 @@ namespace UE::AnimNext
 				ensure(Context.GetStack(GroupLeader.TraitPtr, TraitStack));
 				ensure(TraitStack.GetInterface(GroupSyncTrait));
 
-				LeaderProgressRatio = GroupSyncTrait.AdvanceBy(Context, GroupLeader.TraitState.GetDeltaTime());
+				LeaderProgress = GroupSyncTrait.AdvanceBy(Context, GroupLeader.TraitState.GetDeltaTime());
 			}
 
+			const float LeaderProgressRatio = LeaderProgress.GetPositionRatio();
+
 			// Advance every follower to the same progress ratio as the leader
 			for (int32 MemberIndex = 0; MemberIndex < NumMembers; ++MemberIndex)
 			{
diff --git a/Engine/Plugins/Experimental/AnimNextAnimGraph/Source/AnimNextAnimGraph/Private/TraitInterfaces/IGroupSynchronization.cpp b/Engine/Plugins/Experimental/AnimNextAnimGraph/Source/AnimNextAnimGraph/Private/TraitInterfaces/IGroupSynchronization.cpp
--- a/Engine/Plugins/Experimental/AnimNextAnimGraph/Source/AnimNextAnimGraph/Private/TraitInterfaces/IGroupSynchronization.cpp
+++ b/Engine/Plugins/Experimental/AnimNextAnimGraph/Source/AnimNextAnimGraph/Private/TraitInterfaces/IGroupSynchronization.cpp
@@ -28,7 +28,7 @@ namespace UE::AnimNext
 		return EAnimGroupRole::CanBeLeader;
 	}
 
-	float IGroupSynchronization::AdvanceBy(FExecutionContext& Context, const TTraitBinding<IGroupSynchronization>& Binding, float DeltaTime) const
+	FTimelineProgress IGroupSynchronization::AdvanceBy(FExecutionContext& Context, const TTraitBinding<IGroupSynchronization>& Binding, float DeltaTime) const
 	{
 		TTraitBinding<IGroupSynchronization> SuperBinding;
 		if (Binding.GetStackInterfaceSuper(SuperBinding))
@@ -36,7 +36,7 @@ namespace UE::AnimNext
 			return SuperBinding.AdvanceBy(Context, DeltaTime);
 		}
 
-		return 0.0f;
+		return FTimelineProgress();
 	}
 
 	void IGroupSynchronization::AdvanceToRatio(FExecutionContext& Context, const TTraitBinding<IGroupSynchronization>& Binding, float ProgressRatio) const
diff --git a/Engine/Plugins/Experimental/AnimNextAnimGraph/Source/AnimNextAnimGraph/Private/TraitInterfaces/ITimeline.cpp b/Engine/Plugins/Experimental/AnimNextAnimGraph/Source/AnimNextAnimGraph/Private/TraitInterfaces/ITimeline.cpp
--- a/Engine/Plugins/Experimental/AnimNextAnimGraph/Source/AnimNextAnimGraph/Private/TraitInterfaces/ITimeline.cpp
+++ b/Engine/Plugins/Experimental/AnimNextAnimGraph/Source/AnimNextAnimGraph/Private/TraitInterfaces/ITimeline.cpp
@@ -17,7 +17,29 @@ namespace UE::AnimNext
 		return 1.0f;
 	}
 
-	float ITimeline::AdvanceBy(FExecutionContext& Context, const TTraitBinding<ITimeline>& Binding, float DeltaTime) const
+	FTimelineProgress ITimeline::GetProgress(FExecutionContext& Context, const TTraitBinding<ITimeline>& Binding) const
+	{
+		TTraitBinding<ITimeline> SuperBinding;
+		if (Binding.GetStackInterfaceSuper(SuperBinding))
+		{
+			return SuperBinding.GetProgress(Context);
+		}
+
+		return FTimelineProgress();
+	}
+
+	FTimelineProgress ITimeline::SimulateAdvanceBy(FExecutionContext& Context, const TTraitBinding<ITimeline>& Binding, float DeltaTime) const
+	{
+		TTraitBinding<ITimeline> SuperBinding;
+		if (Binding.GetStackInterfaceSuper(SuperBinding))
+		{
+			return SuperBinding.SimulateAdvanceBy(Context, DeltaTime);
+		}
+
+		return FTimelineProgress();
+	}
+
+	FTimelineProgress ITimeline::AdvanceBy(FExecutionContext& Context, const TTraitBinding<ITimeline>& Binding, float DeltaTime) const
 	{
 		TTraitBinding<ITimeline> SuperBinding;
 		if (Binding.GetStackInterfaceSuper(SuperBinding))
@@ -25,7 +47,7 @@ namespace UE::AnimNext
 			return SuperBinding.AdvanceBy(Context, DeltaTime);
 		}
 
-		return 0.0f;
+		return FTimelineProgress();
 	}
 
 	void ITimeline::AdvanceToRatio(FExecutionContext& Context, const TTraitBinding<ITimeline>& Binding, float ProgressRatio) const
diff --git a/Engine/Plugins/Experimental/AnimNextAnimGraph/Source/AnimNextAnimGraph/Private/Traits/SequencePlayer.cpp b/Engine/Plugins/Experimental/AnimNextAnimGraph/Source/AnimNextAnimGraph/Private/Traits/SequencePlayer.cpp
--- a/Engine/Plugins/Experimental/AnimNextAnimGraph/Source/AnimNextAnimGraph/Private/Traits/SequencePlayer.cpp
+++ b/Engine/Plugins/Experimental/AnimNextAnimGraph/Source/AnimNextAnimGraph/Private/Traits/SequencePlayer.cpp
@@ -53,7 +53,42 @@ namespace UE::AnimNext
 		return SharedData->GetPlayRate(Binding);
 	}
 
-	float FSequencePlayerTrait::AdvanceBy(FExecutionContext& Context, const TTraitBinding<ITimeline>& Binding, float DeltaTime) const
+	FTimelineProgress FSequencePlayerTrait::GetProgress(FExecutionContext& Context, const TTraitBinding<ITimeline>& Binding) const
+	{
+		const FInstanceData* InstanceData = Binding.GetInstanceData<FInstanceData>();
+		if (const UAnimSequence* AnimSeq = InstanceData->AnimSequence.Get())
+		{
+			return FTimelineProgress(InstanceData->InternalTimeAccumulator, AnimSeq->GetPlayLength());
+		}
+
+		return FTimelineProgress();
+	}
+
+	FTimelineProgress FSequencePlayerTrait::SimulateAdvanceBy(FExecutionContext& Context, const TTraitBinding<ITimeline>& Binding, float DeltaTime) const
+	{
+		const FInstanceData* InstanceData = Binding.GetInstanceData<FInstanceData>();
+
+		if (UAnimSequence* AnimSeq = InstanceData->AnimSequence.Get())
+		{
+			const FSharedData* SharedData = Binding.GetSharedData<FSharedData>();
+
+			TTraitBinding<ITimeline> TimelineTrait;
+			Binding.GetStackInterface(TimelineTrait);
+
+			const float PlayRate = TimelineTrait.GetPlayRate(Context);
+			const bool bIsLooping = SharedData->GetbLoop(Binding);
+			const float SequenceLength = AnimSeq->GetPlayLength();
+
+			float Position = InstanceData->InternalTimeAccumulator;
+			FAnimationRuntime::AdvanceTime(bIsLooping, DeltaTime * PlayRate, Position, SequenceLength);
+
+			return FTimelineProgress(Position, SequenceLength);
+		}
+
+		return FTimelineProgress();
+	}
+
+	FTimelineProgress FSequencePlayerTrait::AdvanceBy(FExecutionContext& Context, const TTraitBinding<ITimeline>& Binding, float DeltaTime) const
 	{
 		FInstanceData* InstanceData = Binding.GetInstanceData<FInstanceData>();
 
@@ -70,10 +105,10 @@ namespace UE::AnimNext
 
 			FAnimationRuntime::AdvanceTime(bIsLooping, DeltaTime * PlayRate, InstanceData->InternalTimeAccumulator, SequenceLength);
 
-			return FMath::Clamp(InstanceData->InternalTimeAccumulator / SequenceLength, 0.0f, 1.0f);
+			return FTimelineProgress(InstanceData->InternalTimeAccumulator, SequenceLength);
 		}
 
-		return 0.0f;
+		return FTimelineProgress();
 	}
 
 	void FSequencePlayerTrait::AdvanceToRatio(FExecutionContext& Context, const TTraitBinding<ITimeline>& Binding, float ProgressRatio) const
diff --git a/Engine/Plugins/Experimental/AnimNextAnimGraph/Source/AnimNextAnimGraph/Private/Traits/SequencePlayer.h b/Engine/Plugins/Experimental/AnimNextAnimGraph/Source/AnimNextAnimGraph/Private/Traits/SequencePlayer.h
--- a/Engine/Plugins/Experimental/AnimNextAnimGraph/Source/AnimNextAnimGraph/Private/Traits/SequencePlayer.h
+++ b/Engine/Plugins/Experimental/AnimNextAnimGraph/Source/AnimNextAnimGraph/Private/Traits/SequencePlayer.h
@@ -73,7 +73,9 @@ namespace UE::AnimNext
 
 		// ITimeline impl
 		virtual float GetPlayRate(FExecutionContext& Context, const TTraitBinding<ITimeline>& Binding) const override;
-		virtual float AdvanceBy(FExecutionContext& Context, const TTraitBinding<ITimeline>& Binding, float DeltaTime) const override;
+		virtual FTimelineProgress GetProgress(FExecutionContext& Context, const TTraitBinding<ITimeline>& Binding) const override;
+		virtual FTimelineProgress SimulateAdvanceBy(FExecutionContext& Context, const TTraitBinding<ITimeline>& Binding, float DeltaTime) const override;
+		virtual FTimelineProgress AdvanceBy(FExecutionContext& Context, const TTraitBinding<ITimeline>& Binding, float DeltaTime) const override;
 		virtual void AdvanceToRatio(FExecutionContext& Context, const TTraitBinding<ITimeline>& Binding, float ProgressRatio) const override;
 
 		// IUpdate impl
diff --git a/Engine/Plugins/Experimental/AnimNextAnimGraph/Source/AnimNextAnimGraph/Private/Traits/SynchronizeUsingGroups.cpp b/Engine/Plugins/Experimental/AnimNextAnimGraph/Source/AnimNextAnimGraph/Private/Traits/SynchronizeUsingGroups.cpp
--- a/Engine/Plugins/Experimental/AnimNextAnimGraph/Source/AnimNextAnimGraph/Private/Traits/SynchronizeUsingGroups.cpp
+++ b/Engine/Plugins/Experimental/AnimNextAnimGraph/Source/AnimNextAnimGraph/Private/Traits/SynchronizeUsingGroups.cpp
@@ -58,7 +58,7 @@ namespace UE::AnimNext
 		return SharedData->GroupRole;
 	}
 
-	float FSynchronizeUsingGroupsTrait::AdvanceBy(FExecutionContext& Context, const TTraitBinding<IGroupSynchronization>& Binding, float DeltaTime) const
+	FTimelineProgress FSynchronizeUsingGroupsTrait::AdvanceBy(FExecutionContext& Context, const TTraitBinding<IGroupSynchronization>& Binding, float DeltaTime) const
 	{
 		FInstanceData* InstanceData = Binding.GetInstanceData<FInstanceData>();
 
@@ -68,11 +68,11 @@ namespace UE::AnimNext
 		TTraitBinding<ITimeline> TimelineTrait;
 		Binding.GetStackInterface(TimelineTrait);
 
-		const float ProgressRatio = TimelineTrait.AdvanceBy(Context, DeltaTime);
+		const FTimelineProgress Progress = TimelineTrait.AdvanceBy(Context, DeltaTime);
 
 		InstanceData->bFreezeTimeline = true;
 
-		return ProgressRatio;
+		return Progress;
 	}
 
 	void FSynchronizeUsingGroupsTrait::AdvanceToRatio(FExecutionContext& Context, const TTraitBinding<IGroupSynchronization>& Binding, float ProgressRatio) const
@@ -90,12 +90,23 @@ namespace UE::AnimNext
 		InstanceData->bFreezeTimeline = true;
 	}
 
-	float FSynchronizeUsingGroupsTrait::AdvanceBy(FExecutionContext& Context, const TTraitBinding<ITimeline>& Binding, float DeltaTime) const
+	FTimelineProgress FSynchronizeUsingGroupsTrait::SimulateAdvanceBy(FExecutionContext& Context, const TTraitBinding<ITimeline>& Binding, float DeltaTime) const
 	{
 		const FInstanceData* InstanceData = Binding.GetInstanceData<FInstanceData>();
 		if (InstanceData->bFreezeTimeline)
 		{
-			return 0.0f;	// If the timeline is frozen, we don't advance
+			return ITimeline::GetProgress(Context, Binding);	// If the timeline is frozen, we don't advance, return the previous value
+		}
+
+		return ITimeline::SimulateAdvanceBy(Context, Binding, DeltaTime);
+	}
+
+	FTimelineProgress FSynchronizeUsingGroupsTrait::AdvanceBy(FExecutionContext& Context, const TTraitBinding<ITimeline>& Binding, float DeltaTime) const
+	{
+		const FInstanceData* InstanceData = Binding.GetInstanceData<FInstanceData>();
+		if (InstanceData->bFreezeTimeline)
+		{
+			return ITimeline::GetProgress(Context, Binding);	// If the timeline is frozen, we don't advance, return the previous value
 		}
 
 		return ITimeline::AdvanceBy(Context, Binding, DeltaTime);
diff --git a/Engine/Plugins/Experimental/AnimNextAnimGraph/Source/AnimNextAnimGraph/Private/Traits/SynchronizeUsingGroups.h b/Engine/Plugins/Experimental/AnimNextAnimGraph/Source/AnimNextAnimGraph/Private/Traits/SynchronizeUsingGroups.h
--- a/Engine/Plugins/Experimental/AnimNextAnimGraph/Source/AnimNextAnimGraph/Private/Traits/SynchronizeUsingGroups.h
+++ b/Engine/Plugins/Experimental/AnimNextAnimGraph/Source/AnimNextAnimGraph/Private/Traits/SynchronizeUsingGroups.h
@@ -50,11 +50,12 @@ namespace UE::AnimNext
 		// IGroupSynchronization impl
 		virtual FName GetGroupName(FExecutionContext& Context, const TTraitBinding<IGroupSynchronization>& Binding) const override;
 		virtual EAnimGroupRole::Type GetGroupRole(FExecutionContext& Context, const TTraitBinding<IGroupSynchronization>& Binding) const override;
-		virtual float AdvanceBy(FExecutionContext& Context, const TTraitBinding<IGroupSynchronization>& Binding, float DeltaTime) const override;
+		virtual FTimelineProgress AdvanceBy(FExecutionContext& Context, const TTraitBinding<IGroupSynchronization>& Binding, float DeltaTime) const override;
 		virtual void AdvanceToRatio(FExecutionContext& Context, const TTraitBinding<IGroupSynchronization>& Binding, float ProgressRatio) const override;
 
 		// ITimeline impl
-		virtual float AdvanceBy(FExecutionContext& Context, const TTraitBinding<ITimeline>& Binding, float DeltaTime) const override;
+		virtual FTimelineProgress SimulateAdvanceBy(FExecutionContext& Context, const TTraitBinding<ITimeline>& Binding, float DeltaTime) const override;
+		virtual FTimelineProgress AdvanceBy(FExecutionContext& Context, const TTraitBinding<ITimeline>& Binding, float DeltaTime) const override;
 		virtual void AdvanceToRatio(FExecutionContext& Context, const TTraitBinding<ITimeline>& Binding, float ProgressRatio) const override;
 	};
 }
diff --git a/Engine/Plugins/Experimental/AnimNextAnimGraph/Source/AnimNextAnimGraph/Public/TraitInterfaces/IGroupSynchronization.h b/Engine/Plugins/Experimental/AnimNextAnimGraph/Source/AnimNextAnimGraph/Public/TraitInterfaces/IGroupSynchronization.h
--- a/Engine/Plugins/Experimental/AnimNextAnimGraph/Source/AnimNextAnimGraph/Public/TraitInterfaces/IGroupSynchronization.h
+++ b/Engine/Plugins/Experimental/AnimNextAnimGraph/Source/AnimNextAnimGraph/Public/TraitInterfaces/IGroupSynchronization.h
@@ -6,6 +6,7 @@
 #include "Animation/AnimationAsset.h"
 #include "TraitCore/ITraitInterface.h"
 #include "TraitCore/TraitBinding.h"
+#include "TraitInterfaces/ITimeline.h"
 
 namespace UE::AnimNext
 {
@@ -26,7 +27,7 @@ namespace UE::AnimNext
 
 		// Called by the sync group graph instance component once a group has been synchronized to advance time on the leader
 		// Returns the progress ratio of playback: 0.0 = start of animation, 1.0 = end of animation
-		virtual float AdvanceBy(FExecutionContext& Context, const TTraitBinding<IGroupSynchronization>& Binding, float DeltaTime) const;
+		virtual FTimelineProgress AdvanceBy(FExecutionContext& Context, const TTraitBinding<IGroupSynchronization>& Binding, float DeltaTime) const;
 
 		// Called by the sync group graph instance component once a group has been synchronized to advance time on each follower
 		// Progress ratio must be between [0.0, 1.0]
@@ -52,7 +53,7 @@ namespace UE::AnimNext
 		}
 
 		// @see IGroupSynchronization::AdvanceBy
-		float AdvanceBy(FExecutionContext& Context, float DeltaTime) const
+		FTimelineProgress AdvanceBy(FExecutionContext& Context, float DeltaTime) const
 		{
 			return GetInterface()->AdvanceBy(Context, *this, DeltaTime);
 		}
diff --git a/Engine/Plugins/Experimental/AnimNextAnimGraph/Source/AnimNextAnimGraph/Public/TraitInterfaces/ITimeline.h b/Engine/Plugins/Experimental/AnimNextAnimGraph/Source/AnimNextAnimGraph/Public/TraitInterfaces/ITimeline.h
--- a/Engine/Plugins/Experimental/AnimNextAnimGraph/Source/AnimNextAnimGraph/Public/TraitInterfaces/ITimeline.h
+++ b/Engine/Plugins/Experimental/AnimNextAnimGraph/Source/AnimNextAnimGraph/Public/TraitInterfaces/ITimeline.h
@@ -8,6 +8,49 @@
 
 namespace UE::AnimNext
 {
+	/**
+	 * Timeline Progress
+	 * 
+	 * Encapsulates the progress along a timeline.
+	 */
+	struct FTimelineProgress
+	{
+		// Default construct with no progress
+		FTimelineProgress() = default;
+
+		// Construct with a specific position and duration
+		FTimelineProgress(float InPosition, float InDuration)
+			: Position(InPosition)
+			, Duration(InDuration)
+		{
+		}
+
+		// Resets the timeline progress to its initial state
+		void Reset()
+		{
+			Position = Duration = 0.0f;
+		}
+
+		// Returns the timeline duration in seconds
+		float GetDuration() const { return Duration; }
+
+		// Returns the timeline position in seconds
+		float GetPosition() const { return Position; }
+
+		// Returns the timeline position as a ratio (0.0 = start of timeline, 1.0 = end of timeline)
+		float GetPositionRatio() const { return Duration != 0.0f ? FMath::Clamp(Position / Duration, 0.0f, 1.0f) : 0.0f; }
+
+		// Returns the time left to play in the timeline
+		float GetTimeLeft() const { return Duration - Position; }
+
+	private:
+		// Timeline position in seconds
+		float Position = 0.0f;
+
+		// Timeline duration in seconds
+		float Duration = 0.0f;
+	};
+
 	/**
 	 * ITimeline
 	 *
@@ -20,9 +63,16 @@ namespace UE::AnimNext
 		// Returns the play rate of this timeline
 		virtual float GetPlayRate(FExecutionContext& Context, const TTraitBinding<ITimeline>& Binding) const;
 
+		// Returns the progress of this timeline
+		virtual FTimelineProgress GetProgress(FExecutionContext& Context, const TTraitBinding<ITimeline>& Binding) const;
+
+		// Simulates the advance of time by the provided delta time (positive or negative) on this timeline
+		// Returns the progress of playback
+		virtual FTimelineProgress SimulateAdvanceBy(FExecutionContext& Context, const TTraitBinding<ITimeline>& Binding, float DeltaTime) const;
+
 		// Advances time by the provided delta time (positive or negative) on this timeline
 		// Returns the progress of playback
-		virtual float AdvanceBy(FExecutionContext& Context, const TTraitBinding<ITimeline>& Binding, float DeltaTime) const;
+		virtual FTimelineProgress AdvanceBy(FExecutionContext& Context, const TTraitBinding<ITimeline>& Binding, float DeltaTime) const;
 
 		// Advances time to the specified progress ratio on this timeline
 		// Progress ratio must be between [0.0, 1.0]
@@ -41,8 +91,20 @@ namespace UE::AnimNext
 			return GetInterface()->GetPlayRate(Context, *this);
 		}
 
+		// @see ITimeline::GetProgress
+		FTimelineProgress GetProgress(FExecutionContext& Context) const
+		{
+			return GetInterface()->GetProgress(Context, *this);
+		}
+
+		// @see ITimeline::SimulateAdvanceBy
+		FTimelineProgress SimulateAdvanceBy(FExecutionContext& Context, float DeltaTime) const
+		{
+			return GetInterface()->SimulateAdvanceBy(Context, *this, DeltaTime);
+		}
+
 		// @see ITimeline::AdvanceBy
-		float AdvanceBy(FExecutionContext& Context, float DeltaTime) const
+		FTimelineProgress AdvanceBy(FExecutionContext& Context, float DeltaTime) const
 		{
 			return GetInterface()->AdvanceBy(Context, *this, DeltaTime);
 		}
