Implement PlayAnim API for AnimNext

diff --git a/Engine/Plugins/Experimental/AnimNextAnimGraph/AnimNextAnimGraph.uplugin b/Engine/Plugins/Experimental/AnimNextAnimGraph/AnimNextAnimGraph.uplugin
--- a/Engine/Plugins/Experimental/AnimNextAnimGraph/AnimNextAnimGraph.uplugin
+++ b/Engine/Plugins/Experimental/AnimNextAnimGraph/AnimNextAnimGraph.uplugin
@@ -20,6 +20,10 @@
 		{
 			"Name": "AnimNext",
 			"Enabled": true
+		},
+		{
+			"Name": "Chooser",
+			"Enabled": true
 		}
 	],
 	"Modules" :
diff --git a/Engine/Plugins/Experimental/AnimNextAnimGraph/Source/AnimNextAnimGraph/AnimNextAnimGraph.Build.cs b/Engine/Plugins/Experimental/AnimNextAnimGraph/Source/AnimNextAnimGraph/AnimNextAnimGraph.Build.cs
--- a/Engine/Plugins/Experimental/AnimNextAnimGraph/Source/AnimNextAnimGraph/AnimNextAnimGraph.Build.cs
+++ b/Engine/Plugins/Experimental/AnimNextAnimGraph/Source/AnimNextAnimGraph/AnimNextAnimGraph.Build.cs
@@ -14,6 +14,8 @@ namespace UnrealBuildTool.Rules
 					"RigVM",
 					"Engine",
 					"AnimNext",
+					"Chooser",
+					"StructUtils",
 				}
 			);
 		}
diff --git a/Engine/Plugins/Experimental/AnimNextAnimGraph/Source/AnimNextAnimGraph/Internal/PlayAnim/PlayAnimCallbackProxy.h b/Engine/Plugins/Experimental/AnimNextAnimGraph/Source/AnimNextAnimGraph/Internal/PlayAnim/PlayAnimCallbackProxy.h
new file mode 100644
--- /dev/null
+++ b/Engine/Plugins/Experimental/AnimNextAnimGraph/Source/AnimNextAnimGraph/Internal/PlayAnim/PlayAnimCallbackProxy.h
@@ -0,0 +1,64 @@
+// Copyright Epic Games, Inc. All Rights Reserved.
+#pragma once
+
+#include "CoreMinimal.h"
+#include "PlayAnim/PlayAnimRequest.h"
+
+#include "PlayAnimCallbackProxy.generated.h"
+
+DECLARE_DYNAMIC_MULTICAST_DELEGATE(FOnPlayAnimPlayDelegate);
+
+UCLASS(MinimalAPI)
+class UPlayAnimCallbackProxy : public UObject
+{
+	GENERATED_UCLASS_BODY()
+
+	// Called when the provided animation object finished playing and hasn't been interrupted
+	UPROPERTY(BlueprintAssignable)
+	FOnPlayAnimPlayDelegate OnCompleted;
+
+	// Called when the provided animation object starts blending out and hasn't been interrupted
+	UPROPERTY(BlueprintAssignable)
+	FOnPlayAnimPlayDelegate OnBlendOut;
+
+	// Called when the provided animation object has been interrupted (or failed to play)
+	UPROPERTY(BlueprintAssignable)
+	FOnPlayAnimPlayDelegate OnInterrupted;
+
+	// Called to perform the query internally
+	UFUNCTION(BlueprintCallable, meta = (BlueprintInternalUseOnly = "true"))
+	static ANIMNEXTANIMGRAPH_API UPlayAnimCallbackProxy* CreateProxyObjectForPlayAnim(
+		class UAnimNextComponent* AnimNextComponent,
+		FName SlotName,
+		class UAnimSequence* AnimationObject,
+		float PlayRate = 1.0f,
+		float StartPosition = 0.0f,
+		FAnimNextPlayAnimBlendSettings BlendInSettings = FAnimNextPlayAnimBlendSettings(),
+		FAnimNextPlayAnimBlendSettings BlendOutSettings = FAnimNextPlayAnimBlendSettings());
+
+public:
+	//~ Begin UObject Interface
+	ANIMNEXTANIMGRAPH_API virtual void BeginDestroy() override;
+	//~ End UObject Interface
+
+protected:
+	ANIMNEXTANIMGRAPH_API void OnPlayAnimCompleted(const UE::AnimNext::FPlayAnimRequest& Request);
+	ANIMNEXTANIMGRAPH_API void OnPlayAnimInterrupted(const UE::AnimNext::FPlayAnimRequest& Request);
+	ANIMNEXTANIMGRAPH_API void OnPlayAnimBlendingOut(const UE::AnimNext::FPlayAnimRequest& Request);
+
+	// Attempts to play a montage with the specified settings. Returns whether it started or not.
+	ANIMNEXTANIMGRAPH_API bool Play(
+		class UAnimNextComponent* AnimNextComponent,
+		FName SlotName,
+		class UAnimSequence* AnimationObject,
+		float PlayRate,
+		float StartPosition,
+		const UE::AnimNext::FPlayAnimBlendSettings& BlendInSettings,
+		const UE::AnimNext::FPlayAnimBlendSettings& BlendOutSettings);
+
+private:
+	void Reset();
+
+	UE::AnimNext::FPlayAnimRequestPtr PlayingRequest;
+	bool bWasInterrupted = false;
+};
diff --git a/Engine/Plugins/Experimental/AnimNextAnimGraph/Source/AnimNextAnimGraph/Private/Module.cpp b/Engine/Plugins/Experimental/AnimNextAnimGraph/Source/AnimNextAnimGraph/Private/Module.cpp
--- a/Engine/Plugins/Experimental/AnimNextAnimGraph/Source/AnimNextAnimGraph/Private/Module.cpp
+++ b/Engine/Plugins/Experimental/AnimNextAnimGraph/Source/AnimNextAnimGraph/Private/Module.cpp
@@ -8,6 +8,7 @@
 #include "Modules/ModuleManager.h"
 #include "RigVMCore/RigVMRegistry.h"
 #include "Animation/AnimSequence.h"
+#include "Chooser.h"
 
 #include "IAnimNextModule.h"
 #include "TraitInterfaces/IEvaluate.h"
@@ -29,6 +30,7 @@ public:
 			{ UBlendProfile::StaticClass(), FRigVMRegistry::ERegisterObjectOperation::Class },
 			{ UCurveFloat::StaticClass(), FRigVMRegistry::ERegisterObjectOperation::Class },
 			{ UAnimNextGraph::StaticClass(), FRigVMRegistry::ERegisterObjectOperation::Class },
+			{ UChooserTable::StaticClass(), FRigVMRegistry::ERegisterObjectOperation::Class },
 		};
 
 		FRigVMRegistry::Get().RegisterObjectTypes(AllowedObjectTypes);
diff --git a/Engine/Plugins/Experimental/AnimNextAnimGraph/Source/AnimNextAnimGraph/Private/PlayAnim/PlayAnimCallbackProxy.cpp b/Engine/Plugins/Experimental/AnimNextAnimGraph/Source/AnimNextAnimGraph/Private/PlayAnim/PlayAnimCallbackProxy.cpp
new file mode 100644
--- /dev/null
+++ b/Engine/Plugins/Experimental/AnimNextAnimGraph/Source/AnimNextAnimGraph/Private/PlayAnim/PlayAnimCallbackProxy.cpp
@@ -0,0 +1,116 @@
+// Copyright Epic Games, Inc. All Rights Reserved.
+
+#include "PlayAnim/PlayAnimCallbackProxy.h"
+
+#include UE_INLINE_GENERATED_CPP_BY_NAME(PlayAnimCallbackProxy)
+
+//////////////////////////////////////////////////////////////////////////
+// UPlayAnimCallbackProxy
+
+UPlayAnimCallbackProxy::UPlayAnimCallbackProxy(const FObjectInitializer& ObjectInitializer)
+	: Super(ObjectInitializer)
+{
+}
+
+UPlayAnimCallbackProxy* UPlayAnimCallbackProxy::CreateProxyObjectForPlayAnim(
+	UAnimNextComponent* AnimNextComponent,
+	FName SlotName,
+	UAnimSequence* AnimationObject,
+	float PlayRate,
+	float StartPosition,
+	UE::AnimNext::FPlayAnimBlendSettings BlendInSettings,
+	UE::AnimNext::FPlayAnimBlendSettings BlendOutSettings)
+{
+	UPlayAnimCallbackProxy* Proxy = NewObject<UPlayAnimCallbackProxy>();
+	Proxy->SetFlags(RF_StrongRefOnFrame);
+	Proxy->Play(AnimNextComponent, SlotName, AnimationObject, PlayRate, StartPosition, BlendInSettings, BlendOutSettings);
+	return Proxy;
+}
+
+bool UPlayAnimCallbackProxy::Play(
+	UAnimNextComponent* AnimNextComponent,
+	FName SlotName,
+	UAnimSequence* AnimationObject,
+	float PlayRate,
+	float StartPosition,
+	const UE::AnimNext::FPlayAnimBlendSettings& BlendInSettings,
+	const UE::AnimNext::FPlayAnimBlendSettings& BlendOutSettings)
+{
+	bool bPlayedSuccessfully = false;
+	if (AnimNextComponent != nullptr)
+	{
+		UE::AnimNext::FPlayAnimRequestArgs RequestArgs;
+		RequestArgs.SlotName = SlotName;
+		RequestArgs.AnimationObject = AnimationObject;
+		RequestArgs.PlayRate = PlayRate;
+		RequestArgs.StartPosition = StartPosition;
+		RequestArgs.BlendInSettings = BlendInSettings;
+		RequestArgs.BlendOutSettings = BlendOutSettings;
+
+		auto Request = UE::AnimNext::MakePlayAnimRequest();
+		Request->OnCompleted.BindUObject(this, &UPlayAnimCallbackProxy::OnPlayAnimCompleted);
+		Request->OnInterrupted.BindUObject(this, &UPlayAnimCallbackProxy::OnPlayAnimInterrupted);
+		Request->OnBlendingOut.BindUObject(this, &UPlayAnimCallbackProxy::OnPlayAnimBlendingOut);
+
+		bPlayedSuccessfully = Request->Play(RequestArgs, AnimNextComponent);
+
+		PlayingRequest = Request;
+		bWasInterrupted = false;
+	}
+
+	if (!bPlayedSuccessfully)
+	{
+		OnInterrupted.Broadcast();
+		Reset();
+	}
+
+	return bPlayedSuccessfully;
+}
+
+void UPlayAnimCallbackProxy::OnPlayAnimCompleted(const UE::AnimNext::FPlayAnimRequest& Request)
+{
+	if (!bWasInterrupted)
+	{
+		const UE::AnimNext::EPlayAnimStatus Status = Request.GetStatus();
+		check(!EnumHasAnyFlags(Status, UE::AnimNext::EPlayAnimStatus::Interrupted));
+
+		if (EnumHasAnyFlags(Status, UE::AnimNext::EPlayAnimStatus::Expired))
+		{
+			OnInterrupted.Broadcast();
+		}
+		else
+		{
+			OnCompleted.Broadcast();
+		}
+	}
+
+	Reset();
+}
+
+void UPlayAnimCallbackProxy::OnPlayAnimInterrupted(const UE::AnimNext::FPlayAnimRequest& Request)
+{
+	bWasInterrupted = true;
+
+	OnInterrupted.Broadcast();
+}
+
+void UPlayAnimCallbackProxy::OnPlayAnimBlendingOut(const UE::AnimNext::FPlayAnimRequest& Request)
+{
+	if (!bWasInterrupted)
+	{
+		OnBlendOut.Broadcast();
+	}
+}
+
+void UPlayAnimCallbackProxy::Reset()
+{
+	PlayingRequest = nullptr;
+	bWasInterrupted = false;
+}
+
+void UPlayAnimCallbackProxy::BeginDestroy()
+{
+	Reset();
+
+	Super::BeginDestroy();
+}
diff --git a/Engine/Plugins/Experimental/AnimNextAnimGraph/Source/AnimNextAnimGraph/Private/PlayAnim/PlayAnimEvents.cpp b/Engine/Plugins/Experimental/AnimNextAnimGraph/Source/AnimNextAnimGraph/Private/PlayAnim/PlayAnimEvents.cpp
new file mode 100644
--- /dev/null
+++ b/Engine/Plugins/Experimental/AnimNextAnimGraph/Source/AnimNextAnimGraph/Private/PlayAnim/PlayAnimEvents.cpp
@@ -0,0 +1,29 @@
+// Copyright Epic Games, Inc. All Rights Reserved.
+
+#include "PlayAnim/PlayAnimEvents.h"
+
+#include "TraitCore/TraitEventList.h"
+
+namespace UE::AnimNext
+{
+	void FPlayAnim_PlayEvent::OnExpired(FTraitEventList& OutputEventList)
+	{
+		auto ActionEvent = MakeTraitEvent<FAnimNextSchedule_ActionEvent>();
+		ActionEvent->ActionFunction = [Request = Request]()
+			{
+				Request->OnStatusUpdate(EPlayAnimStatus::Expired);
+			};
+
+		OutputEventList.Push(MoveTemp(ActionEvent));
+	}
+
+	void FPlayAnim_StatusUpdateEvent::Execute() const
+	{
+		Request->OnStatusUpdate(Status);
+	}
+
+	void FPlayAnim_TimelineUpdateEvent::Execute() const
+	{
+		Request->OnTimelineUpdate(TimelineProgress);
+	}
+}
diff --git a/Engine/Plugins/Experimental/AnimNextAnimGraph/Source/AnimNextAnimGraph/Private/PlayAnim/PlayAnimEvents.h b/Engine/Plugins/Experimental/AnimNextAnimGraph/Source/AnimNextAnimGraph/Private/PlayAnim/PlayAnimEvents.h
new file mode 100644
--- /dev/null
+++ b/Engine/Plugins/Experimental/AnimNextAnimGraph/Source/AnimNextAnimGraph/Private/PlayAnim/PlayAnimEvents.h
@@ -0,0 +1,84 @@
+// Copyright Epic Games, Inc. All Rights Reserved.
+
+#pragma once
+
+#include "CoreMinimal.h"
+
+#include "PlayAnim/PlayAnimRequest.h"
+#include "PlayAnim/PlayAnimStatus.h"
+#include "Scheduler/ScheduleEvents.h"
+#include "TraitCore/TraitEvent.h"
+#include "TraitInterfaces/ITimeline.h"
+
+namespace UE::AnimNext
+{
+	/**
+	 * PlayAnim Play Event
+	 *
+	 * Event raised when a Play request is made.
+	 * It encapsulates everything needed to service an animation request.
+	 *
+	 * If no sub-graph is provided, this event will request that the input source plays instead.
+	 */
+	struct FPlayAnim_PlayEvent : public FAnimNextTraitEvent
+	{
+		DECLARE_ANIM_TRAIT_EVENT(FPlayAnim_PlayEvent, FAnimNextTraitEvent)
+
+		FPlayAnimRequestPtr Request;
+
+	protected:
+		// FAnimNextTraitEvent impl
+		virtual void OnExpired(FTraitEventList& OutputEventList) override;
+	};
+
+	/**
+	 * PlayAnim Stop Event
+	 *
+	 * Event raised when a Stop request is made.
+	 * It encapsulates everything needed to service an animation request.
+	 */
+	struct FPlayAnim_StopEvent : public FAnimNextTraitEvent
+	{
+		DECLARE_ANIM_TRAIT_EVENT(FPlayAnim_StopEvent, FAnimNextTraitEvent)
+
+		FPlayAnimRequestPtr Request;
+	};
+
+	/**
+	 * PlayAnim Status Update
+	 *
+	 * Event raised when the status of a request changes.
+	 */
+	struct FPlayAnim_StatusUpdateEvent : public FAnimNextSchedule_ActionEvent
+	{
+		DECLARE_ANIM_TRAIT_EVENT(FPlayAnim_StatusUpdateEvent, FAnimNextSchedule_ActionEvent)
+
+		// FAnimNextSchedule_ActionEvent impl
+		virtual void Execute() const override;
+
+		// The request to update
+		FPlayAnimRequestPtr Request;
+
+		// The current request status
+		EPlayAnimStatus Status = EPlayAnimStatus::None;
+	};
+
+	/**
+	 * PlayAnim Timeline Update
+	 *
+	 * Event raised when a request is playing with its updated timeline progress.
+	 */
+	struct FPlayAnim_TimelineUpdateEvent : public FAnimNextSchedule_ActionEvent
+	{
+		DECLARE_ANIM_TRAIT_EVENT(FPlayAnim_TimelineUpdateEvent, FAnimNextSchedule_ActionEvent)
+
+		// FAnimNextSchedule_ActionEvent impl
+		virtual void Execute() const override;
+
+		// The request to update
+		FPlayAnimRequestPtr Request;
+
+		// The current request timeline progress
+		FTimelineProgress TimelineProgress;
+	};
+}
diff --git a/Engine/Plugins/Experimental/AnimNextAnimGraph/Source/AnimNextAnimGraph/Private/PlayAnim/PlayAnimRequest.cpp b/Engine/Plugins/Experimental/AnimNextAnimGraph/Source/AnimNextAnimGraph/Private/PlayAnim/PlayAnimRequest.cpp
new file mode 100644
--- /dev/null
+++ b/Engine/Plugins/Experimental/AnimNextAnimGraph/Source/AnimNextAnimGraph/Private/PlayAnim/PlayAnimRequest.cpp
@@ -0,0 +1,160 @@
+// Copyright Epic Games, Inc. All Rights Reserved.
+
+#include "PlayAnim/PlayAnimRequest.h"
+
+#include "PlayAnim/PlayAnimEvents.h"
+#include "Component/AnimNextComponent.h"
+
+namespace UE::AnimNext
+{
+	bool FPlayAnimRequest::Play(const FPlayAnimRequestArgs& InRequestArgs, UAnimNextComponent* InComponent)
+	{
+		check(IsInGameThread());
+		if (InRequestArgs.AnimationObject == nullptr || InComponent == nullptr)
+		{
+			return false;	// Nothing to play
+		}
+
+		if (Status != EPlayAnimStatus::None)
+		{
+			return false;	// Already playing, cannot play again
+		}
+
+		RequestArgs = InRequestArgs;
+		Component = InComponent;
+		Status = EPlayAnimStatus::Pending;
+
+		auto StartEvent = MakeTraitEvent<FPlayAnim_PlayEvent>();
+		StartEvent->Request = AsShared();
+
+		Component->QueueInputTraitEvent(StartEvent);
+
+		PendingStartEvent = StartEvent;
+
+		return true;
+	}
+
+	void FPlayAnimRequest::Stop()
+	{
+		check(IsInGameThread());
+		if (!EnumHasAnyFlags(Status, EPlayAnimStatus::Playing))
+		{
+			return;	// Not playing
+		}
+		else if (EnumHasAnyFlags(Status, EPlayAnimStatus::Interrupted))
+		{
+			return;	// We already got interrupted
+		}
+
+		// If we are still pending, just cancel our event
+		if (PendingStartEvent)
+		{
+			PendingStartEvent->MarkConsumed();
+			PendingStartEvent = nullptr;
+		}
+		else
+		{
+			auto StopEvent = MakeTraitEvent<FPlayAnim_StopEvent>();
+			StopEvent->Request = AsShared();
+
+			Component->QueueInputTraitEvent(StopEvent);
+		}
+	}
+
+	const FPlayAnimRequestArgs& FPlayAnimRequest::GetArgs() const
+	{
+		return RequestArgs;
+	}
+
+	EPlayAnimStatus FPlayAnimRequest::GetStatus() const
+	{
+		check(IsInGameThread());
+		return Status;
+	}
+
+	FTimelineProgress FPlayAnimRequest::GetTimelineProgress() const
+	{
+		check(IsInGameThread());
+		return TimelineProgress;
+	}
+
+	bool FPlayAnimRequest::HasExpired() const
+	{
+		check(IsInGameThread());
+		return EnumHasAnyFlags(Status, EPlayAnimStatus::Expired);
+	}
+
+	bool FPlayAnimRequest::HasCompleted() const
+	{
+		check(IsInGameThread());
+		return EnumHasAnyFlags(Status, EPlayAnimStatus::Completed);
+	}
+
+	bool FPlayAnimRequest::IsPlaying() const
+	{
+		check(IsInGameThread());
+		return EnumHasAnyFlags(Status, EPlayAnimStatus::Playing);
+	}
+
+	bool FPlayAnimRequest::IsBlendingOut() const
+	{
+		check(IsInGameThread());
+		return EnumHasAnyFlags(Status, EPlayAnimStatus::BlendingOut);
+	}
+
+	bool FPlayAnimRequest::WasInterrupted() const
+	{
+		check(IsInGameThread());
+		return EnumHasAnyFlags(Status, EPlayAnimStatus::Interrupted);
+	}
+
+	void FPlayAnimRequest::OnStatusUpdate(EPlayAnimStatus NewStatus)
+	{
+		check(IsInGameThread());
+		switch (NewStatus)
+		{
+		case EPlayAnimStatus::Playing:
+			ensureMsgf(Status == EPlayAnimStatus::Pending, TEXT("Expected PlayAnim status to be pending, found: %u"), (uint32)Status);
+			Status = NewStatus;
+			PendingStartEvent = nullptr;
+			OnStarted.ExecuteIfBound(*this);
+			break;
+		case EPlayAnimStatus::Playing | EPlayAnimStatus::Interrupted:
+			ensureMsgf(EnumHasAnyFlags(Status, EPlayAnimStatus::Playing), TEXT("Expected PlayAnim status to be playing, found: %u"), (uint32)Status);
+			EnumAddFlags(Status, EPlayAnimStatus::Interrupted);
+			OnInterrupted.ExecuteIfBound(*this);
+			break;
+		case EPlayAnimStatus::BlendingOut:
+			ensureMsgf(EnumHasAnyFlags(Status, EPlayAnimStatus::Playing), TEXT("Expected PlayAnim status to be playing, found: %u"), (uint32)Status);
+			EnumAddFlags(Status, EPlayAnimStatus::BlendingOut);
+			OnBlendingOut.ExecuteIfBound(*this);
+			break;
+		case EPlayAnimStatus::Completed:
+			ensureMsgf(EnumHasAnyFlags(Status, EPlayAnimStatus::Playing), TEXT("Expected PlayAnim status to be playing, found: %u"), (uint32)Status);
+			// Maintain our interrupted status if it was present
+			Status = EPlayAnimStatus::Completed | (Status & EPlayAnimStatus::Interrupted);
+			OnCompleted.ExecuteIfBound(*this);
+			break;
+		case EPlayAnimStatus::Expired:
+			ensureMsgf(Status == EPlayAnimStatus::Pending, TEXT("Expected PlayAnim status to be pending, found: %u"), (uint32)Status);
+			Status = NewStatus;
+			OnCompleted.ExecuteIfBound(*this);
+			break;
+		default:
+			ensureMsgf(false, TEXT("Unsupported PlayAnim status update value: %u"), (uint32)NewStatus);
+			break;
+		}
+	}
+
+	void FPlayAnimRequest::OnTimelineUpdate(FTimelineProgress NewTimelineProgress)
+	{
+		check(IsInGameThread());
+		TimelineProgress = NewTimelineProgress;
+	}
+
+	void FPlayAnimRequest::AddReferencedObjects(FReferenceCollector& Collector)
+	{
+		Collector.AddReferencedObject(RequestArgs.AnimationObject);
+		Collector.AddReferencedObject(Component);
+	}
+}
diff --git a/Engine/Plugins/Experimental/AnimNextAnimGraph/Source/AnimNextAnimGraph/Private/PlayAnim/PlayAnimSlotTrait.cpp b/Engine/Plugins/Experimental/AnimNextAnimGraph/Source/AnimNextAnimGraph/Private/PlayAnim/PlayAnimSlotTrait.cpp
new file mode 100644
--- /dev/null
+++ b/Engine/Plugins/Experimental/AnimNextAnimGraph/Source/AnimNextAnimGraph/Private/PlayAnim/PlayAnimSlotTrait.cpp
@@ -0,0 +1,565 @@
+// Copyright Epic Games, Inc. All Rights Reserved.
+
+#include "PlayAnim/PlayAnimSlotTrait.h"
+
+#include "Animation/AnimSequence.h"
+#include "TraitCore/ExecutionContext.h"
+#include "TraitCore/NodeInstance.h"
+#include "TraitInterfaces/ITimeline.h"
+#include "Graph/AnimNextGraphInstance.h"
+#include "Scheduler/ScheduleEvents.h"
+
+namespace UE::AnimNext
+{
+	AUTO_REGISTER_ANIM_TRAIT(FPlayAnimSlotTrait)
+
+	// Trait implementation boilerplate
+	#define TRAIT_INTERFACE_ENUMERATOR(GeneratorMacro) \
+		GeneratorMacro(IDiscreteBlend) \
+		GeneratorMacro(IGarbageCollection) \
+		GeneratorMacro(IHierarchy) \
+		GeneratorMacro(ISmoothBlend) \
+		GeneratorMacro(IInertializerBlend) \
+		GeneratorMacro(IUpdate) \
+		GeneratorMacro(IUpdateTraversal) \
+
+	#define TRAIT_EVENT_ENUMERATOR(GeneratorMacro) \
+		GeneratorMacro(FPlayAnimSlotTrait::OnPlayEvent) \
+		GeneratorMacro(FPlayAnimSlotTrait::OnStopEvent) \
+
+	GENERATE_ANIM_TRAIT_IMPLEMENTATION(FPlayAnimSlotTrait, TRAIT_INTERFACE_ENUMERATOR, NULL_ANIM_TRAIT_INTERFACE_ENUMERATOR, TRAIT_EVENT_ENUMERATOR)
+	#undef TRAIT_INTERFACE_ENUMERATOR
+	#undef TRAIT_EVENT_ENUMERATOR
+
+	void FPlayAnimSlotTrait::FPlayAnimSlotRequest::Initialize(FPlayAnimRequestPtr InRequest, const FPlayAnimBlendSettings& InBlendSettings, const UAnimNextGraph* InSubGraph)
+	{
+		Request = InRequest;
+		BlendSettings = InBlendSettings;
+		SubGraph = InSubGraph;
+
+		// If no input is provided, we'll use the source
+		State = InSubGraph != nullptr ? EPlayAnimRequestState::Active : EPlayAnimRequestState::ActiveSource;
+		bWasRelevant = false;
+	}
+
+	void FPlayAnimSlotTrait::FInstanceData::Construct(const FExecutionContext& Context, const FTraitBinding& Binding)
+	{
+		FTrait::FInstanceData::Construct(Context, Binding);
+
+		IGarbageCollection::RegisterWithGC(Context, Binding);
+	}
+
+	void FPlayAnimSlotTrait::FInstanceData::Destruct(const FExecutionContext& Context, const FTraitBinding& Binding)
+	{
+		FTrait::FInstanceData::Destruct(Context, Binding);
+
+		IGarbageCollection::UnregisterWithGC(Context, Binding);
+	}
+
+	int32 FPlayAnimSlotTrait::FindFreeRequestIndexOrAdd(FInstanceData& InstanceData)
+	{
+		// Find an empty request we can use
+		const int32 NumRequests = InstanceData.SlotRequests.Num();
+		for (int32 RequestIndex = 0; RequestIndex < NumRequests; ++RequestIndex)
+		{
+			if (InstanceData.SlotRequests[RequestIndex].State == EPlayAnimRequestState::Inactive)
+			{
+				// This request is inactive, we can re-use it
+				return RequestIndex;
+			}
+		}
+
+		// All requests are in use, add a new one
+		return InstanceData.SlotRequests.AddDefaulted();
+	}
+
+	ETraitStackPropagation FPlayAnimSlotTrait::OnPlayEvent(const FExecutionContext& Context, FTraitBinding& Binding, FPlayAnim_PlayEvent& Event) const
+	{
+		const FSharedData* SharedData = Binding.GetSharedData<FSharedData>();
+		const FName SlotName = SharedData->GetSlotName(Binding);
+
+		const FAnimNextPlayAnimRequestArgs& RequestArgs = Event.Request->GetArgs();
+		if (SlotName == RequestArgs.SlotName)
+		{
+			FInstanceData* InstanceData = Binding.GetInstanceData<FInstanceData>();
+
+			ensureMsgf(!InstanceData->PendingRequest.IsValid(), TEXT("PlayAnim slot %s already contained a pending request, it will be overwritten"), *SlotName.ToString());
+
+			// Overwrite any request we might have, we'll pick it up on the next update
+			InstanceData->PendingRequest.Reset();
+			InstanceData->PendingRequest.Request = Event.Request;
+
+			Event.MarkConsumed();
+		}
+
+		return ETraitStackPropagation::Continue;
+	}
+
+	ETraitStackPropagation FPlayAnimSlotTrait::OnStopEvent(const FExecutionContext& Context, FTraitBinding& Binding, FPlayAnim_StopEvent& Event) const
+	{
+		const FSharedData* SharedData = Binding.GetSharedData<FSharedData>();
+		const FName SlotName = SharedData->GetSlotName(Binding);
+
+		const FAnimNextPlayAnimRequestArgs& RequestArgs = Event.Request->GetArgs();
+		if (SlotName == RequestArgs.SlotName)
+		{
+			FInstanceData* InstanceData = Binding.GetInstanceData<FInstanceData>();
+
+			// Reset any pending request we might have, and cancel it
+			InstanceData->PendingRequest.Reset();
+			InstanceData->PendingRequest.bStop = true;
+
+			Event.MarkConsumed();
+		}
+
+		return ETraitStackPropagation::Continue;
+	}
+
+	uint32 FPlayAnimSlotTrait::GetNumChildren(const FExecutionContext& Context, const TTraitBinding<IHierarchy>& Binding) const
+	{
+		const FInstanceData* InstanceData = Binding.GetInstanceData<FInstanceData>();
+		return InstanceData->SlotRequests.Num();
+	}
+
+	void FPlayAnimSlotTrait::GetChildren(const FExecutionContext& Context, const TTraitBinding<IHierarchy>& Binding, FChildrenArray& Children) const
+	{
+		const FInstanceData* InstanceData = Binding.GetInstanceData<FInstanceData>();
+
+		for (const FPlayAnimSlotRequest& SlotRequest : InstanceData->SlotRequests)
+		{
+			// Even if the request is inactive, we queue an empty handle
+			Children.Add(SlotRequest.ChildPtr);
+		}
+	}
+
+	void FPlayAnimSlotTrait::PreUpdate(FUpdateTraversalContext& Context, const TTraitBinding<IUpdate>& Binding, const FTraitUpdateState& TraitState) const
+	{
+		const FSharedData* SharedData = Binding.GetSharedData<FSharedData>();
+		FInstanceData* InstanceData = Binding.GetInstanceData<FInstanceData>();
+
+		TTraitBinding<IDiscreteBlend> DiscreteBlendTrait;
+		Binding.GetStackInterface(DiscreteBlendTrait);
+
+		const bool bHasActiveSubGraph = InstanceData->CurrentlyActiveRequestIndex != INDEX_NONE;
+		if (bHasActiveSubGraph)
+		{
+			InstanceData->SlotRequests[InstanceData->CurrentlyActiveRequestIndex].bWasRelevant = true;
+		}
+
+		bool bJustTransitioned = false;
+		if (InstanceData->PendingRequest.IsValid() || !bHasActiveSubGraph)
+		{
+			const FPlayAnimRequestPtr Request = InstanceData->PendingRequest.Request;
+
+			// Clear it out now in case we early out below
+			InstanceData->PendingRequest.Reset();
+
+			FPlayAnimBlendSettings BlendSettings;
+			const UAnimNextGraph* SubGraph = nullptr;
+			if (Request)
+			{
+				// This is a new pending request, lookup the sub-graph to use with our chooser and the desired animation object
+				const FAnimNextPlayAnimRequestArgs& RequestArgs = Request->GetArgs();
+				if (RequestArgs.AnimationObject)
+				{
+					if (const UChooserTable* Chooser = SharedData->GetSubGraphChooser(Binding))
+					{
+						FAnimNextPlayAnimChooserParameters ChooserParameters;
+						ChooserParameters.AnimationObjectType = RequestArgs.AnimationObject->GetClass();
+
+						FChooserEvaluationContext ChooserContext;
+						ChooserContext.AddStructParam(ChooserParameters);
+
+						UChooserTable::EvaluateChooser(ChooserContext, Chooser, FObjectChooserBase::FObjectChooserIteratorCallback::CreateLambda([&SubGraph](UObject* InResult)
+							{
+								SubGraph = Cast<UAnimNextGraph>(InResult);
+								return FObjectChooserBase::EIteratorStatus::Stop;
+							}));
+					}
+
+					if (SubGraph != nullptr)
+					{
+						// Check for re-entrancy and early-out if we are linking back to the current instance or one of its parents
+						const FName EntryPoint = SubGraph->GetDefaultEntryPoint();
+						const FAnimNextGraphInstance* OwnerGraphInstance = &Binding.GetTraitPtr().GetNodeInstance()->GetOwner();
+						while (OwnerGraphInstance != nullptr)
+						{
+							if (OwnerGraphInstance->UsesGraph(SubGraph) && OwnerGraphInstance->UsesEntryPoint(EntryPoint))
+							{
+								return;
+							}
+
+							OwnerGraphInstance = OwnerGraphInstance->GetParentGraphInstance();
+						}
+					}
+
+					BlendSettings = RequestArgs.BlendInSettings;
+				}
+			}
+
+			if (bHasActiveSubGraph)
+			{
+				// Queue our status update
+				const FPlayAnimSlotRequest& OldSlotRequest = InstanceData->SlotRequests[InstanceData->CurrentlyActiveRequestIndex];
+				if (OldSlotRequest.State == EPlayAnimRequestState::Active)
+				{
+					auto StatusUpdateEvent = MakeTraitEvent<FPlayAnim_StatusUpdateEvent>();
+					StatusUpdateEvent->Request = OldSlotRequest.Request;
+					StatusUpdateEvent->Status = EPlayAnimStatus::Playing | EPlayAnimStatus::Interrupted;
+
+					Context.RaiseOutputTraitEvent(StatusUpdateEvent);
+				}
+			}
+
+			// Find an empty request we can use
+			const int32 FreeRequestIndex = FindFreeRequestIndexOrAdd(*InstanceData);
+
+			FPlayAnimSlotRequest& SlotRequest = InstanceData->SlotRequests[FreeRequestIndex];
+			SlotRequest.Initialize(Request, BlendSettings, SubGraph);
+
+			const int32 OldChildIndex = InstanceData->CurrentlyActiveRequestIndex;
+			const int32 NewChildIndex = FreeRequestIndex;
+
+			InstanceData->CurrentlyActiveRequestIndex = FreeRequestIndex;
+
+			DiscreteBlendTrait.OnBlendTransition(Context, OldChildIndex, NewChildIndex);
+
+			bJustTransitioned = true;
+		}
+
+		float CurrentRequestTimeLeft = 0.0f;
+
+		// Broadcast our timeline progress
+		const int32 NumSlotRequests = InstanceData->SlotRequests.Num();
+		for (int32 RequestIndex = 0; RequestIndex < NumSlotRequests; ++RequestIndex)
+		{
+			const FPlayAnimSlotRequest& SlotRequest = InstanceData->SlotRequests[RequestIndex];
+			if (SlotRequest.State != EPlayAnimRequestState::Active)
+			{
+				continue;	// We don't care about this slot request
+			}
+
+			FTraitStackBinding ChildStack;
+			ensure(Context.GetStack(SlotRequest.ChildPtr, ChildStack));
+
+			TTraitBinding<ITimeline> ChildTimelineTrait;
+			ensure(ChildStack.GetInterface(ChildTimelineTrait));
+
+			const FTimelineProgress ChildProgress = ChildTimelineTrait.SimulateAdvanceBy(Context, TraitState.GetDeltaTime());
+
+			if (InstanceData->CurrentlyActiveRequestIndex == RequestIndex)
+			{
+				CurrentRequestTimeLeft = ChildProgress.GetTimeLeft();
+			}
+
+			{
+				auto TimelineUpdateEvent = MakeTraitEvent<FPlayAnim_TimelineUpdateEvent>();
+				TimelineUpdateEvent->Request = SlotRequest.Request;
+				TimelineUpdateEvent->TimelineProgress = ChildProgress;
+
+				Context.RaiseOutputTraitEvent(TimelineUpdateEvent);
+			}
+		}
+
+		// Check if we are blending out
+		if (!bJustTransitioned && InstanceData->CurrentlyActiveRequestIndex != INDEX_NONE)
+		{
+			const FPlayAnimSlotRequest& ActiveSlotRequest = InstanceData->SlotRequests[InstanceData->CurrentlyActiveRequestIndex];
+
+			if (ActiveSlotRequest.State == EPlayAnimRequestState::Active)
+			{
+				const FAnimNextPlayAnimRequestArgs& RequestArgs = ActiveSlotRequest.Request->GetArgs();
+
+				const float BlendOutTime = RequestArgs.BlendOutSettings.Blend.BlendTime;
+				if (CurrentRequestTimeLeft <= BlendOutTime)
+				{
+					// We are ready to start blending out
+					{
+						auto StatusUpdateEvent = MakeTraitEvent<FPlayAnim_StatusUpdateEvent>();
+						StatusUpdateEvent->Request = ActiveSlotRequest.Request;
+						StatusUpdateEvent->Status = EPlayAnimStatus::BlendingOut;
+
+						Context.RaiseOutputTraitEvent(StatusUpdateEvent);
+					}
+
+					// Find an empty request we can use
+					const int32 FreeRequestIndex = FindFreeRequestIndexOrAdd(*InstanceData);
+
+					FPlayAnimSlotRequest& FreeSlotRequest = InstanceData->SlotRequests[FreeRequestIndex];
+					FreeSlotRequest.Initialize(FPlayAnimRequestPtr(), RequestArgs.BlendOutSettings, nullptr);
+
+					const int32 OldChildIndex = InstanceData->CurrentlyActiveRequestIndex;
+					const int32 NewChildIndex = FreeRequestIndex;
+
+					InstanceData->CurrentlyActiveRequestIndex = FreeRequestIndex;
+
+					DiscreteBlendTrait.OnBlendTransition(Context, OldChildIndex, NewChildIndex);
+				}
+			}
+		}
+
+		if (InstanceData->CurrentlyActiveRequestIndex != INDEX_NONE)
+		{
+			const FPlayAnimSlotRequest& ActiveSlotRequest = InstanceData->SlotRequests[InstanceData->CurrentlyActiveRequestIndex];
+			if (ActiveSlotRequest.State == EPlayAnimRequestState::Active)
+			{
+				InstanceData->SubGraphParameters->Update(TraitState.GetDeltaTime());
+
+				FParamStack& ParamStack = FParamStack::Get();
+				InstanceData->SubGraphParametersLayerHandle = ParamStack.PushLayer(InstanceData->SubGraphParameters->GetLayerHandle());
+			}
+		}
+	}
+
+	void FPlayAnimSlotTrait::PostUpdate(FUpdateTraversalContext& Context, const TTraitBinding<IUpdate>& Binding, const FTraitUpdateState& TraitState) const
+	{
+		FInstanceData* InstanceData = Binding.GetInstanceData<FInstanceData>();
+
+		if (InstanceData->SubGraphParametersLayerHandle.IsValid())
+		{
+			FParamStack& ParamStack = FParamStack::Get();
+			ParamStack.PopLayer(InstanceData->SubGraphParametersLayerHandle);
+
+			InstanceData->SubGraphParametersLayerHandle.Invalidate();
+		}
+	}
+
+	void FPlayAnimSlotTrait::QueueChildrenForTraversal(FUpdateTraversalContext& Context, const TTraitBinding<IUpdateTraversal>& Binding, const FTraitUpdateState& TraitState, FUpdateTraversalQueue& TraversalQueue) const
+	{
+		const FInstanceData* InstanceData = Binding.GetInstanceData<FInstanceData>();
+
+		const int32 NumRequests = InstanceData->SlotRequests.Num();
+		check(NumRequests != 0);	// Should never happen since the source is always present
+
+		TTraitBinding<IDiscreteBlend> DiscreteBlendTrait;
+		Binding.GetStackInterface(DiscreteBlendTrait);
+
+		for (int32 RequestIndex = 0; RequestIndex < NumRequests; ++RequestIndex)
+		{
+			const FPlayAnimSlotRequest& SlotRequest = InstanceData->SlotRequests[RequestIndex];
+			const float BlendWeight = DiscreteBlendTrait.GetBlendWeight(Context, RequestIndex);
+
+			FTraitUpdateState RequestSlotTraitState = TraitState
+				.WithWeight(BlendWeight)
+				.AsBlendingOut(RequestIndex != InstanceData->CurrentlyActiveRequestIndex)
+				.AsNewlyRelevant(!SlotRequest.bWasRelevant);
+
+			TraversalQueue.Push(InstanceData->SlotRequests[RequestIndex].ChildPtr, RequestSlotTraitState);
+		}
+	}
+
+	float FPlayAnimSlotTrait::GetBlendWeight(FExecutionContext& Context, const TTraitBinding<IDiscreteBlend>& Binding, int32 ChildIndex) const
+	{
+		const FInstanceData* InstanceData = Binding.GetInstanceData<FInstanceData>();
+
+		if (ChildIndex == InstanceData->CurrentlyActiveRequestIndex)
+		{
+			return 1.0f;	// Active child has full weight
+		}
+		else if (InstanceData->SlotRequests.IsValidIndex(ChildIndex))
+		{
+			return 0.0f;	// Other children have no weight
+		}
+		else
+		{
+			// Invalid child index
+			return -1.0f;
+		}
+	}
+
+	int32 FPlayAnimSlotTrait::GetBlendDestinationChildIndex(FExecutionContext& Context, const TTraitBinding<IDiscreteBlend>& Binding) const
+	{
+		const FInstanceData* InstanceData = Binding.GetInstanceData<FInstanceData>();
+
+		return InstanceData->CurrentlyActiveRequestIndex;
+	}
+
+	void FPlayAnimSlotTrait::OnBlendTransition(FExecutionContext& Context, const TTraitBinding<IDiscreteBlend>& Binding, int32 OldChildIndex, int32 NewChildIndex) const
+	{
+		TTraitBinding<IDiscreteBlend> DiscreteBlendTrait;
+		Binding.GetStackInterface(DiscreteBlendTrait);
+
+		// We initiate immediately when we transition
+		DiscreteBlendTrait.OnBlendInitiated(Context, NewChildIndex);
+
+		// We terminate immediately when we transition
+		DiscreteBlendTrait.OnBlendTerminated(Context, OldChildIndex);
+	}
+
+	void FPlayAnimSlotTrait::OnBlendInitiated(FExecutionContext& Context, const TTraitBinding<IDiscreteBlend>& Binding, int32 ChildIndex) const
+	{
+		FInstanceData* InstanceData = Binding.GetInstanceData<FInstanceData>();
+
+		if (InstanceData->SlotRequests.IsValidIndex(ChildIndex))
+		{
+			// Allocate our new request instance
+			FPlayAnimSlotRequest& SlotRequest = InstanceData->SlotRequests[ChildIndex];
+
+			if (SlotRequest.State == EPlayAnimRequestState::Active)
+			{
+				const FName EntryPoint = SlotRequest.SubGraph->GetDefaultEntryPoint();
+				SlotRequest.SubGraph->AllocateInstance(Binding.GetTraitPtr().GetNodeInstance()->GetOwner(), SlotRequest.GraphInstance, EntryPoint);
+				SlotRequest.ChildPtr = SlotRequest.GraphInstance.GetGraphRootPtr();
+
+				// Setup our graph parameters
+				const FAnimNextPlayAnimRequestArgs& RequestArgs = SlotRequest.Request->GetArgs();
+
+				const FString SubGraphPath = SlotRequest.SubGraph->GetPathName();
+
+				const FName AnimationObjectName(FString::Printf(TEXT("%s:AnimationObject"), *SubGraphPath));
+				const FName StartPositionName(FString::Printf(TEXT("%s:StartPosition"), *SubGraphPath));
+				const FName PlayRateName(FString::Printf(TEXT("%s:PlayRate"), *SubGraphPath));
+				const FName IsLoopingName(FString::Printf(TEXT("%s:IsLooping"), *SubGraphPath));
+
+				FInstancedPropertyBag SubGraphParameters;
+				SubGraphParameters.AddProperty(AnimationObjectName, EPropertyBagPropertyType::Object, UAnimSequence::StaticClass());
+				SubGraphParameters.SetValueObject(AnimationObjectName, RequestArgs.AnimationObject);
+				SubGraphParameters.AddProperty(StartPositionName, EPropertyBagPropertyType::Double);
+				SubGraphParameters.SetValueFloat(StartPositionName, RequestArgs.StartPosition);
+				SubGraphParameters.AddProperty(PlayRateName, EPropertyBagPropertyType::Double);
+				SubGraphParameters.SetValueFloat(PlayRateName, RequestArgs.PlayRate);
+				SubGraphParameters.AddProperty(IsLoopingName, EPropertyBagPropertyType::Bool);
+				SubGraphParameters.SetValueBool(IsLoopingName, false);
+
+				InstanceData->SubGraphParameters = MakeUnique<FPropertyBagProxy>(NAME_None, MoveTemp(SubGraphParameters));
+
+				// TODO: Validate that our child implements the ITimeline interface
+
+				{
+					// Queue our status update
+					auto StatusUpdateEvent = MakeTraitEvent<FPlayAnim_StatusUpdateEvent>();
+					StatusUpdateEvent->Request = SlotRequest.Request;
+					StatusUpdateEvent->Status = EPlayAnimStatus::Playing;
+
+					Context.RaiseOutputTraitEvent(StatusUpdateEvent);
+				}
+			}
+			else if (SlotRequest.State == EPlayAnimRequestState::ActiveSource)
+			{
+				const FSharedData* SharedData = Binding.GetSharedData<FSharedData>();
+
+				SlotRequest.ChildPtr = Context.AllocateNodeInstance(Binding, SharedData->Source);
+			}
+		}
+	}
+
+	void FPlayAnimSlotTrait::OnBlendTerminated(FExecutionContext& Context, const TTraitBinding<IDiscreteBlend>& Binding, int32 ChildIndex) const
+	{
+		FInstanceData* InstanceData = Binding.GetInstanceData<FInstanceData>();
+
+		if (InstanceData->SlotRequests.IsValidIndex(ChildIndex))
+		{
+			// Deallocate our request instance
+			FPlayAnimSlotRequest& SlotRequest = InstanceData->SlotRequests[ChildIndex];
+
+			if (SlotRequest.State == EPlayAnimRequestState::Active)
+			{
+				SlotRequest.GraphInstance.Release();
+
+				{
+					// Queue our status update
+					auto StatusUpdateEvent = MakeTraitEvent<FPlayAnim_StatusUpdateEvent>();
+					StatusUpdateEvent->Request = SlotRequest.Request;
+					StatusUpdateEvent->Status = EPlayAnimStatus::Completed;
+
+					Context.RaiseOutputTraitEvent(StatusUpdateEvent);
+				}
+			}
+
+			SlotRequest.Request = nullptr;
+			SlotRequest.ChildPtr.Reset();
+			SlotRequest.State = EPlayAnimRequestState::Inactive;
+			SlotRequest.bWasRelevant = false;
+		}
+	}
+
+	float FPlayAnimSlotTrait::GetBlendTime(FExecutionContext& Context, const TTraitBinding<ISmoothBlend>& Binding, int32 ChildIndex) const
+	{
+		const FInstanceData* InstanceData = Binding.GetInstanceData<FInstanceData>();
+		if (InstanceData->SlotRequests.IsValidIndex(ChildIndex))
+		{
+			const FPlayAnimSlotRequest& SlotRequest = InstanceData->SlotRequests[ChildIndex];
+			return SlotRequest.BlendSettings.Blend.BlendTime;
+		}
+		else
+		{
+			// Unknown child
+			return 0.0f;
+		}
+	}
+
+	EAlphaBlendOption FPlayAnimSlotTrait::GetBlendType(FExecutionContext& Context, const TTraitBinding<ISmoothBlend>& Binding, int32 ChildIndex) const
+	{
+		const FInstanceData* InstanceData = Binding.GetInstanceData<FInstanceData>();
+		if (InstanceData->SlotRequests.IsValidIndex(ChildIndex))
+		{
+			const FPlayAnimSlotRequest& SlotRequest = InstanceData->SlotRequests[ChildIndex];
+			return SlotRequest.BlendSettings.Blend.BlendOption;
+		}
+		else
+		{
+			// Unknown child
+			return EAlphaBlendOption::Linear;
+		}
+	}
+
+	UCurveFloat* FPlayAnimSlotTrait::GetCustomBlendCurve(FExecutionContext& Context, const TTraitBinding<ISmoothBlend>& Binding, int32 ChildIndex) const
+	{
+		const FInstanceData* InstanceData = Binding.GetInstanceData<FInstanceData>();
+		if (InstanceData->SlotRequests.IsValidIndex(ChildIndex))
+		{
+			const FPlayAnimSlotRequest& SlotRequest = InstanceData->SlotRequests[ChildIndex];
+			return SlotRequest.BlendSettings.Blend.CustomCurve;
+		}
+		else
+		{
+			// Unknown child
+			return nullptr;
+		}
+	}
+
+	float FPlayAnimSlotTrait::GetBlendTime(FExecutionContext& Context, const TTraitBinding<IInertializerBlend>& Binding, int32 ChildIndex) const
+	{
+		const FInstanceData* InstanceData = Binding.GetInstanceData<FInstanceData>();
+		if (InstanceData->SlotRequests.IsValidIndex(ChildIndex))
+		{
+			const FPlayAnimSlotRequest& SlotRequest = InstanceData->SlotRequests[ChildIndex];
+			if (SlotRequest.BlendSettings.BlendMode == EAnimNextPlayAnimBlendMode::Inertialization)
+			{
+				return SlotRequest.BlendSettings.Blend.BlendTime;
+			}
+			else
+			{
+				// Not an inertializing blend
+				return 0.0f;
+			}
+		}
+		else
+		{
+			// Unknown child
+			return 0.0f;
+		}
+	}
+
+	void FPlayAnimSlotTrait::AddReferencedObjects(const FExecutionContext& Context, const TTraitBinding<IGarbageCollection>& Binding, FReferenceCollector& Collector) const
+	{
+		IGarbageCollection::AddReferencedObjects(Context, Binding, Collector);
+
+		FInstanceData* InstanceData = Binding.GetInstanceData<FInstanceData>();
+
+		if (InstanceData->PendingRequest.Request)
+		{
+			InstanceData->PendingRequest.Request->AddReferencedObjects(Collector);
+		}
+
+		for (FPlayAnimSlotRequest& SlotRequest : InstanceData->SlotRequests)
+		{
+			if (SlotRequest.Request)
+			{
+				SlotRequest.Request->AddReferencedObjects(Collector);
+			}
+
+			Collector.AddReferencedObject(SlotRequest.SubGraph);
+		}
+	}
+}
diff --git a/Engine/Plugins/Experimental/AnimNextAnimGraph/Source/AnimNextAnimGraph/Private/PlayAnim/PlayAnimSlotTrait.h b/Engine/Plugins/Experimental/AnimNextAnimGraph/Source/AnimNextAnimGraph/Private/PlayAnim/PlayAnimSlotTrait.h
new file mode 100644
--- /dev/null
+++ b/Engine/Plugins/Experimental/AnimNextAnimGraph/Source/AnimNextAnimGraph/Private/PlayAnim/PlayAnimSlotTrait.h
@@ -0,0 +1,191 @@
+// Copyright Epic Games, Inc. All Rights Reserved.
+
+#pragma once
+
+#include "CoreMinimal.h"
+
+#include "TraitCore/Trait.h"
+#include "TraitInterfaces/IDiscreteBlend.h"
+#include "TraitInterfaces/IGarbageCollection.h"
+#include "TraitInterfaces/IHierarchy.h"
+#include "TraitInterfaces/IInertializerBlend.h"
+#include "TraitInterfaces/ISmoothBlend.h"
+#include "TraitInterfaces/IUpdate.h"
+#include "Graph/AnimNextGraph.h"
+#include "Graph/AnimNextGraphInstancePtr.h"
+#include "PlayAnim/PlayAnimEvents.h"
+#include "Param/PropertyBagProxy.h"
+#include "Chooser.h"
+
+#include "PlayAnimSlotTrait.generated.h"
+
+USTRUCT(meta = (DisplayName = "PlayAnim Chooser Parameters"))
+struct FAnimNextPlayAnimChooserParameters
+{
+	GENERATED_BODY()
+
+	// The class type of the animation object
+	UPROPERTY()
+	TObjectPtr<UClass> AnimationObjectType;
+};
+
+USTRUCT(meta = (DisplayName = "PlayAnim Slot"))
+struct FAnimNextPlayAnimSlotTraitSharedData : public FAnimNextTraitSharedData
+{
+	GENERATED_BODY()
+
+	/** Default input when no animation request has been made on this slot. */
+	UPROPERTY()
+	FAnimNextTraitHandle Source;
+
+	/** The name of this slot that the PlayAnim API refers to. */
+	UPROPERTY(EditAnywhere, Category = "Default")
+	FName SlotName;
+
+	/** The chooser instance to use to convert animation objects into a sub-graph to play. */
+	UPROPERTY(EditAnywhere, Category = "Default")
+	TObjectPtr<UChooserTable> SubGraphChooser;
+
+	// Latent pin support boilerplate
+	#define TRAIT_LATENT_PROPERTIES_ENUMERATOR(GeneratorMacro) \
+		GeneratorMacro(SlotName) \
+		GeneratorMacro(SubGraphChooser) \
+
+	GENERATE_TRAIT_LATENT_PROPERTIES(FAnimNextPlayAnimSlotTraitSharedData, TRAIT_LATENT_PROPERTIES_ENUMERATOR)
+	#undef TRAIT_LATENT_PROPERTIES_ENUMERATOR
+};
+
+namespace UE::AnimNext
+{
+	/**
+	 * FPlayAnimSlotTrait
+	 * 
+	 * A trait that provides a slot onto which a PlayAnim request can be processed.
+	 * It allows for this trait to act as a pass-through when not actively used and
+	 * when a PlayAnim request is made to start playing a child instance, we blend to it.
+	 */
+	struct FPlayAnimSlotTrait : FBaseTrait, IUpdate, IUpdateTraversal, IHierarchy, IDiscreteBlend, ISmoothBlend, IInertializerBlend, IGarbageCollection
+	{
+		DECLARE_ANIM_TRAIT(FPlayAnimSlotTrait, 0x5667d3ee, FBaseTrait)
+
+		// Slot request state
+		enum class EPlayAnimRequestState : uint8
+		{
+			// Slot request is inactive
+			Inactive,
+
+			// Slot request is active and using a sub-graph
+			Active,
+
+			// Slot request is active and using the source input
+			ActiveSource,
+		};
+
+		struct FPlayAnimPendingRequest
+		{
+			// The PlayAnim request
+			FPlayAnimRequestPtr Request;
+
+			// Whether or not a Stop request was issued
+			bool bStop = false;
+
+			// Returns whether or not we have a pending request
+			bool IsValid() const { return !!Request || bStop; }
+
+			// Resets the pending request
+			void Reset()
+			{
+				Request = nullptr;
+				bStop = false;
+			}
+		};
+
+		struct FPlayAnimSlotRequest
+		{
+			// The PlayAnim request
+			FPlayAnimRequestPtr Request;
+
+			// The blend settings to use
+			FPlayAnimBlendSettings BlendSettings;
+
+			// The sub-graph used by the graph instance, as selected by the chooser
+			TObjectPtr<const UAnimNextGraph> SubGraph;
+
+			// The graph instance
+			FAnimNextGraphInstancePtr GraphInstance;
+
+			// Our child handle
+			// If we use the source input, this is a strong handle to it, otherwise we are a weak handle to the graph instance's root
+			FTraitPtr ChildPtr;
+
+			// The current request state
+			EPlayAnimRequestState State = EPlayAnimRequestState::Inactive;
+
+			// Whether or not this slot state was previously relevant
+			bool bWasRelevant = false;
+
+			// Initializes a request to begin playing
+			void Initialize(FPlayAnimRequestPtr InRequest, const FPlayAnimBlendSettings& InBlendSettings, const UAnimNextGraph* InSubGraph);
+		};
+
+		using FSharedData = FAnimNextPlayAnimSlotTraitSharedData;
+
+		struct FInstanceData : FTrait::FInstanceData
+		{
+			// List of PlayAnim slot requests
+			TArray<FPlayAnimSlotRequest> SlotRequests;
+
+			// PlayAnim pending request
+			FPlayAnimPendingRequest PendingRequest;
+
+			// The index of the currently active request
+			// All other requests are blending out
+			int32 CurrentlyActiveRequestIndex = INDEX_NONE;
+
+			// The parameters that control our sub-graphs
+			TUniquePtr<FPropertyBagProxy> SubGraphParameters;
+
+			// The param stack layer handle for our sub-graph parameters
+			FParamStack::FPushedLayerHandle SubGraphParametersLayerHandle;
+
+			void Construct(const FExecutionContext& Context, const FTraitBinding& Binding);
+			void Destruct(const FExecutionContext& Context, const FTraitBinding& Binding);
+		};
+
+		// Internal impl
+		static int32 FindFreeRequestIndexOrAdd(FInstanceData& InstanceData);
+
+		// Event handling impl
+		ETraitStackPropagation OnPlayEvent(const FExecutionContext& Context, FTraitBinding& Binding, FPlayAnim_PlayEvent& Event) const;
+		ETraitStackPropagation OnStopEvent(const FExecutionContext& Context, FTraitBinding& Binding, FPlayAnim_StopEvent& Event) const;
+
+		// IUpdate impl
+		virtual void PreUpdate(FUpdateTraversalContext& Context, const TTraitBinding<IUpdate>& Binding, const FTraitUpdateState& TraitState) const override;
+		virtual void PostUpdate(FUpdateTraversalContext& Context, const TTraitBinding<IUpdate>& Binding, const FTraitUpdateState& TraitState) const override;
+
+		// IUpdateTraversal impl
+		virtual void QueueChildrenForTraversal(FUpdateTraversalContext& Context, const TTraitBinding<IUpdateTraversal>& Binding, const FTraitUpdateState& TraitState, FUpdateTraversalQueue& TraversalQueue) const override;
+
+		// IHierarchy impl
+		virtual uint32 GetNumChildren(const FExecutionContext& Context, const TTraitBinding<IHierarchy>& Binding) const override;
+		virtual void GetChildren(const FExecutionContext& Context, const TTraitBinding<IHierarchy>& Binding, FChildrenArray& Children) const override;
+
+		// IDiscreteBlend impl
+		virtual float GetBlendWeight(FExecutionContext& Context, const TTraitBinding<IDiscreteBlend>& Binding, int32 ChildIndex) const override;
+		virtual int32 GetBlendDestinationChildIndex(FExecutionContext& Context, const TTraitBinding<IDiscreteBlend>& Binding) const override;
+		virtual void OnBlendTransition(FExecutionContext& Context, const TTraitBinding<IDiscreteBlend>& Binding, int32 OldChildIndex, int32 NewChildIndex) const override;
+		virtual void OnBlendInitiated(FExecutionContext& Context, const TTraitBinding<IDiscreteBlend>& Binding, int32 ChildIndex) const override;
+		virtual void OnBlendTerminated(FExecutionContext& Context, const TTraitBinding<IDiscreteBlend>& Binding, int32 ChildIndex) const override;
+
+		// ISmoothBlend impl
+		virtual float GetBlendTime(FExecutionContext& Context, const TTraitBinding<ISmoothBlend>& Binding, int32 ChildIndex) const override;
+		virtual EAlphaBlendOption GetBlendType(FExecutionContext& Context, const TTraitBinding<ISmoothBlend>& Binding, int32 ChildIndex) const override;
+		virtual UCurveFloat* GetCustomBlendCurve(FExecutionContext& Context, const TTraitBinding<ISmoothBlend>& Binding, int32 ChildIndex) const override;
+
+		// IInertializerBlend impl
+		virtual float GetBlendTime(FExecutionContext& Context, const TTraitBinding<IInertializerBlend>& Binding, int32 ChildIndex) const override;
+
+		// IGarbageCollection impl
+		virtual void AddReferencedObjects(const FExecutionContext& Context, const TTraitBinding<IGarbageCollection>& Binding, FReferenceCollector& Collector) const override;
+	};
+}
diff --git a/Engine/Plugins/Experimental/AnimNextAnimGraph/Source/AnimNextAnimGraph/Public/PlayAnim/PlayAnimRequest.h b/Engine/Plugins/Experimental/AnimNextAnimGraph/Source/AnimNextAnimGraph/Public/PlayAnim/PlayAnimRequest.h
new file mode 100644
--- /dev/null
+++ b/Engine/Plugins/Experimental/AnimNextAnimGraph/Source/AnimNextAnimGraph/Public/PlayAnim/PlayAnimRequest.h
@@ -0,0 +1,197 @@
+// Copyright Epic Games, Inc. All Rights Reserved.
+
+#pragma once
+
+#include "CoreMinimal.h"
+
+#include "AlphaBlend.h"
+#include "TraitCore/TraitEvent.h"
+#include "TraitInterfaces/ITimeline.h"
+#include "PlayAnim/PlayAnimStatus.h"
+
+#include "PlayAnimRequest.generated.h"
+
+class FReferenceCollector;
+class UAnimNextGraph;
+class UAnimNextComponent;
+class UAnimSequence;
+class UBlendProfile;
+
+namespace UE::AnimNext
+{
+	struct FPlayAnim_PlayEvent;
+	struct FPlayAnim_StatusUpdateEvent;
+	struct FPlayAnim_TimelineUpdateEvent;
+}
+
+UENUM()
+enum class EAnimNextPlayAnimBlendMode : uint8
+{
+	// Uses standard weight based blend
+	Standard,
+
+	// Uses inertialization. Requires an inertialization trait somewhere earlier in the graph.
+	Inertialization,
+};
+
+/**
+ * PlayAnim Blend Settings
+ *
+ * Encapsulates the blend settings used by Play animation requests.
+ */
+USTRUCT(BlueprintType)
+struct FAnimNextPlayAnimBlendSettings
+{
+	GENERATED_BODY()
+
+	/** Blend Profile to use for this blend */
+	//UPROPERTY(EditAnywhere, BlueprintReadWrite, Category = "Blend", meta = (DisplayAfter = "Blend"))
+	//TObjectPtr<UBlendProfile> BlendProfile;
+
+	/** AlphaBlend options (time, curve, etc.) */
+	UPROPERTY(EditAnywhere, BlueprintReadWrite, Category = "Blend", meta = (DisplayAfter = "BlendMode"))
+	FAlphaBlendArgs Blend;
+
+	/** Type of blend mode (Standard vs Inertial) */
+	UPROPERTY(EditAnywhere, BlueprintReadWrite, Category = "Blend")
+	EAnimNextPlayAnimBlendMode BlendMode = EAnimNextPlayAnimBlendMode::Standard;
+};
+
+/**
+ * PlayAnim Request Arguments
+ *
+ * Encapsulates the parameters required to initiate a Play animation request.
+ */
+USTRUCT()
+struct FAnimNextPlayAnimRequestArgs
+{
+	GENERATED_BODY()
+
+	// The slot name to play the animation object on with this request
+	UPROPERTY()
+	FName SlotName;
+
+	// The animation object to play with this request
+	UPROPERTY()
+	TObjectPtr<UAnimSequence> AnimationObject;
+
+	// The play rate of the request
+	UPROPERTY()
+	float PlayRate = 1.0f;
+
+	// The timeline start position of the request
+	UPROPERTY()
+	float StartPosition = 0.0f;
+
+	// The blend settings to use when blending in
+	UPROPERTY()
+	FAnimNextPlayAnimBlendSettings BlendInSettings;
+
+	// The blend settings to use when blending out (if not interrupted)
+	UPROPERTY()
+	FAnimNextPlayAnimBlendSettings BlendOutSettings;
+};
+
+namespace UE::AnimNext
+{
+	// Create a namespaced aliases to simplify usage
+	using EPlayAnimBlendMode = EAnimNextPlayAnimBlendMode;
+	using FPlayAnimBlendSettings = FAnimNextPlayAnimBlendSettings;
+	using FPlayAnimRequestArgs = FAnimNextPlayAnimRequestArgs;
+
+	struct FPlayAnimRequest;
+	DECLARE_DELEGATE_OneParam(FAnimNextOnPlayAnimStarted, const FPlayAnimRequest&)
+	DECLARE_DELEGATE_OneParam(FAnimNextOnPlayAnimCompleted, const FPlayAnimRequest&)
+	DECLARE_DELEGATE_OneParam(FAnimNextOnPlayAnimInterrupted, const FPlayAnimRequest&)
+	DECLARE_DELEGATE_OneParam(FAnimNextOnPlayAnimBlendingOut, const FPlayAnimRequest&)
+
+	/**
+	 * PlayAnim Request
+	 *
+	 * Instances of this class represent individual requests to the PlayAnim system.
+	 * They are allocated as shared pointers and ownership is split between gameplay (until
+	 * it no longer cares about a particular request) and the animation slot that plays it
+	 * (until the request completes).
+	 * 
+	 * Use MakePlayAnimRequest(...) to construct instances of this type.
+	 */
+	struct FPlayAnimRequest : public TSharedFromThis<FPlayAnimRequest, ESPMode::ThreadSafe>
+	{
+		// Callback called when the request starts playing (status transitions from pending to playing)
+		FAnimNextOnPlayAnimStarted OnStarted;
+
+		// Callback called when the request completes (status transitions from playing to completed)
+		FAnimNextOnPlayAnimCompleted OnCompleted;
+
+		// Callback called when the request is interrupted (either by calling Stop on it or by another request)
+		FAnimNextOnPlayAnimInterrupted OnInterrupted;
+
+		// Callback called when the request starts blending out (if it wasn't interrupted)
+		FAnimNextOnPlayAnimBlendingOut OnBlendingOut;
+
+		// Sends this request to the specified component and it will attempt to play with the requested arguments
+		bool Play(const FPlayAnimRequestArgs& InRequestArgs, UAnimNextComponent* InComponent);
+
+		// Interrupts this request and request that we transition to the source input on the playing slot
+		void Stop();
+
+		// Returns the arguments this request is playing
+		const FPlayAnimRequestArgs& GetArgs() const;
+
+		// Returns the request status
+		EPlayAnimStatus GetStatus() const;
+
+		// Returns the current timeline progress
+		FTimelineProgress GetTimelineProgress() const;
+
+		// Returns whether or not this request has expired
+		bool HasExpired() const;
+
+		// Returns whether or not this request has completed (might have been interrupted)
+		bool HasCompleted() const;
+
+		// Returns whether or not this request is playing (might be blending out or interrupted)
+		bool IsPlaying() const;
+
+		// Returns whether or not this request is blending out
+		bool IsBlendingOut() const;
+
+		// Returns whether or not this request was interrupted (by Stop or by another request)
+		bool WasInterrupted() const;
+
+		// GC API
+		void AddReferencedObjects(FReferenceCollector& Collector);
+
+	private:
+		void OnStatusUpdate(EPlayAnimStatus NewStatus);
+		void OnTimelineUpdate(FTimelineProgress NewTimelineProgress);
+
+		// The request arguments
+		FPlayAnimRequestArgs RequestArgs;
+
+		// The component we are playing on
+		TObjectPtr<UAnimNextComponent> Component;
+
+		// The pending start event if we haven't started playing yet
+		FAnimNextTraitEventPtr PendingStartEvent;
+
+		// The current request status
+		EPlayAnimStatus Status = EPlayAnimStatus::None;
+
+		// The current timeline progress
+		FTimelineProgress TimelineProgress;
+
+		friend FPlayAnim_PlayEvent;
+		friend FPlayAnim_StatusUpdateEvent;
+		friend FPlayAnim_TimelineUpdateEvent;
+	};
+
+	// Create a shared pointer alias for PlayAnim requests
+	using FPlayAnimRequestPtr = TSharedPtr<FPlayAnimRequest, ESPMode::ThreadSafe>;
+
+	// Constructs a PlayAnim request object
+	inline FPlayAnimRequestPtr MakePlayAnimRequest()
+	{
+		return MakeShared<FPlayAnimRequest>();
+	}
+}
diff --git a/Engine/Plugins/Experimental/AnimNextAnimGraph/Source/AnimNextAnimGraph/Public/PlayAnim/PlayAnimStatus.h b/Engine/Plugins/Experimental/AnimNextAnimGraph/Source/AnimNextAnimGraph/Public/PlayAnim/PlayAnimStatus.h
new file mode 100644
--- /dev/null
+++ b/Engine/Plugins/Experimental/AnimNextAnimGraph/Source/AnimNextAnimGraph/Public/PlayAnim/PlayAnimStatus.h
@@ -0,0 +1,51 @@
+// Copyright Epic Games, Inc. All Rights Reserved.
+
+#pragma once
+
+#include "CoreMinimal.h"
+
+#include "PlayAnimStatus.generated.h"
+
+/**
+ * PlayAnim Request Status
+ */
+UENUM()
+enum class EAnimNextPlayAnimStatus : uint8
+{
+	//////////////////////////////////////////////////////////////////////////
+	// Primary status is held in the bottom 5 bits
+	// These are mutually exclusive
+
+	// No current status
+	None			= 0x00,
+
+	// Animation request has been queue but hasn't been processed yet
+	Pending			= 0x01,
+
+	// Animation is playing (may be blending in or out, may be interrupted)
+	Playing			= 0x02,
+
+	// Animation has finished playing (may have the interrupted flag to signal that it didn't stop on its own)
+	Completed		= 0x04,
+
+	// Animation request expired without playing
+	Expired			= 0x10,
+
+
+	//////////////////////////////////////////////////////////////////////////
+	// Additive status is held in the top 3 bits
+
+	// Animation was interrupted by another animation request (or a stop request)
+	Interrupted		= 0x40,
+
+	// Animation is blending out
+	BlendingOut		= 0x80,
+};
+
+ENUM_CLASS_FLAGS(EAnimNextPlayAnimStatus);
+
+namespace UE::AnimNext
+{
+	// Create a namespaced alias to simplify usage
+	using EPlayAnimStatus = EAnimNextPlayAnimStatus;
+}
diff --git a/Engine/Plugins/Experimental/AnimNextAnimGraph/Source/AnimNextAnimGraphUncookedOnly/AnimNextAnimGraphUncookedOnly.Build.cs b/Engine/Plugins/Experimental/AnimNextAnimGraph/Source/AnimNextAnimGraphUncookedOnly/AnimNextAnimGraphUncookedOnly.Build.cs
--- a/Engine/Plugins/Experimental/AnimNextAnimGraph/Source/AnimNextAnimGraphUncookedOnly/AnimNextAnimGraphUncookedOnly.Build.cs
+++ b/Engine/Plugins/Experimental/AnimNextAnimGraph/Source/AnimNextAnimGraphUncookedOnly/AnimNextAnimGraphUncookedOnly.Build.cs
@@ -29,6 +29,7 @@ namespace UnrealBuildTool.Rules
 					new string[]
 					{
 						"UnrealEd",
+						"BlueprintGraph",
 					}
 				);
 			}
diff --git a/Engine/Plugins/Experimental/AnimNextAnimGraph/Source/AnimNextAnimGraphUncookedOnly/Internal/PlayAnim/K2Node_PlayAnim.h b/Engine/Plugins/Experimental/AnimNextAnimGraph/Source/AnimNextAnimGraphUncookedOnly/Internal/PlayAnim/K2Node_PlayAnim.h
new file mode 100644
--- /dev/null
+++ b/Engine/Plugins/Experimental/AnimNextAnimGraph/Source/AnimNextAnimGraphUncookedOnly/Internal/PlayAnim/K2Node_PlayAnim.h
@@ -0,0 +1,24 @@
+// Copyright Epic Games, Inc. All Rights Reserved.
+
+#pragma once
+
+#include "CoreMinimal.h"
+
+#include "K2Node_BaseAsyncTask.h"
+
+#include "K2Node_PlayAnim.generated.h"
+
+UCLASS()
+class UK2Node_PlayAnim : public UK2Node_BaseAsyncTask
+{
+	GENERATED_UCLASS_BODY()
+
+	//~ Begin UEdGraphNode Interface
+	virtual FText GetTooltipText() const override;
+	virtual FText GetNodeTitle(ENodeTitleType::Type TitleType) const override;
+	//~ End UEdGraphNode Interface
+
+	//~ Begin UK2Node Interface
+	virtual FText GetMenuCategory() const override;
+	//~ End UK2Node Interface
+};
diff --git a/Engine/Plugins/Experimental/AnimNextAnimGraph/Source/AnimNextAnimGraphUncookedOnly/Private/PlayAnim/K2Node_PlayAnim.cpp b/Engine/Plugins/Experimental/AnimNextAnimGraph/Source/AnimNextAnimGraphUncookedOnly/Private/PlayAnim/K2Node_PlayAnim.cpp
new file mode 100644
--- /dev/null
+++ b/Engine/Plugins/Experimental/AnimNextAnimGraph/Source/AnimNextAnimGraphUncookedOnly/Private/PlayAnim/K2Node_PlayAnim.cpp
@@ -0,0 +1,32 @@
+// Copyright Epic Games, Inc. All Rights Reserved.
+
+#include "PlayAnim/K2Node_PlayAnim.h"
+
+#include "PlayAnim/PlayAnimCallbackProxy.h"
+
+#define LOCTEXT_NAMESPACE "K2Node_PlayAnim"
+
+UK2Node_PlayAnim::UK2Node_PlayAnim(const FObjectInitializer& ObjectInitializer)
+	: Super(ObjectInitializer)
+{
+	ProxyFactoryFunctionName = GET_FUNCTION_NAME_CHECKED(UPlayAnimCallbackProxy, CreateProxyObjectForPlayAnim);
+	ProxyFactoryClass = UPlayAnimCallbackProxy::StaticClass();
+	ProxyClass = UPlayAnimCallbackProxy::StaticClass();
+}
+
+FText UK2Node_PlayAnim::GetTooltipText() const
+{
+	return LOCTEXT("K2Node_PlayAnim_Tooltip", "Plays an Animation object on an AnimNextComponent");
+}
+
+FText UK2Node_PlayAnim::GetNodeTitle(ENodeTitleType::Type TitleType) const
+{
+	return LOCTEXT("PlayAnim", "Play Animation");
+}
+
+FText UK2Node_PlayAnim::GetMenuCategory() const
+{
+	return LOCTEXT("PlayAnimCategory", "Animation|AnimNext");
+}
+
+#undef LOCTEXT_NAMESPACE
