RigVM: renaming decorator to trait

#rnx

diff --git a/Engine/Plugins/Experimental/AnimNext/Source/AnimNext/Internal/Graph/RigDecorator_AnimNextCppTrait.h b/Engine/Plugins/Experimental/AnimNext/Source/AnimNext/Internal/Graph/RigDecorator_AnimNextCppTrait.h
--- a/Engine/Plugins/Experimental/AnimNext/Source/AnimNext/Internal/Graph/RigDecorator_AnimNextCppTrait.h
+++ b/Engine/Plugins/Experimental/AnimNext/Source/AnimNext/Internal/Graph/RigDecorator_AnimNextCppTrait.h
@@ -2,7 +2,7 @@
 
 #pragma once
 
-#include "RigVMCore/RigVMDecorator.h"
+#include "RigVMCore/RigVMTrait.h"
 
 #include "RigDecorator_AnimNextCppTrait.generated.h"
 
@@ -13,7 +13,7 @@ namespace UE::AnimNext { struct FTrait; }
  * The trait shared data UScriptStruct determines which properties are exposed.
  */
 USTRUCT(BlueprintType)
-struct ANIMNEXT_API FRigDecorator_AnimNextCppDecorator : public FRigVMDecorator
+struct ANIMNEXT_API FRigDecorator_AnimNextCppDecorator : public FRigVMTrait
 {
 	GENERATED_BODY()
 
@@ -24,7 +24,7 @@ struct ANIMNEXT_API FRigDecorator_AnimNextCppDecorator : public FRigVMDecorator
 #if WITH_EDITOR
 	virtual void GetProgrammaticPins(URigVMController* InController, int32 InParentPinIndex, const FString& InDefaultValue, struct FRigVMPinInfoArray& OutPinArray) const override;
 	
-	virtual UScriptStruct* GetDecoratorSharedDataStruct() const override { return DecoratorSharedDataStruct; }
+	virtual UScriptStruct* GetTraitSharedDataStruct() const override { return DecoratorSharedDataStruct; }
 
 	const UE::AnimNext::FTrait* GetTrait() const;
 #endif
diff --git a/Engine/Plugins/Experimental/AnimNext/Source/AnimNextEditor/Private/Graph/AnimNextGraph_EdGraphNodeCustomization.cpp b/Engine/Plugins/Experimental/AnimNext/Source/AnimNextEditor/Private/Graph/AnimNextGraph_EdGraphNodeCustomization.cpp
--- a/Engine/Plugins/Experimental/AnimNext/Source/AnimNextEditor/Private/Graph/AnimNextGraph_EdGraphNodeCustomization.cpp
+++ b/Engine/Plugins/Experimental/AnimNext/Source/AnimNextEditor/Private/Graph/AnimNextGraph_EdGraphNodeCustomization.cpp
@@ -100,7 +100,7 @@ void FAnimNextGraph_EdGraphNodeCustomization::GenerateTraitData(UAnimNextGraph_E
 	if (URigVMNode* ModelNode = EdGraphNode->GetModelNode())
 	{
 		// Obtain the pins from the stack
-		const TArray<URigVMPin*> TraitPins = ModelNode->GetDecoratorPins();
+		const TArray<URigVMPin*> TraitPins = ModelNode->GetTraitPins();
 		if (TraitPins.Num() > 0)
 		{
 			// For each Trait (represented as a pin in the node)
@@ -112,16 +112,16 @@ void FAnimNextGraph_EdGraphNodeCustomization::GenerateTraitData(UAnimNextGraph_E
 				}
 
 				// Create a temporary trait instance, in order to get the correct TraitSharedDataStruct
-				if (TSharedPtr<FStructOnScope> ScopedTrait = ModelNode->GetDecoratorInstance(TraitPin->GetFName()))
+				if (TSharedPtr<FStructOnScope> ScopedTrait = ModelNode->GetTraitInstance(TraitPin->GetFName()))
 				{
 					// Create a scoped struct with the Trait Shared Instance Data and store it for later use
 					// The idea is creating one template per Trait type, but pass all the selected instances
 					// So the multiselection works
-					const FRigVMDecorator* Trait = (FRigVMDecorator*)ScopedTrait->GetStructMemory();
+					const FRigVMTrait* Trait = (FRigVMTrait*)ScopedTrait->GetStructMemory();
 
 					// Import the default pin values into the Trait Shared Instance Data
 					FTraitStackDetailsData* TraitData = nullptr;
-					if (UScriptStruct* TraitSharedInstanceData = Trait->GetDecoratorSharedDataStruct())
+					if (UScriptStruct* TraitSharedInstanceData = Trait->GetTraitSharedDataStruct())
 					{
 						const TSharedPtr<FCategoryDetailsData>* TraitDataPtr = CategoryDetailsData.FindByPredicate([TraitSharedInstanceData](const TSharedPtr<FCategoryDetailsData>& InItem)
 							{
diff --git a/Engine/Plugins/Experimental/AnimNext/Source/AnimNextEditor/Private/Graph/STraitEditorView.cpp b/Engine/Plugins/Experimental/AnimNext/Source/AnimNextEditor/Private/Graph/STraitEditorView.cpp
--- a/Engine/Plugins/Experimental/AnimNext/Source/AnimNextEditor/Private/Graph/STraitEditorView.cpp
+++ b/Engine/Plugins/Experimental/AnimNext/Source/AnimNextEditor/Private/Graph/STraitEditorView.cpp
@@ -113,11 +113,11 @@ FReply STraitEditorView::OnTraitClicked(const FTraitUID InClickedTraitUID)
 				const FName NewTraitTypeName = *Trait->GetTraitName();
 				if (TraitIndex == INDEX_NONE || bIsAddingMissingBaseTrait)
 				{
-					Controller->AddTrait(EdGraphNode->GetFName(), NewTraitTypeName, PinIndex);
+					Controller->AddTraitByName(EdGraphNode->GetFName(), NewTraitTypeName, PinIndex);
 				}
 				else
 				{
-					Controller->SwapTrait(EdGraphNode->GetFName(), SwapTraitData->TraitName, PinIndex, NewTraitTypeName);
+					Controller->SwapTraitByName(EdGraphNode->GetFName(), SwapTraitData->TraitName, PinIndex, NewTraitTypeName);
 				}
 			}
 		}
@@ -149,7 +149,7 @@ FReply STraitEditorView::OnTraitDeleteRequest(const FTraitUID InTraitUIDToDelete
 					{
 						if (UAnimNextGraph_Controller* Controller = Cast<UAnimNextGraph_Controller>(EdGraphNode->GetController()))
 						{
-							Controller->RemoveTrait(EdGraphNode->GetFName(), TraitDataEditorDef->TraitName);
+							Controller->RemoveTraitByName(EdGraphNode->GetFName(), TraitDataEditorDef->TraitName);
 						}
 					}
 				}
@@ -197,11 +197,11 @@ void STraitEditorView::ExecuteTraitDrag(const FTraitUID DraggedTraitUID, const F
 					{
 						if(bIsAddingMissingBaseTrait)
 						{
-							Controller->AddTrait(EdGraphNode->GetFName(), DraggedTraitTypeName, PinIndex);
+							Controller->AddTraitByName(EdGraphNode->GetFName(), DraggedTraitTypeName, PinIndex);
 						}
 						else
 						{
-							Controller->SwapTrait(EdGraphNode->GetFName(), SwapTraitData->TraitName, PinIndex, DraggedTraitTypeName);
+							Controller->SwapTraitByName(EdGraphNode->GetFName(), SwapTraitData->TraitName, PinIndex, DraggedTraitTypeName);
 						}
 					}
 					else
@@ -215,7 +215,7 @@ void STraitEditorView::ExecuteTraitDrag(const FTraitUID DraggedTraitUID, const F
 						}
 						else
 						{
-							Controller->AddTrait(EdGraphNode->GetFName(), DraggedTraitTypeName, PinIndex + 1);
+							Controller->AddTraitByName(EdGraphNode->GetFName(), DraggedTraitTypeName, PinIndex + 1);
 						}
 					}
 				}
@@ -523,7 +523,7 @@ int32 STraitEditorView::GetTraitPinIndex(UAnimNextGraph_EdGraphNode* InEdGraphNo
 			int32 Index = 0;
 			for (const URigVMPin* Pin : Pins)
 			{
-				if (!Pin->IsDecoratorPin() || Pin->IsExecuteContext())
+				if (!Pin->IsTraitPin() || Pin->IsExecuteContext())
 				{
 					Index++;
 					continue;
@@ -536,7 +536,7 @@ int32 STraitEditorView::GetTraitPinIndex(UAnimNextGraph_EdGraphNode* InEdGraphNo
 		else
 		{
 			// Obtain the pins from the stack
-			const TArray<URigVMPin*> TraitPins = ModelNode->GetDecoratorPins();
+			const TArray<URigVMPin*> TraitPins = ModelNode->GetTraitPins();
 			for (const URigVMPin* TraitPin : TraitPins)
 			{
 				if (TraitPin->IsExecuteContext())
@@ -577,7 +577,7 @@ int32 STraitEditorView::GetTraitPinIndex(UAnimNextGraph_EdGraphNode* InEdGraphNo
 			int32 NumBaseTraits = 0;
 
 			// Obtain the pins from the stack
-			const TArray<URigVMPin*> TraitPins = ModelNode->GetDecoratorPins();
+			const TArray<URigVMPin*> TraitPins = ModelNode->GetTraitPins();
 			const int32 NumPins = TraitPins.Num();
 			if (NumPins > 0)
 			{
@@ -595,11 +595,11 @@ int32 STraitEditorView::GetTraitPinIndex(UAnimNextGraph_EdGraphNode* InEdGraphNo
 					}
 
 					// Create a temporary trait instance, in order to get the correct TraitSharedDataStruct
-					if (TSharedPtr<FStructOnScope> ScopedTrait = ModelNode->GetDecoratorInstance(TraitPin->GetFName()))
+					if (TSharedPtr<FStructOnScope> ScopedTrait = ModelNode->GetTraitInstance(TraitPin->GetFName()))
 					{
-						if (const FRigVMDecorator* RigVMTrait = (FRigVMDecorator*)ScopedTrait->GetStructMemory())
+						if (const FRigVMTrait* RigVMTrait = (FRigVMTrait*)ScopedTrait->GetStructMemory())
 						{
-							if (TWeakObjectPtr<UScriptStruct> TraitStruct = RigVMTrait->GetDecoratorSharedDataStruct(); TraitStruct.IsValid())
+							if (TWeakObjectPtr<UScriptStruct> TraitStruct = RigVMTrait->GetTraitSharedDataStruct(); TraitStruct.IsValid())
 							{
 								if (const FTrait* Trait = TraitRegistry.Find(TraitStruct.Get()))
 								{
diff --git a/Engine/Plugins/Experimental/AnimNext/Source/AnimNextUncookedOnly/Internal/Graph/AnimNextGraph_Controller.h b/Engine/Plugins/Experimental/AnimNext/Source/AnimNextUncookedOnly/Internal/Graph/AnimNextGraph_Controller.h
--- a/Engine/Plugins/Experimental/AnimNext/Source/AnimNextUncookedOnly/Internal/Graph/AnimNextGraph_Controller.h
+++ b/Engine/Plugins/Experimental/AnimNext/Source/AnimNextUncookedOnly/Internal/Graph/AnimNextGraph_Controller.h
@@ -28,20 +28,21 @@ class UAnimNextGraph_Controller : public URigVMController
 	GENERATED_BODY()
 
 public:
+	
 	// Adds a new Trait to the Stack, with default struct values
 	// Returns Trait Instance Name (or NAME_None on failure)
 	UFUNCTION(BlueprintCallable, Category = RigVMController)
-	ANIMNEXTUNCOOKEDONLY_API FName AddTrait(FName InNodeName, FName InNewTraitTypeName, int32 InPinIndex, const FString& InNewTraitDefaultValue = TEXT(""), bool bSetupUndoRedo = true, bool bPrintPythonCommand = false);
+	ANIMNEXTUNCOOKEDONLY_API FName AddTraitByName(FName InNodeName, FName InNewTraitTypeName, int32 InPinIndex, const FString& InNewTraitDefaultValue = TEXT(""), bool bSetupUndoRedo = true, bool bPrintPythonCommand = false);
 
 	// Removes a Trait from the Stack, using Trait Instance Name
 	// Returns operation success (true) or failure (false)
 	UFUNCTION(BlueprintCallable, Category = RigVMController)
-	ANIMNEXTUNCOOKEDONLY_API bool RemoveTrait(FName InNodeName, FName InTraitInstanceName, bool bSetupUndoRedo = true, bool bPrintPythonCommand = false);
+	ANIMNEXTUNCOOKEDONLY_API bool RemoveTraitByName(FName InNodeName, FName InTraitInstanceName, bool bSetupUndoRedo = true, bool bPrintPythonCommand = false);
 
 	// Swap a Trait from the Stack with a new one, using existing Trait Instance Name and new Trait Type Name
 	// Returns Trait Instance Name (or NAME_None on failure)
 	UFUNCTION(BlueprintCallable, Category = RigVMController)
-	ANIMNEXTUNCOOKEDONLY_API FName SwapTrait(FName InNodeName, FName InTraitInstanceName, int32 InCurrentTraitPinIndex, FName InNewTraitTypeName, const FString& InNewTraitDefaultValue = TEXT(""), bool bSetupUndoRedo = true, bool bPrintPythonCommand = false);
+	ANIMNEXTUNCOOKEDONLY_API FName SwapTraitByName(FName InNodeName, FName InTraitInstanceName, int32 InCurrentTraitPinIndex, FName InNewTraitTypeName, const FString& InNewTraitDefaultValue = TEXT(""), bool bSetupUndoRedo = true, bool bPrintPythonCommand = false);
 
 	// Move a Trait from its current PinIndex to the specified one (moving it visually in the stack)
 	// Returns operation success (true) or failure (false)
diff --git a/Engine/Plugins/Experimental/AnimNext/Source/AnimNextUncookedOnly/Private/Graph/AnimNextGraph_Controller.cpp b/Engine/Plugins/Experimental/AnimNext/Source/AnimNextUncookedOnly/Private/Graph/AnimNextGraph_Controller.cpp
--- a/Engine/Plugins/Experimental/AnimNext/Source/AnimNextUncookedOnly/Private/Graph/AnimNextGraph_Controller.cpp
+++ b/Engine/Plugins/Experimental/AnimNext/Source/AnimNextUncookedOnly/Private/Graph/AnimNextGraph_Controller.cpp
@@ -17,7 +17,7 @@
 #include "TraitCore/TraitRegistry.h"
 #include "TraitCore/TraitUID.h"
 
-FName UAnimNextGraph_Controller::AddTrait(FName InNodeName, FName InNewTraitTypeName, int32 InPinIndex, const FString& InNewTraitDefaultValue, bool bSetupUndoRedo, bool bPrintPythonCommand)
+FName UAnimNextGraph_Controller::AddTraitByName(FName InNodeName, FName InNewTraitTypeName, int32 InPinIndex, const FString& InNewTraitDefaultValue, bool bSetupUndoRedo, bool bPrintPythonCommand)
 {
 	if (!IsValidGraph())
 	{
@@ -79,10 +79,10 @@ FName UAnimNextGraph_Controller::AddTrait(FName InNodeName, FName InNewTraitType
 
 	const FName TraitName = InNewTraitTypeName;
 
-	return AddDecorator(InNodeName, *CppDecoratorStruct->GetPathName(), TraitName, DefaultValue, InPinIndex, bSetupUndoRedo, bPrintPythonCommand);
+	return AddTrait(InNodeName, *CppDecoratorStruct->GetPathName(), TraitName, DefaultValue, InPinIndex, bSetupUndoRedo, bPrintPythonCommand);
 }
 
-bool UAnimNextGraph_Controller::RemoveTrait(FName InNodeName, FName InTraitInstanceName, bool bSetupUndoRedo, bool bPrintPythonCommand)
+bool UAnimNextGraph_Controller::RemoveTraitByName(FName InNodeName, FName InTraitInstanceName, bool bSetupUndoRedo, bool bPrintPythonCommand)
 {
 	if (!IsValidGraph())
 	{
@@ -102,10 +102,10 @@ bool UAnimNextGraph_Controller::RemoveTrait(FName InNodeName, FName InTraitInsta
 
 	// Avoid multiple VM recompilations for internal operations
 	FRigVMControllerCompileBracketScope CompileScope(this);
-	return RemoveDecorator(InNodeName, InTraitInstanceName, bSetupUndoRedo, bPrintPythonCommand);
+	return RemoveTrait(InNodeName, InTraitInstanceName, bSetupUndoRedo, bPrintPythonCommand);
 }
 
-FName UAnimNextGraph_Controller::SwapTrait(FName InNodeName, FName InTraitInstanceName, int32 InCurrentTraitPinIndex, FName InNewTraitTypeName, const FString& InNewTraitDefaultValue, bool bSetupUndoRedo, bool bPrintPythonCommand)
+FName UAnimNextGraph_Controller::SwapTraitByName(FName InNodeName, FName InTraitInstanceName, int32 InCurrentTraitPinIndex, FName InNewTraitTypeName, const FString& InNewTraitDefaultValue, bool bSetupUndoRedo, bool bPrintPythonCommand)
 {
 	if (!IsValidGraph())
 	{
@@ -125,9 +125,9 @@ FName UAnimNextGraph_Controller::SwapTrait(FName InNodeName, FName InTraitInstan
 
 	// Avoid multiple VM recompilations, for each operation
 	FRigVMControllerCompileBracketScope CompileScope(this);
-	if (RemoveTrait(InNodeName, InTraitInstanceName, bSetupUndoRedo, bPrintPythonCommand))
+	if (RemoveTraitByName(InNodeName, InTraitInstanceName, bSetupUndoRedo, bPrintPythonCommand))
 	{
-		return AddTrait(InNodeName, InNewTraitTypeName, InCurrentTraitPinIndex, InNewTraitDefaultValue, bSetupUndoRedo, bPrintPythonCommand);
+		return AddTraitByName(InNodeName, InNewTraitTypeName, InCurrentTraitPinIndex, InNewTraitDefaultValue, bSetupUndoRedo, bPrintPythonCommand);
 	}
 
 	return NAME_None;
@@ -166,7 +166,7 @@ bool UAnimNextGraph_Controller::SetTraitPinIndex(FName InNodeName, FName InTrait
 		return false;
 	}
 
-	URigVMPin* TraitPin = Node->FindDecorator(InTraitInstanceName);
+	URigVMPin* TraitPin = Node->FindTrait(InTraitInstanceName);
 	if (TraitPin == nullptr)
 	{
 		ReportError(TEXT("The node does not contain a Trait with the provided name."));
@@ -177,20 +177,20 @@ bool UAnimNextGraph_Controller::SetTraitPinIndex(FName InNodeName, FName InTrait
 	const FString TraitDefaultValue = TraitPin->GetDefaultValue();
 
 	// TODO zzz : Is there a better way to get a Trait* from a TraitPin ?
-	if (TSharedPtr<FStructOnScope> ScopedTrait = Node->GetDecoratorInstance(TraitPin->GetFName()))
+	if (TSharedPtr<FStructOnScope> ScopedTrait = Node->GetTraitInstance(TraitPin->GetFName()))
 	{
-		const FRigVMDecorator* Decorator = (FRigVMDecorator*)ScopedTrait->GetStructMemory();
-		if (const UScriptStruct* TraitSharedInstanceData = Decorator->GetDecoratorSharedDataStruct())
+		const FRigVMTrait* Trait = (FRigVMTrait*)ScopedTrait->GetStructMemory();
+		if (const UScriptStruct* TraitSharedInstanceData = Trait->GetTraitSharedDataStruct())
 		{
 			UE::AnimNext::FTraitRegistry& TraitRegistry = UE::AnimNext::FTraitRegistry::Get();
-			if (const UE::AnimNext::FTrait* Trait = TraitRegistry.Find(TraitSharedInstanceData))
+			if (const UE::AnimNext::FTrait* AnimNextTrait = TraitRegistry.Find(TraitSharedInstanceData))
 			{
 				// Avoid multiple VM recompilations, for each operation
 				FRigVMControllerCompileBracketScope CompileScope(this);
-				if (RemoveTrait(InNodeName, InTraitInstanceName, bSetupUndoRedo, bPrintPythonCommand))
+				if (RemoveTraitByName(InNodeName, InTraitInstanceName, bSetupUndoRedo, bPrintPythonCommand))
 				{
 					// TOOO zzz : Why TraitName is a string ?
-					if (AddTrait(InNodeName, FName(Trait->GetTraitName()), InNewPinIndex, TraitDefaultValue, bSetupUndoRedo, bPrintPythonCommand) != NAME_None)
+					if (AddTraitByName(InNodeName, FName(AnimNextTrait->GetTraitName()), InNewPinIndex, TraitDefaultValue, bSetupUndoRedo, bPrintPythonCommand) != NAME_None)
 					{
 						return true;
 					}
diff --git a/Engine/Plugins/Experimental/AnimNext/Source/AnimNextUncookedOnly/Private/Graph/AnimNextGraph_EdGraphNode.cpp b/Engine/Plugins/Experimental/AnimNext/Source/AnimNextUncookedOnly/Private/Graph/AnimNextGraph_EdGraphNode.cpp
--- a/Engine/Plugins/Experimental/AnimNext/Source/AnimNextUncookedOnly/Private/Graph/AnimNextGraph_EdGraphNode.cpp
+++ b/Engine/Plugins/Experimental/AnimNext/Source/AnimNextUncookedOnly/Private/Graph/AnimNextGraph_EdGraphNode.cpp
@@ -41,14 +41,14 @@ void UAnimNextGraph_EdGraphNode::ConfigurePin(UEdGraphPin* EdGraphPin, const URi
 	{
 		if (const URigVMPin* DecoratorPin = ModelPin->GetParentPin())
 		{
-			if (DecoratorPin->IsDecoratorPin())
+			if (DecoratorPin->IsTraitPin())
 			{
 				check(DecoratorPin->GetScriptStruct() == FRigDecorator_AnimNextCppDecorator::StaticStruct());
 
-				TSharedPtr<FStructOnScope> DecoratorScope = DecoratorPin->GetDecoratorInstance();
+				TSharedPtr<FStructOnScope> DecoratorScope = DecoratorPin->GetTraitInstance();
 				const FRigDecorator_AnimNextCppDecorator* VMDecorator = (const FRigDecorator_AnimNextCppDecorator*)DecoratorScope->GetStructMemory();
 
-				const UScriptStruct* TraitStruct = VMDecorator->GetDecoratorSharedDataStruct();
+				const UScriptStruct* TraitStruct = VMDecorator->GetTraitSharedDataStruct();
 				check(TraitStruct != nullptr);
 
 				const FProperty* PinProperty = TraitStruct->FindPropertyByName(ModelPin->GetFName());
@@ -114,7 +114,7 @@ void UAnimNextGraph_EdGraphNode::BuildAddTraitContextMenu(UToolMenu* SubMenu)
 			FUIAction(FExecuteAction::CreateLambda(
 				[this, Trait, VMController, VMNode, CppDecoratorStruct, DefaultValue, DisplayName]()
 				{
-					VMController->AddDecorator(
+					VMController->AddTrait(
 						VMNode->GetFName(),
 						*CppDecoratorStruct->GetPathName(),
 						*DisplayName,
diff --git a/Engine/Plugins/Experimental/AnimNext/Source/AnimNextUncookedOnly/Private/Graph/AnimNextGraph_EditorData.cpp b/Engine/Plugins/Experimental/AnimNext/Source/AnimNextUncookedOnly/Private/Graph/AnimNextGraph_EditorData.cpp
--- a/Engine/Plugins/Experimental/AnimNext/Source/AnimNextUncookedOnly/Private/Graph/AnimNextGraph_EditorData.cpp
+++ b/Engine/Plugins/Experimental/AnimNext/Source/AnimNextUncookedOnly/Private/Graph/AnimNextGraph_EditorData.cpp
@@ -171,7 +171,7 @@ void UAnimNextGraph_EditorData::HandleModifiedEvent(ERigVMGraphNotifType InNotif
 		{
 			if (URigVMPin* Pin = Cast<URigVMPin>(InSubject))
 			{
-				if (Pin->IsDecoratorPin())
+				if (Pin->IsTraitPin())
 				{
 					RequestAutoVMRecompilation();
 				}
diff --git a/Engine/Plugins/Experimental/AnimNext/Source/AnimNextUncookedOnly/Private/UncookedOnlyUtils.cpp b/Engine/Plugins/Experimental/AnimNext/Source/AnimNextUncookedOnly/Private/UncookedOnlyUtils.cpp
--- a/Engine/Plugins/Experimental/AnimNext/Source/AnimNextUncookedOnly/Private/UncookedOnlyUtils.cpp
+++ b/Engine/Plugins/Experimental/AnimNext/Source/AnimNextUncookedOnly/Private/UncookedOnlyUtils.cpp
@@ -119,14 +119,14 @@ namespace Private
 		const TArray<URigVMPin*>& Pins = DecoratorStackNode->GetPins();
 		for (URigVMPin* Pin : Pins)
 		{
-			if (!Pin->IsDecoratorPin())
+			if (!Pin->IsTraitPin())
 			{
 				continue;	// Not a decorator pin
 			}
 
 			if (Pin->GetScriptStruct() == FRigDecorator_AnimNextCppDecorator::StaticStruct())
 			{
-				TSharedPtr<FStructOnScope> DecoratorScope = Pin->GetDecoratorInstance();
+				TSharedPtr<FStructOnScope> DecoratorScope = Pin->GetTraitInstance();
 				FRigDecorator_AnimNextCppDecorator* VMDecorator = (FRigDecorator_AnimNextCppDecorator*)DecoratorScope->GetStructMemory();
 
 				if (const FTrait* Trait = VMDecorator->GetTrait())
@@ -195,7 +195,7 @@ namespace Private
 							// If this is the case, we bind to the first trait index since we only allowed a single base trait per stack
 							// Otherwise we lookup the trait index we are linked to
 							const URigVMPin* SourceDecoratorPin = PinLinks[0]->GetSourcePin()->GetParentPin();
-							SourceTraitIndex = SourceDecoratorPin != nullptr ? SourceTraitStack->DecoratorStackNode->GetDecoratorPins().IndexOfByKey(SourceDecoratorPin) : 0;
+							SourceTraitIndex = SourceDecoratorPin != nullptr ? SourceTraitStack->DecoratorStackNode->GetTraitPins().IndexOfByKey(SourceDecoratorPin) : 0;
 						}
 
 						if (SourceNodeHandle.IsValid())
@@ -317,7 +317,7 @@ namespace Private
 				Prop->ExportText_Direct(DefaultValue, &CppDecoratorStructInstance, &DefaultCppDecoratorStructInstance, nullptr, PPF_SerializedAsImportText);
 			}
 
-			const FName ReferencePoseDecoratorName = VMController->AddDecorator(VMReferencePoseNode->GetFName(), *CppDecoratorStruct->GetPathName(), TEXT("ReferencePose"), DefaultValue, INDEX_NONE, false, false);
+			const FName ReferencePoseDecoratorName = VMController->AddTrait(VMReferencePoseNode->GetFName(), *CppDecoratorStruct->GetPathName(), TEXT("ReferencePose"), DefaultValue, INDEX_NONE, false, false);
 			check(!ReferencePoseDecoratorName.IsNone());
 
 			URigVMPin* OutputPin = VMReferencePoseNode->FindPin(GET_MEMBER_NAME_STRING_CHECKED(FRigUnit_AnimNextTraitStack, Result));
@@ -416,7 +416,7 @@ namespace Private
 				Prop->ExportText_Direct(DefaultValue, &CppDecoratorStructInstance, &DefaultCppDecoratorStructInstance, nullptr, PPF_SerializedAsImportText);
 			}
 
-			VMController->AddDecorator(VMNode->GetFName(), *CppDecoratorStruct->GetPathName(), TEXT("ReferencePose"), DefaultValue, INDEX_NONE, false, false);
+			VMController->AddTrait(VMNode->GetFName(), *CppDecoratorStruct->GetPathName(), TEXT("ReferencePose"), DefaultValue, INDEX_NONE, false, false);
 
 			FTraitStackMapping Mapping(VMNode);
 			ForEachTraitInStack(VMNode,
@@ -661,7 +661,7 @@ void FUtils::CompileVM(UAnimNextGraph* InGraph)
 	URigVMCompiler* Compiler = URigVMCompiler::StaticClass()->GetDefaultObject<URigVMCompiler>();
 	EditorData->VMCompileSettings.SetExecuteContextStruct(FAnimNextExecuteContext::StaticStruct());
 	FRigVMCompileSettings Settings = (EditorData->bCompileInDebugMode) ? FRigVMCompileSettings::Fast(EditorData->VMCompileSettings.GetExecuteContextStruct()) : EditorData->VMCompileSettings;
-	Settings.ASTSettings.bSetupDecorators = false; // disable the default implementation of decorators for now
+	Settings.ASTSettings.bSetupTraits = false; // disable the default implementation of decorators for now
 	Settings.ASTSettings.ReportDelegate.BindLambda([InGraph](EMessageSeverity::Type InType, UObject* InObject, const FString& InString)
 	{
 		FMessageLog("AnimNextCompilerResults").Message(InType, FText::FromString(InString));
diff --git a/Engine/Plugins/Experimental/AnimNextAnimGraph/Source/AnimNextAnimGraphTestSuite/Private/AnimNextAnimGraphTraitGraphTest.cpp b/Engine/Plugins/Experimental/AnimNextAnimGraph/Source/AnimNextAnimGraphTestSuite/Private/AnimNextAnimGraphTraitGraphTest.cpp
--- a/Engine/Plugins/Experimental/AnimNextAnimGraph/Source/AnimNextAnimGraphTestSuite/Private/AnimNextAnimGraphTraitGraphTest.cpp
+++ b/Engine/Plugins/Experimental/AnimNextAnimGraph/Source/AnimNextAnimGraphTestSuite/Private/AnimNextAnimGraphTraitGraphTest.cpp
@@ -115,12 +115,12 @@ bool FAnimationAnimNextEditorTest_GraphAddTrait::RunTest(const FString& InParame
 		UE_RETURN_ON_ERROR(Controller != nullptr, "FAnimationAnimNextEditorTest_GraphAddTrait -> Failed to get RigVM controller");
 
 		// Create an empty trait stack node
-		URigVMUnitNode* DecoratorStackNode = Controller->AddUnitNode(FRigUnit_AnimNextTraitStack::StaticStruct(), FRigVMStruct::ExecuteName, FVector2D(0.0f, 0.0f), FString(), false);
-		UE_RETURN_ON_ERROR(DecoratorStackNode != nullptr, "FAnimationAnimNextRuntimeTest_GraphAddDecorator -> Failed to create decorator stack node");
+		URigVMUnitNode* TraitStackNode = Controller->AddUnitNode(FRigUnit_AnimNextTraitStack::StaticStruct(), FRigVMStruct::ExecuteName, FVector2D(0.0f, 0.0f), FString(), false);
+		UE_RETURN_ON_ERROR(TraitStackNode != nullptr, "FAnimationAnimNextRuntimeTest_GraphAddTrait -> Failed to create trait stack node");
 
 		// Add a trait
-		const UScriptStruct* CppDecoratorStruct = FRigDecorator_AnimNextCppDecorator::StaticStruct();
-		UE_RETURN_ON_ERROR(CppDecoratorStruct != nullptr, "FAnimationAnimNextRuntimeTest_GraphAddDecorator -> Failed to get find Cpp decorator static struct");
+		const UScriptStruct* CppTraitStruct = FRigDecorator_AnimNextCppDecorator::StaticStruct();
+		UE_RETURN_ON_ERROR(CppTraitStruct != nullptr, "FAnimationAnimNextRuntimeTest_GraphAddDecorator -> Failed to get find Cpp trait static struct");
 
 		const FTrait* Trait = FTraitRegistry::Get().Find(FTestTrait::TraitUID);
 		UE_RETURN_ON_ERROR(Trait != nullptr, "FAnimationAnimNextEditorTest_GraphAddTrait -> Failed to find test trait");
@@ -134,7 +134,7 @@ bool FAnimationAnimNextEditorTest_GraphAddTrait::RunTest(const FString& InParame
 			FRigDecorator_AnimNextCppDecorator CppDecoratorStructInstance;
 			CppDecoratorStructInstance.DecoratorSharedDataStruct = ScriptStruct;
 
-			UE_RETURN_ON_ERROR(CppDecoratorStructInstance.CanBeAddedToNode(DecoratorStackNode, nullptr), "FAnimationAnimNextEditorTest_GraphAddTrait -> Trait cannot be added to decorator stack node");
+			UE_RETURN_ON_ERROR(CppDecoratorStructInstance.CanBeAddedToNode(TraitStackNode, nullptr), "FAnimationAnimNextEditorTest_GraphAddTrait -> Trait cannot be added to trait stack node");
 
 			const FProperty* Prop = FAnimNextCppDecoratorWrapper::StaticStruct()->FindPropertyByName(GET_MEMBER_NAME_STRING_CHECKED(FAnimNextCppDecoratorWrapper, CppDecorator));
 			UE_RETURN_ON_ERROR(Prop != nullptr, "FAnimationAnimNextEditorTest_GraphAddTrait -> Failed to find wrapper property");
@@ -146,49 +146,49 @@ bool FAnimationAnimNextEditorTest_GraphAddTrait::RunTest(const FString& InParame
 		ScriptStruct->GetStringMetaDataHierarchical(FRigVMStruct::DisplayNameMetaName, &DisplayNameMetadata);
 		const FString DisplayName = DisplayNameMetadata.IsEmpty() ? Trait->GetTraitName() : DisplayNameMetadata;
 
-		const FName DecoratorName = Controller->AddDecorator(
-			DecoratorStackNode->GetFName(),
-			*CppDecoratorStruct->GetPathName(),
+		const FName TraitName = Controller->AddTrait(
+			TraitStackNode->GetFName(),
+			*CppTraitStruct->GetPathName(),
 			*DisplayName,
 			DefaultValue, INDEX_NONE, true, true);
-		UE_RETURN_ON_ERROR(DecoratorName == DisplayName, TEXT("FAnimationAnimNextEditorTest_GraphAddTrait -> Unexpected decorator name"));
+		UE_RETURN_ON_ERROR(TraitName == DisplayName, TEXT("FAnimationAnimNextEditorTest_GraphAddTrait -> Unexpected trait name"));
 
-		URigVMPin* DecoratorPin = DecoratorStackNode->FindPin(*DisplayName);
-		UE_RETURN_ON_ERROR(DecoratorPin != nullptr, TEXT("FAnimationAnimNextEditorTest_GraphAddTrait -> Failed to find decorator pin"));
+		URigVMPin* TraitPin = TraitStackNode->FindPin(*DisplayName);
+		UE_RETURN_ON_ERROR(TraitPin != nullptr, TEXT("FAnimationAnimNextEditorTest_GraphAddTrait -> Failed to find trait pin"));
 
-		// Our first pin is the hard coded output result, decorator pins follow
-		UE_RETURN_ON_ERROR(DecoratorStackNode->GetPins().Num() == 2, TEXT("FAnimationAnimNextEditorTest_GraphAddTrait -> Unexpected number of pins"));
-		UE_RETURN_ON_ERROR(DecoratorPin->IsDecoratorPin() == true, TEXT("FAnimationAnimNextEditorTest_GraphAddTrait -> Unexpected pin type"));
-		UE_RETURN_ON_ERROR(DecoratorPin->GetFName() == DecoratorName, TEXT("FAnimationAnimNextEditorTest_GraphAddTrait -> Unexpected pin name"));
-		UE_RETURN_ON_ERROR(DecoratorPin->GetCPPTypeObject() == FRigDecorator_AnimNextCppDecorator::StaticStruct(), TEXT("FAnimationAnimNextEditorTest_GraphAddTrait -> Unexpected pin type"));
+		// Our first pin is the hard coded output result, trait pins follow
+		UE_RETURN_ON_ERROR(TraitStackNode->GetPins().Num() == 2, TEXT("FAnimationAnimNextEditorTest_GraphAddTrait -> Unexpected number of pins"));
+		UE_RETURN_ON_ERROR(TraitPin->IsTraitPin() == true, TEXT("FAnimationAnimNextEditorTest_GraphAddTrait -> Unexpected pin type"));
+		UE_RETURN_ON_ERROR(TraitPin->GetFName() == TraitName, TEXT("FAnimationAnimNextEditorTest_GraphAddTrait -> Unexpected pin name"));
+		UE_RETURN_ON_ERROR(TraitPin->GetCPPTypeObject() == FRigDecorator_AnimNextCppDecorator::StaticStruct(), TEXT("FAnimationAnimNextEditorTest_GraphAddTrait -> Unexpected pin type"));
 
-		// Our first sub-pin is the hard coded script struct member that parametrizes the decorator, dynamic decorator sub-pins follow
-		UE_RETURN_ON_ERROR(DecoratorPin->GetSubPins().Num() == 6, TEXT("FAnimationAnimNextEditorTest_GraphAddTrait -> Unexpected decorator sub pins"));
+		// Our first sub-pin is the hard coded script struct member that parametrizes the trait, dynamic trait sub-pins follow
+		UE_RETURN_ON_ERROR(TraitPin->GetSubPins().Num() == 6, TEXT("FAnimationAnimNextEditorTest_GraphAddTrait -> Unexpected trait sub pins"));
 
 		// SomeInt32
-		UE_RETURN_ON_ERROR(DecoratorPin->GetSubPins()[1]->GetCPPType() == TEXT("int32"), TEXT("FAnimationAnimNextEditorTest_GraphAddTrait -> Unexpected decorator pin type"));
-		UE_RETURN_ON_ERROR(DecoratorPin->GetSubPins()[1]->GetDefaultValue() == TEXT("3"), TEXT("FAnimationAnimNextEditorTest_GraphAddTrait -> Unexpected decorator pin value"));
-		UE_RETURN_ON_ERROR(!DecoratorPin->GetSubPins()[1]->IsLazy(), TEXT("FAnimationAnimNextEditorTest_GraphAddTrait -> Expected non-lazy decorator pin"));
+		UE_RETURN_ON_ERROR(TraitPin->GetSubPins()[1]->GetCPPType() == TEXT("int32"), TEXT("FAnimationAnimNextEditorTest_GraphAddTrait -> Unexpected trait pin type"));
+		UE_RETURN_ON_ERROR(TraitPin->GetSubPins()[1]->GetDefaultValue() == TEXT("3"), TEXT("FAnimationAnimNextEditorTest_GraphAddTrait -> Unexpected trait pin value"));
+		UE_RETURN_ON_ERROR(!TraitPin->GetSubPins()[1]->IsLazy(), TEXT("FAnimationAnimNextEditorTest_GraphAddTrait -> Expected non-lazy trait pin"));
 
 		// SomeFloat
-		UE_RETURN_ON_ERROR(DecoratorPin->GetSubPins()[2]->GetCPPType() == TEXT("float"), TEXT("FAnimationAnimNextEditorTest_GraphAddTrait -> Unexpected decorator pin type"));
-		UE_RETURN_ON_ERROR(DecoratorPin->GetSubPins()[2]->GetDefaultValue() == TEXT("34.000000"), TEXT("FAnimationAnimNextEditorTest_GraphAddTrait -> Unexpected decorator pin value"));
-		UE_RETURN_ON_ERROR(!DecoratorPin->GetSubPins()[2]->IsLazy(), TEXT("FAnimationAnimNextEditorTest_GraphAddTrait -> Expected non-lazy decorator pin"));
+		UE_RETURN_ON_ERROR(TraitPin->GetSubPins()[2]->GetCPPType() == TEXT("float"), TEXT("FAnimationAnimNextEditorTest_GraphAddTrait -> Unexpected trait pin type"));
+		UE_RETURN_ON_ERROR(TraitPin->GetSubPins()[2]->GetDefaultValue() == TEXT("34.000000"), TEXT("FAnimationAnimNextEditorTest_GraphAddTrait -> Unexpected trait pin value"));
+		UE_RETURN_ON_ERROR(!TraitPin->GetSubPins()[2]->IsLazy(), TEXT("FAnimationAnimNextEditorTest_GraphAddTrait -> Expected non-lazy trait pin"));
 
 		// SomeLatentInt32
-		UE_RETURN_ON_ERROR(DecoratorPin->GetSubPins()[3]->GetCPPType() == TEXT("int32"), TEXT("FAnimationAnimNextEditorTest_GraphAddTrait -> Unexpected decorator pin type"));
-		UE_RETURN_ON_ERROR(DecoratorPin->GetSubPins()[3]->GetDefaultValue() == TEXT("3"), TEXT("FAnimationAnimNextEditorTest_GraphAddTrait -> Unexpected decorator pin value"));
-		UE_RETURN_ON_ERROR(DecoratorPin->GetSubPins()[3]->IsLazy(), TEXT("FAnimationAnimNextEditorTest_GraphAddTrait -> Expected lazy decorator pin"));
+		UE_RETURN_ON_ERROR(TraitPin->GetSubPins()[3]->GetCPPType() == TEXT("int32"), TEXT("FAnimationAnimNextEditorTest_GraphAddTrait -> Unexpected trait pin type"));
+		UE_RETURN_ON_ERROR(TraitPin->GetSubPins()[3]->GetDefaultValue() == TEXT("3"), TEXT("FAnimationAnimNextEditorTest_GraphAddTrait -> Unexpected trait pin value"));
+		UE_RETURN_ON_ERROR(TraitPin->GetSubPins()[3]->IsLazy(), TEXT("FAnimationAnimNextEditorTest_GraphAddTrait -> Expected lazy trait pin"));
 
 		// SomeOtherLatentInt32
-		UE_RETURN_ON_ERROR(DecoratorPin->GetSubPins()[4]->GetCPPType() == TEXT("int32"), TEXT("FAnimationAnimNextEditorTest_GraphAddTrait -> Unexpected decorator pin type"));
-		UE_RETURN_ON_ERROR(DecoratorPin->GetSubPins()[4]->GetDefaultValue() == TEXT("3"), TEXT("FAnimationAnimNextEditorTest_GraphAddTrait -> Unexpected decorator pin value"));
-		UE_RETURN_ON_ERROR(DecoratorPin->GetSubPins()[4]->IsLazy(), TEXT("FAnimationAnimNextEditorTest_GraphAddTrait -> Expected lazy decorator pin"));
+		UE_RETURN_ON_ERROR(TraitPin->GetSubPins()[4]->GetCPPType() == TEXT("int32"), TEXT("FAnimationAnimNextEditorTest_GraphAddTrait -> Unexpected trait pin type"));
+		UE_RETURN_ON_ERROR(TraitPin->GetSubPins()[4]->GetDefaultValue() == TEXT("3"), TEXT("FAnimationAnimNextEditorTest_GraphAddTrait -> Unexpected trait pin value"));
+		UE_RETURN_ON_ERROR(TraitPin->GetSubPins()[4]->IsLazy(), TEXT("FAnimationAnimNextEditorTest_GraphAddTrait -> Expected lazy trait pin"));
 
 		// SomeLatentFloat
-		UE_RETURN_ON_ERROR(DecoratorPin->GetSubPins()[5]->GetCPPType() == TEXT("float"), TEXT("FAnimationAnimNextEditorTest_GraphAddTrait -> Unexpected decorator pin type"));
-		UE_RETURN_ON_ERROR(DecoratorPin->GetSubPins()[5]->GetDefaultValue() == TEXT("34.000000"), TEXT("FAnimationAnimNextEditorTest_GraphAddTrait -> Unexpected decorator pin value"));
-		UE_RETURN_ON_ERROR(DecoratorPin->GetSubPins()[5]->IsLazy(), TEXT("FAnimationAnimNextEditorTest_GraphAddTrait -> Expected lazy decorator pin"));
+		UE_RETURN_ON_ERROR(TraitPin->GetSubPins()[5]->GetCPPType() == TEXT("float"), TEXT("FAnimationAnimNextEditorTest_GraphAddTrait -> Unexpected trait pin type"));
+		UE_RETURN_ON_ERROR(TraitPin->GetSubPins()[5]->GetDefaultValue() == TEXT("34.000000"), TEXT("FAnimationAnimNextEditorTest_GraphAddTrait -> Unexpected trait pin value"));
+		UE_RETURN_ON_ERROR(TraitPin->GetSubPins()[5]->IsLazy(), TEXT("FAnimationAnimNextEditorTest_GraphAddTrait -> Expected lazy trait pin"));
 	}
 
 	Tests::FUtils::CleanupAfterTests();
@@ -222,7 +222,7 @@ bool FAnimationAnimNextEditorTest_GraphTraitOperations::RunTest(const FString& I
 
 		// Create an empty trait stack node
 		URigVMUnitNode* TraitStackNode = Controller->AddUnitNode(FRigUnit_AnimNextTraitStack::StaticStruct(), FRigVMStruct::ExecuteName, FVector2D(0.0f, 0.0f), FString(), false);
-		UE_RETURN_ON_ERROR(TraitStackNode != nullptr, "FAnimationAnimNextEditorTest_GraphTraitOperations -> Failed to create decorator stack node");
+		UE_RETURN_ON_ERROR(TraitStackNode != nullptr, "FAnimationAnimNextEditorTest_GraphTraitOperations -> Failed to create trait stack node");
 
 		const TArray<URigVMPin*>& NodePins = TraitStackNode->GetPins();
 
@@ -235,45 +235,45 @@ bool FAnimationAnimNextEditorTest_GraphTraitOperations::RunTest(const FString& I
 
 			const FName TraitTypeName = *Trait->GetTraitName();
 
-			TraitInstanceName = Controller->AddTrait(TraitStackNode->GetFName(), TraitTypeName, INDEX_NONE);
+			TraitInstanceName = Controller->AddTraitByName(TraitStackNode->GetFName(), TraitTypeName, INDEX_NONE);
 			UE_RETURN_ON_ERROR(TraitInstanceName == TraitTypeName, TEXT("FAnimationAnimNextEditorTest_GraphTraitOperations -> Unexpected Trait name"));
 
 			URigVMPin* TraitPin = TraitStackNode->FindPin(TraitInstanceName.ToString());
 			UE_RETURN_ON_ERROR(TraitPin != nullptr, TEXT("FAnimationAnimNextEditorTest_GraphTraitOperations -> Failed to find Trait pin"));
 
-			// Our first pin is the hard coded output result, decorator pins follow
+			// Our first pin is the hard coded output result, trait pins follow
 			UE_RETURN_ON_ERROR(NodePins.Num() == 2, TEXT("FAnimationAnimNextEditorTest_GraphTraitOperations -> Unexpected number of pins"));
-			UE_RETURN_ON_ERROR(TraitPin->IsDecoratorPin() == true, TEXT("FAnimationAnimNextEditorTest_GraphTraitOperations -> Unexpected pin type"));
+			UE_RETURN_ON_ERROR(TraitPin->IsTraitPin() == true, TEXT("FAnimationAnimNextEditorTest_GraphTraitOperations -> Unexpected pin type"));
 			UE_RETURN_ON_ERROR(TraitPin->GetFName() == TraitInstanceName, TEXT("FAnimationAnimNextEditorTest_GraphTraitOperations -> Unexpected pin name"));
 			UE_RETURN_ON_ERROR(TraitPin->GetCPPTypeObject() == FRigDecorator_AnimNextCppDecorator::StaticStruct(), TEXT("FAnimationAnimNextEditorTest_GraphTraitOperations -> Unexpected pin type"));
 
-			// Our first sub-pin is the hard coded script struct member that parametrizes the decorator, dynamic decorator sub-pins follow
-			UE_RETURN_ON_ERROR(TraitPin->GetSubPins().Num() == 6, TEXT("FAnimationAnimNextEditorTest_GraphTraitOperations -> Unexpected decorator sub pins"));
+			// Our first sub-pin is the hard coded script struct member that parametrizes the trait, dynamic trait sub-pins follow
+			UE_RETURN_ON_ERROR(TraitPin->GetSubPins().Num() == 6, TEXT("FAnimationAnimNextEditorTest_GraphTraitOperations -> Unexpected trait sub pins"));
 
 			// SomeInt32
-			UE_RETURN_ON_ERROR(TraitPin->GetSubPins()[1]->GetCPPType() == TEXT("int32"), TEXT("FAnimationAnimNextEditorTest_GraphTraitOperations -> Unexpected decorator pin type"));
-			UE_RETURN_ON_ERROR(TraitPin->GetSubPins()[1]->GetDefaultValue() == TEXT("3"), TEXT("FAnimationAnimNextEditorTest_GraphTraitOperations -> Unexpected decorator pin value"));
-			UE_RETURN_ON_ERROR(!TraitPin->GetSubPins()[1]->IsLazy(), TEXT("FAnimationAnimNextEditorTest_GraphTraitOperations -> Expected non-lazy decorator pin"));
+			UE_RETURN_ON_ERROR(TraitPin->GetSubPins()[1]->GetCPPType() == TEXT("int32"), TEXT("FAnimationAnimNextEditorTest_GraphTraitOperations -> Unexpected trait pin type"));
+			UE_RETURN_ON_ERROR(TraitPin->GetSubPins()[1]->GetDefaultValue() == TEXT("3"), TEXT("FAnimationAnimNextEditorTest_GraphTraitOperations -> Unexpected trait pin value"));
+			UE_RETURN_ON_ERROR(!TraitPin->GetSubPins()[1]->IsLazy(), TEXT("FAnimationAnimNextEditorTest_GraphTraitOperations -> Expected non-lazy trait pin"));
 
 			// SomeFloat
-			UE_RETURN_ON_ERROR(TraitPin->GetSubPins()[2]->GetCPPType() == TEXT("float"), TEXT("FAnimationAnimNextEditorTest_GraphTraitOperations -> Unexpected decorator pin type"));
-			UE_RETURN_ON_ERROR(TraitPin->GetSubPins()[2]->GetDefaultValue() == TEXT("34.000000"), TEXT("FAnimationAnimNextEditorTest_GraphTraitOperations -> Unexpected decorator pin value"));
-			UE_RETURN_ON_ERROR(!TraitPin->GetSubPins()[2]->IsLazy(), TEXT("FAnimationAnimNextEditorTest_GraphTraitOperations -> Expected non-lazy decorator pin"));
+			UE_RETURN_ON_ERROR(TraitPin->GetSubPins()[2]->GetCPPType() == TEXT("float"), TEXT("FAnimationAnimNextEditorTest_GraphTraitOperations -> Unexpected trait pin type"));
+			UE_RETURN_ON_ERROR(TraitPin->GetSubPins()[2]->GetDefaultValue() == TEXT("34.000000"), TEXT("FAnimationAnimNextEditorTest_GraphTraitOperations -> Unexpected trait pin value"));
+			UE_RETURN_ON_ERROR(!TraitPin->GetSubPins()[2]->IsLazy(), TEXT("FAnimationAnimNextEditorTest_GraphTraitOperations -> Expected non-lazy trait pin"));
 
 			// SomeLatentInt32
-			UE_RETURN_ON_ERROR(TraitPin->GetSubPins()[3]->GetCPPType() == TEXT("int32"), TEXT("FAnimationAnimNextEditorTest_GraphTraitOperations -> Unexpected decorator pin type"));
-			UE_RETURN_ON_ERROR(TraitPin->GetSubPins()[3]->GetDefaultValue() == TEXT("3"), TEXT("FAnimationAnimNextEditorTest_GraphTraitOperations -> Unexpected decorator pin value"));
-			UE_RETURN_ON_ERROR(TraitPin->GetSubPins()[3]->IsLazy(), TEXT("FAnimationAnimNextEditorTest_GraphTraitOperations -> Expected lazy decorator pin"));
+			UE_RETURN_ON_ERROR(TraitPin->GetSubPins()[3]->GetCPPType() == TEXT("int32"), TEXT("FAnimationAnimNextEditorTest_GraphTraitOperations -> Unexpected trait pin type"));
+			UE_RETURN_ON_ERROR(TraitPin->GetSubPins()[3]->GetDefaultValue() == TEXT("3"), TEXT("FAnimationAnimNextEditorTest_GraphTraitOperations -> Unexpected trait pin value"));
+			UE_RETURN_ON_ERROR(TraitPin->GetSubPins()[3]->IsLazy(), TEXT("FAnimationAnimNextEditorTest_GraphTraitOperations -> Expected lazy trait pin"));
 
 			// SomeOtherLatentInt32
-			UE_RETURN_ON_ERROR(TraitPin->GetSubPins()[4]->GetCPPType() == TEXT("int32"), TEXT("FAnimationAnimNextEditorTest_GraphTraitOperations -> Unexpected decorator pin type"));
-			UE_RETURN_ON_ERROR(TraitPin->GetSubPins()[4]->GetDefaultValue() == TEXT("3"), TEXT("FAnimationAnimNextEditorTest_GraphTraitOperations -> Unexpected decorator pin value"));
-			UE_RETURN_ON_ERROR(TraitPin->GetSubPins()[4]->IsLazy(), TEXT("FAnimationAnimNextEditorTest_GraphTraitOperations -> Expected lazy decorator pin"));
+			UE_RETURN_ON_ERROR(TraitPin->GetSubPins()[4]->GetCPPType() == TEXT("int32"), TEXT("FAnimationAnimNextEditorTest_GraphTraitOperations -> Unexpected trait pin type"));
+			UE_RETURN_ON_ERROR(TraitPin->GetSubPins()[4]->GetDefaultValue() == TEXT("3"), TEXT("FAnimationAnimNextEditorTest_GraphTraitOperations -> Unexpected trait pin value"));
+			UE_RETURN_ON_ERROR(TraitPin->GetSubPins()[4]->IsLazy(), TEXT("FAnimationAnimNextEditorTest_GraphTraitOperations -> Expected lazy trait pin"));
 
 			// SomeLatentFloat
-			UE_RETURN_ON_ERROR(TraitPin->GetSubPins()[5]->GetCPPType() == TEXT("float"), TEXT("FAnimationAnimNextEditorTest_GraphTraitOperations -> Unexpected decorator pin type"));
-			UE_RETURN_ON_ERROR(TraitPin->GetSubPins()[5]->GetDefaultValue() == TEXT("34.000000"), TEXT("FAnimationAnimNextEditorTest_GraphTraitOperations -> Unexpected decorator pin value"));
-			UE_RETURN_ON_ERROR(TraitPin->GetSubPins()[5]->IsLazy(), TEXT("FAnimationAnimNextEditorTest_GraphTraitOperations -> Expected lazy decorator pin"));
+			UE_RETURN_ON_ERROR(TraitPin->GetSubPins()[5]->GetCPPType() == TEXT("float"), TEXT("FAnimationAnimNextEditorTest_GraphTraitOperations -> Unexpected trait pin type"));
+			UE_RETURN_ON_ERROR(TraitPin->GetSubPins()[5]->GetDefaultValue() == TEXT("34.000000"), TEXT("FAnimationAnimNextEditorTest_GraphTraitOperations -> Unexpected trait pin value"));
+			UE_RETURN_ON_ERROR(TraitPin->GetSubPins()[5]->IsLazy(), TEXT("FAnimationAnimNextEditorTest_GraphTraitOperations -> Expected lazy trait pin"));
 		}
 
 		// --- Undo Add Trait ---
@@ -283,10 +283,10 @@ bool FAnimationAnimNextEditorTest_GraphTraitOperations::RunTest(const FString& I
 			URigVMPin* TraitPin = TraitStackNode->FindPin(TraitInstanceName.ToString());
 			UE_RETURN_ON_ERROR(TraitPin == nullptr, TEXT("FAnimationAnimNextEditorTest_GraphTraitOperations -> Undo AddTrait failed, Trait pin is still present"));
 
-			// Our first pin is the hard coded output result, decorator pins follow
+			// Our first pin is the hard coded output result, trait pins follow
 			const URigVMPin* FirstPin = NodePins[0];
 			UE_RETURN_ON_ERROR(NodePins.Num() == 1, TEXT("FAnimationAnimNextEditorTest_GraphTraitOperations -> Unexpected number of pins"));
-			UE_RETURN_ON_ERROR(FirstPin->IsDecoratorPin() == false, TEXT("FAnimationAnimNextEditorTest_GraphTraitOperations -> Unexpected pin type"));
+			UE_RETURN_ON_ERROR(FirstPin->IsTraitPin() == false, TEXT("FAnimationAnimNextEditorTest_GraphTraitOperations -> Unexpected pin type"));
 		}
 
 		// --- Redo Add Trait ---
@@ -298,16 +298,16 @@ bool FAnimationAnimNextEditorTest_GraphTraitOperations::RunTest(const FString& I
 
 			// Our first pin is the hard coded output result, trait pins follow
 			UE_RETURN_ON_ERROR(NodePins.Num() == 2, TEXT("FAnimationAnimNextEditorTest_GraphTraitOperations -> Unexpected number of pins"));
-			UE_RETURN_ON_ERROR(TraitPin->IsDecoratorPin() == true, TEXT("FAnimationAnimNextEditorTest_GraphTraitOperations -> Unexpected pin type"));
+			UE_RETURN_ON_ERROR(TraitPin->IsTraitPin() == true, TEXT("FAnimationAnimNextEditorTest_GraphTraitOperations -> Unexpected pin type"));
 			UE_RETURN_ON_ERROR(TraitPin->GetFName() == TraitInstanceName, TEXT("FAnimationAnimNextEditorTest_GraphTraitOperations -> Unexpected pin name"));
 			UE_RETURN_ON_ERROR(TraitPin->GetCPPTypeObject() == FRigDecorator_AnimNextCppDecorator::StaticStruct(), TEXT("FAnimationAnimNextEditorTest_GraphTraitOperations -> Unexpected pin type"));
 		}
 
 		// --- Remove the created trait ---
 		{
-			Controller->RemoveTrait(TraitStackNode->GetFName(), TraitInstanceName);
+			Controller->RemoveTraitByName(TraitStackNode->GetFName(), TraitInstanceName);
 
-			// Our first pin is the hard coded output result, decorator pins follow
+			// Our first pin is the hard coded output result, trait pins follow
 			UE_RETURN_ON_ERROR(NodePins.Num() == 1, TEXT("FAnimationAnimNextEditorTest_GraphTraitOperations -> Unexpected number of pins"));
 
 			URigVMPin* DeletedTraitPin = TraitStackNode->FindPin(TraitInstanceName.ToString());
@@ -315,7 +315,7 @@ bool FAnimationAnimNextEditorTest_GraphTraitOperations::RunTest(const FString& I
 
 			// Only the output result pin should be in the pin array
 			const URigVMPin* FirstPin = NodePins[0];
-			UE_RETURN_ON_ERROR(FirstPin->IsDecoratorPin() == false, TEXT("FAnimationAnimNextEditorTest_GraphTraitOperations -> Unexpected pin type"));
+			UE_RETURN_ON_ERROR(FirstPin->IsTraitPin() == false, TEXT("FAnimationAnimNextEditorTest_GraphTraitOperations -> Unexpected pin type"));
 			UE_RETURN_ON_ERROR(FirstPin->GetFName() != TraitInstanceName, TEXT("FAnimationAnimNextEditorTest_GraphTraitOperations -> Unexpected pin name"));
 		}
 
@@ -326,9 +326,9 @@ bool FAnimationAnimNextEditorTest_GraphTraitOperations::RunTest(const FString& I
 			URigVMPin* TraitPin = TraitStackNode->FindPin(TraitInstanceName.ToString());
 			UE_RETURN_ON_ERROR(TraitPin != nullptr, TEXT("FAnimationAnimNextEditorTest_GraphTraitOperations -> Undo failed, unable to find Trait pin"));
 
-			// Our first pin is the hard coded output result, decorator pins follow
+			// Our first pin is the hard coded output result, trait pins follow
 			UE_RETURN_ON_ERROR(NodePins.Num() == 2, TEXT("FAnimationAnimNextEditorTest_GraphTraitOperations -> Unexpected number of pins"));
-			UE_RETURN_ON_ERROR(TraitPin->IsDecoratorPin() == true, TEXT("FAnimationAnimNextEditorTest_GraphTraitOperations -> Unexpected pin type"));
+			UE_RETURN_ON_ERROR(TraitPin->IsTraitPin() == true, TEXT("FAnimationAnimNextEditorTest_GraphTraitOperations -> Unexpected pin type"));
 			UE_RETURN_ON_ERROR(TraitPin->GetFName() == TraitInstanceName, TEXT("FAnimationAnimNextEditorTest_GraphTraitOperations -> Unexpected pin name"));
 			UE_RETURN_ON_ERROR(TraitPin->GetCPPTypeObject() == FRigDecorator_AnimNextCppDecorator::StaticStruct(), TEXT("FAnimationAnimNextEditorTest_GraphTraitOperations -> Unexpected pin type"));
 		}
@@ -340,15 +340,15 @@ bool FAnimationAnimNextEditorTest_GraphTraitOperations::RunTest(const FString& I
 
 			const FName BasicTraitTypeName = *BasicTrait->GetTraitName();
 
-			TraitInstanceName = Controller->SwapTrait(TraitStackNode->GetFName(), TraitInstanceName, 1, BasicTraitTypeName);
+			TraitInstanceName = Controller->SwapTraitByName(TraitStackNode->GetFName(), TraitInstanceName, 1, BasicTraitTypeName);
 			UE_RETURN_ON_ERROR(TraitInstanceName == BasicTraitTypeName, TEXT("FAnimationAnimNextEditorTest_GraphTraitOperations -> Unexpected Trait name"));
 
 			URigVMPin* TraitPin = TraitStackNode->FindPin(TraitInstanceName.ToString());
 			UE_RETURN_ON_ERROR(TraitPin != nullptr, TEXT("FAnimationAnimNextEditorTest_GraphTraitOperations -> Failed to find FTestBasicTrait pin"));
 
-			// Our first pin is the hard coded output result, decorator pins follow
+			// Our first pin is the hard coded output result, trait pins follow
 			UE_RETURN_ON_ERROR(NodePins.Num() == 2, TEXT("FAnimationAnimNextEditorTest_GraphTraitOperations -> Unexpected number of pins"));
-			UE_RETURN_ON_ERROR(TraitPin->IsDecoratorPin() == true, TEXT("FAnimationAnimNextEditorTest_GraphTraitOperations -> Unexpected pin type"));
+			UE_RETURN_ON_ERROR(TraitPin->IsTraitPin() == true, TEXT("FAnimationAnimNextEditorTest_GraphTraitOperations -> Unexpected pin type"));
 			UE_RETURN_ON_ERROR(TraitPin->GetFName() == TraitInstanceName, TEXT("FAnimationAnimNextEditorTest_GraphTraitOperations -> Unexpected pin name"));
 			UE_RETURN_ON_ERROR(TraitPin->GetCPPTypeObject() == FRigDecorator_AnimNextCppDecorator::StaticStruct(), TEXT("FAnimationAnimNextEditorTest_GraphTraitOperations -> Unexpected pin type"));
 		}
@@ -389,16 +389,16 @@ bool FAnimationAnimNextRuntimeTest_GraphExecute::RunTest(const FString& InParame
 		URigVMPin* BeginExecutePin = MainEntryPointNode->FindPin(GET_MEMBER_NAME_STRING_CHECKED(FRigUnit_AnimNextGraphRoot, Result));
 		UE_RETURN_ON_ERROR(BeginExecutePin != nullptr && BeginExecutePin->GetDirection() == ERigVMPinDirection::Input, "FAnimationAnimNextRuntimeTest_GraphExecute -> Failed to create entry point");
 
-		// Create an empty decorator stack node
+		// Create an empty trait stack node
 		URigVMUnitNode* DecoratorStackNode = Controller->AddUnitNode(FRigUnit_AnimNextTraitStack::StaticStruct(), FRigVMStruct::ExecuteName, FVector2D(0.0f, 0.0f), FString(), false);
-		UE_RETURN_ON_ERROR(DecoratorStackNode != nullptr, "FAnimationAnimNextRuntimeTest_GraphExecute -> Failed to create decorator stack node");
+		UE_RETURN_ON_ERROR(DecoratorStackNode != nullptr, "FAnimationAnimNextRuntimeTest_GraphExecute -> Failed to create trait stack node");
 
 		// Link our stack result to our entry point
 		Controller->AddLink(DecoratorStackNode->GetPins()[0], MainEntryPointNode->GetPins()[0]);
 
 		// Add a trait
 		const UScriptStruct* CppDecoratorStruct = FRigDecorator_AnimNextCppDecorator::StaticStruct();
-		UE_RETURN_ON_ERROR(CppDecoratorStruct != nullptr, "FAnimationAnimNextRuntimeTest_GraphExecute -> Failed to get find Cpp decorator static struct");
+		UE_RETURN_ON_ERROR(CppDecoratorStruct != nullptr, "FAnimationAnimNextRuntimeTest_GraphExecute -> Failed to get find Cpp trait static struct");
 
 		const FTrait* Trait = FTraitRegistry::Get().Find(FTestTrait::TraitUID);
 		UE_RETURN_ON_ERROR(Trait != nullptr, "FAnimationAnimNextRuntimeTest_GraphExecute -> Failed to find test trait");
@@ -412,7 +412,7 @@ bool FAnimationAnimNextRuntimeTest_GraphExecute::RunTest(const FString& InParame
 			FRigDecorator_AnimNextCppDecorator CppDecoratorStructInstance;
 			CppDecoratorStructInstance.DecoratorSharedDataStruct = ScriptStruct;
 
-			UE_RETURN_ON_ERROR(CppDecoratorStructInstance.CanBeAddedToNode(DecoratorStackNode, nullptr), "FAnimationAnimNextRuntimeTest_GraphExecute -> Trait cannot be added to decorator stack node");
+			UE_RETURN_ON_ERROR(CppDecoratorStructInstance.CanBeAddedToNode(DecoratorStackNode, nullptr), "FAnimationAnimNextRuntimeTest_GraphExecute -> Trait cannot be added to trait stack node");
 
 			const FProperty* Prop = FAnimNextCppDecoratorWrapper::StaticStruct()->FindPropertyByName(GET_MEMBER_NAME_STRING_CHECKED(FAnimNextCppDecoratorWrapper, CppDecorator));
 			UE_RETURN_ON_ERROR(Prop != nullptr, "FAnimationAnimNextRuntimeTest_GraphExecute -> Failed to find wrapper property");
@@ -424,17 +424,17 @@ bool FAnimationAnimNextRuntimeTest_GraphExecute::RunTest(const FString& InParame
 		ScriptStruct->GetStringMetaDataHierarchical(FRigVMStruct::DisplayNameMetaName, &DisplayNameMetadata);
 		const FString DisplayName = DisplayNameMetadata.IsEmpty() ? Trait->GetTraitName() : DisplayNameMetadata;
 
-		const FName DecoratorName = Controller->AddDecorator(
+		const FName DecoratorName = Controller->AddTrait(
 			DecoratorStackNode->GetFName(),
 			*CppDecoratorStruct->GetPathName(),
 			*DisplayName,
 			DefaultValue, INDEX_NONE, true, true);
-		UE_RETURN_ON_ERROR(DecoratorName == DisplayName, TEXT("FAnimationAnimNextRuntimeTest_GraphExecute -> Unexpected decorator name"));
+		UE_RETURN_ON_ERROR(DecoratorName == DisplayName, TEXT("FAnimationAnimNextRuntimeTest_GraphExecute -> Unexpected trait name"));
 
 		URigVMPin* DecoratorPin = DecoratorStackNode->FindPin(*DisplayName);
-		UE_RETURN_ON_ERROR(DecoratorPin != nullptr, TEXT("FAnimationAnimNextRuntimeTest_GraphExecute -> Failed to find decorator pin"));
+		UE_RETURN_ON_ERROR(DecoratorPin != nullptr, TEXT("FAnimationAnimNextRuntimeTest_GraphExecute -> Failed to find trait pin"));
 
-		// Set some values on our decorator
+		// Set some values on our trait
 		Controller->SetPinDefaultValue(DecoratorPin->GetSubPins()[1]->GetPinPath(), TEXT("78"));
 		Controller->SetPinDefaultValue(DecoratorPin->GetSubPins()[2]->GetPinPath(), TEXT("142.33"));
 
@@ -509,16 +509,16 @@ bool FAnimationAnimNextRuntimeTest_GraphExecuteLatent::RunTest(const FString& In
 		URigVMPin* BeginExecutePin = MainEntryPointNode->FindPin(GET_MEMBER_NAME_STRING_CHECKED(FRigUnit_AnimNextGraphRoot, Result));
 		UE_RETURN_ON_ERROR(BeginExecutePin != nullptr && BeginExecutePin->GetDirection() == ERigVMPinDirection::Input, "FAnimationAnimNextRuntimeTest_GraphExecuteLatent -> Failed to create entry point");
 
-		// Create an empty decorator stack node
+		// Create an empty trait stack node
 		URigVMUnitNode* DecoratorStackNode = Controller->AddUnitNode(FRigUnit_AnimNextTraitStack::StaticStruct(), FRigVMStruct::ExecuteName, FVector2D(0.0f, 0.0f), FString(), false);
-		UE_RETURN_ON_ERROR(DecoratorStackNode != nullptr, "FAnimationAnimNextRuntimeTest_GraphExecuteLatent -> Failed to create decorator stack node");
+		UE_RETURN_ON_ERROR(DecoratorStackNode != nullptr, "FAnimationAnimNextRuntimeTest_GraphExecuteLatent -> Failed to create trait stack node");
 
 		// Link our stack result to our entry point
 		Controller->AddLink(DecoratorStackNode->GetPins()[0], MainEntryPointNode->GetPins()[0]);
 
 		// Add a trait
 		const UScriptStruct* CppDecoratorStruct = FRigDecorator_AnimNextCppDecorator::StaticStruct();
-		UE_RETURN_ON_ERROR(CppDecoratorStruct != nullptr, "FAnimationAnimNextRuntimeTest_GraphExecuteLatent -> Failed to get find Cpp decorator static struct");
+		UE_RETURN_ON_ERROR(CppDecoratorStruct != nullptr, "FAnimationAnimNextRuntimeTest_GraphExecuteLatent -> Failed to get find Cpp trait static struct");
 
 		const FTrait* Trait = FTraitRegistry::Get().Find(FTestTrait::TraitUID);
 		UE_RETURN_ON_ERROR(Trait != nullptr, "FAnimationAnimNextRuntimeTest_GraphExecuteLatent -> Failed to find test trait");
@@ -532,7 +532,7 @@ bool FAnimationAnimNextRuntimeTest_GraphExecuteLatent::RunTest(const FString& In
 			FRigDecorator_AnimNextCppDecorator CppDecoratorStructInstance;
 			CppDecoratorStructInstance.DecoratorSharedDataStruct = ScriptStruct;
 
-			UE_RETURN_ON_ERROR(CppDecoratorStructInstance.CanBeAddedToNode(DecoratorStackNode, nullptr), "FAnimationAnimNextRuntimeTest_GraphExecuteLatent -> Trait cannot be added to decorator stack node");
+			UE_RETURN_ON_ERROR(CppDecoratorStructInstance.CanBeAddedToNode(DecoratorStackNode, nullptr), "FAnimationAnimNextRuntimeTest_GraphExecuteLatent -> Trait cannot be added to trait stack node");
 
 			const FProperty* Prop = FAnimNextCppDecoratorWrapper::StaticStruct()->FindPropertyByName(GET_MEMBER_NAME_STRING_CHECKED(FAnimNextCppDecoratorWrapper, CppDecorator));
 			UE_RETURN_ON_ERROR(Prop != nullptr, "FAnimationAnimNextRuntimeTest_GraphExecuteLatent -> Failed to find wrapper property");
@@ -544,22 +544,22 @@ bool FAnimationAnimNextRuntimeTest_GraphExecuteLatent::RunTest(const FString& In
 		ScriptStruct->GetStringMetaDataHierarchical(FRigVMStruct::DisplayNameMetaName, &DisplayNameMetadata);
 		const FString DisplayName = DisplayNameMetadata.IsEmpty() ? Trait->GetTraitName() : DisplayNameMetadata;
 
-		const FName DecoratorName = Controller->AddDecorator(
+		const FName DecoratorName = Controller->AddTrait(
 			DecoratorStackNode->GetFName(),
 			*CppDecoratorStruct->GetPathName(),
 			*DisplayName,
 			DefaultValue, INDEX_NONE, true, true);
-		UE_RETURN_ON_ERROR(DecoratorName == DisplayName, TEXT("FAnimationAnimNextRuntimeTest_GraphExecuteLatent -> Unexpected decorator name"));
+		UE_RETURN_ON_ERROR(DecoratorName == DisplayName, TEXT("FAnimationAnimNextRuntimeTest_GraphExecuteLatent -> Unexpected trait name"));
 
-		// Set some values on our decorator
+		// Set some values on our trait
 		URigVMPin* DecoratorPin = DecoratorStackNode->FindPin(*DisplayName);
-		UE_RETURN_ON_ERROR(DecoratorPin != nullptr, TEXT("FAnimationAnimNextRuntimeTest_GraphExecuteLatent -> Failed to find decorator pin"));
+		UE_RETURN_ON_ERROR(DecoratorPin != nullptr, TEXT("FAnimationAnimNextRuntimeTest_GraphExecuteLatent -> Failed to find trait pin"));
 
 		Controller->SetPinDefaultValue(DecoratorPin->GetSubPins()[1]->GetPinPath(), TEXT("78"));		// SomeInt32
 		Controller->SetPinDefaultValue(DecoratorPin->GetSubPins()[2]->GetPinPath(), TEXT("142.33"));	// SomeFloat
 		Controller->SetPinDefaultValue(DecoratorPin->GetSubPins()[5]->GetPinPath(), TEXT("1123.31"));	// SomeLatentFloat, inline value on latent pin
 
-		// Set some latent values on our decorator
+		// Set some latent values on our trait
 		{
 			FRigVMFunction_MathIntAdd IntAdd;
 			IntAdd.A = 10;
diff --git a/Engine/Plugins/Runtime/RigVM/Config/BaseRigVM.ini b/Engine/Plugins/Runtime/RigVM/Config/BaseRigVM.ini
new file mode 100644
--- /dev/null
+++ b/Engine/Plugins/Runtime/RigVM/Config/BaseRigVM.ini
@@ -0,0 +1,5 @@
+[CoreRedirects]
+
+; renaming decorator to trait
++StructRedirects=(OldName="RigVMDecorator", NewName="/Script/RigVM.RigVMTrait")
++PropertyRedirects=(OldName="RigVMNode.DecoratorRootPinNames",NewName="RigVMNode.TraitRootPinNames")
\ No newline at end of file
diff --git a/Engine/Plugins/Runtime/RigVM/Source/RigVM/Private/RigVMCore/RigVM.cpp b/Engine/Plugins/Runtime/RigVM/Source/RigVM/Private/RigVMCore/RigVM.cpp
--- a/Engine/Plugins/Runtime/RigVM/Source/RigVM/Private/RigVMCore/RigVM.cpp
+++ b/Engine/Plugins/Runtime/RigVM/Source/RigVM/Private/RigVMCore/RigVM.cpp
@@ -18,7 +18,7 @@
 #include "UObject/UObjectIterator.h"
 #include "UObject/ObjectSaveContext.h"
 #include "RigVMHost.h"
-#include "RigVMCore/RigVMDecorator.h"
+#include "RigVMCore/RigVMTrait.h"
 
 #include UE_INLINE_GENERATED_CPP_BY_NAME(RigVM)
 
@@ -556,9 +556,9 @@ bool URigVM::ValidateBytecode()
 				CheckOperandValidity(Op.Arg);
 				break;
 			}
-			case ERigVMOpCode::SetupDecorators:
+			case ERigVMOpCode::SetupTraits:
 			{
-				const FRigVMSetupDecoratorsOp& Op = ByteCodeStorage.GetOpAt<FRigVMSetupDecoratorsOp>(ByteCodeInstruction);
+				const FRigVMSetupTraitsOp& Op = ByteCodeStorage.GetOpAt<FRigVMSetupTraitsOp>(ByteCodeInstruction);
 				CheckOperandValidity(Op.Arg);
 				break;
 			}
@@ -1152,9 +1152,9 @@ void URigVM::InstructionOpEval(FRigVMExtendedExecuteContext& Context, int32 Inst
 			InOpFunc(Context, InHandleBaseIndex, {}, Arg);
 			break;
 		}
-		case ERigVMOpCode::SetupDecorators:
+		case ERigVMOpCode::SetupTraits:
 		{
-			const FRigVMSetupDecoratorsOp& Op = ByteCode.GetOpAt<FRigVMSetupDecoratorsOp>(Instructions[InstructionIndex]);
+			const FRigVMSetupTraitsOp& Op = ByteCode.GetOpAt<FRigVMSetupTraitsOp>(Instructions[InstructionIndex]);
 			const FRigVMOperand& Arg = Op.Arg;
 			InOpFunc(Context, InHandleBaseIndex, {}, Arg);
 			break;
@@ -2139,22 +2139,22 @@ ERigVMExecuteResult URigVM::ExecuteInstructions(FRigVMExtendedExecuteContext& Co
 				ContextPublicData.InstructionIndex++;
 				break;
 			}
-			case ERigVMOpCode::SetupDecorators:
+			case ERigVMOpCode::SetupTraits:
 			{
-				ContextPublicData.Decorators.Reset();
+				ContextPublicData.Traits.Reset();
 				
-				const TArray<int32>& DecoratorList = *(TArray<int32>*)Context.CachedMemoryHandles[FirstHandleForInstruction[ContextPublicData.InstructionIndex]].GetData();
-				for(const int32 DecoratorPropertyIndex : DecoratorList)
+				const TArray<int32>& TraitList = *(TArray<int32>*)Context.CachedMemoryHandles[FirstHandleForInstruction[ContextPublicData.InstructionIndex]].GetData();
+				for(const int32 TraitPropertyIndex : TraitList)
 				{
-					if(Context.WorkMemoryStorage.GetProperties().IsValidIndex(DecoratorPropertyIndex))
+					if(Context.WorkMemoryStorage.GetProperties().IsValidIndex(TraitPropertyIndex))
 					{
-						const FProperty* Property = Context.WorkMemoryStorage.GetProperties()[DecoratorPropertyIndex];
+						const FProperty* Property = Context.WorkMemoryStorage.GetProperties()[TraitPropertyIndex];
 						if(const FStructProperty* StructProperty = CastField<FStructProperty>(Property))
 						{
-							if(StructProperty->Struct && StructProperty->Struct->IsChildOf(FRigVMDecorator::StaticStruct()))
+							if(StructProperty->Struct && StructProperty->Struct->IsChildOf(FRigVMTrait::StaticStruct()))
 							{
-								ContextPublicData.Decorators.Emplace(
-									Context.WorkMemoryStorage.GetData<FRigVMDecorator>(StructProperty),
+								ContextPublicData.Traits.Emplace(
+									Context.WorkMemoryStorage.GetData<FRigVMTrait>(StructProperty),
 									Cast<UScriptStruct>(StructProperty->Struct));
 							}
 						}
@@ -2500,10 +2500,10 @@ TArray<FString> URigVM::DumpByteCodeAsTextArray(FRigVMExtendedExecuteContext& Co
 				ResultLine = FString::Printf(TEXT("Run Instructions %d-%d (%s)"), Op.StartInstruction, Op.EndInstruction, *GetOperandLabel(Context, Op.Arg, OperandFormatFunction));
 				break;
 			}
-			case ERigVMOpCode::SetupDecorators:
+			case ERigVMOpCode::SetupTraits:
 			{
-				const FRigVMSetupDecoratorsOp& Op = ByteCode.GetOpAt<FRigVMSetupDecoratorsOp>(Instructions[InstructionIndex]);
-				ResultLine = FString::Printf(TEXT("Setup Decorators (%s)"), *GetOperandLabel(Context, Op.Arg, OperandFormatFunction));
+				const FRigVMSetupTraitsOp& Op = ByteCode.GetOpAt<FRigVMSetupTraitsOp>(Instructions[InstructionIndex]);
+				ResultLine = FString::Printf(TEXT("Setup Traits (%s)"), *GetOperandLabel(Context, Op.Arg, OperandFormatFunction));
 				break;
 			}
 			default:
diff --git a/Engine/Plugins/Runtime/RigVM/Source/RigVM/Private/RigVMCore/RigVMByteCode.cpp b/Engine/Plugins/Runtime/RigVM/Source/RigVM/Private/RigVMCore/RigVMByteCode.cpp
--- a/Engine/Plugins/Runtime/RigVM/Source/RigVM/Private/RigVMCore/RigVMByteCode.cpp
+++ b/Engine/Plugins/Runtime/RigVM/Source/RigVM/Private/RigVMCore/RigVMByteCode.cpp
@@ -5,8 +5,8 @@
 #include "UObject/UE5MainStreamObjectVersion.h"
 #include "UObject/FortniteMainBranchObjectVersion.h"
 #include "RigVMCore/RigVM.h"
-#include "RigVMCore/RigVMDecorator.h"
-#include "RigVMCore/RigVMDecoratorScope.h"
+#include "RigVMCore/RigVMTrait.h"
+#include "RigVMCore/RigVMTraitScope.h"
 #include "RigVMObjectVersion.h"
 
 #include UE_INLINE_GENERATED_CPP_BY_NAME(RigVMByteCode)
@@ -434,9 +434,9 @@ void FRigVMByteCode::Save(FArchive& Ar)
 				Ar << Op;
 				break;
 			}
-			case ERigVMOpCode::SetupDecorators:
+			case ERigVMOpCode::SetupTraits:
 			{
-				FRigVMSetupDecoratorsOp Op = GetOpAt<FRigVMSetupDecoratorsOp>(Instruction.ByteCodeIndex);
+				FRigVMSetupTraitsOp Op = GetOpAt<FRigVMSetupTraitsOp>(Instruction.ByteCodeIndex);
 				Ar << Op;
 				break;
 			}
@@ -632,11 +632,11 @@ void FRigVMByteCode::Load(FArchive& Ar)
 				AddOp<FRigVMRunInstructionsOp>(Op);
 				break;
 			}
-			case ERigVMOpCode::SetupDecorators:
+			case ERigVMOpCode::SetupTraits:
 			{
-				FRigVMSetupDecoratorsOp Op;
+				FRigVMSetupTraitsOp Op;
 				Ar << Op;
-				AddOp<FRigVMSetupDecoratorsOp>(Op);
+				AddOp<FRigVMSetupTraitsOp>(Op);
 				break;
 			}
 			default:
@@ -850,9 +850,9 @@ uint32 FRigVMByteCode::GetOperatorHash(const FRigVMInstruction& InInstruction) c
 		{
 			return GetTypeHash(GetOpAt<FRigVMRunInstructionsOp>(InInstruction));
 		}
-		case ERigVMOpCode::SetupDecorators:
+		case ERigVMOpCode::SetupTraits:
 		{
-			return GetTypeHash(GetOpAt<FRigVMSetupDecoratorsOp>(InInstruction));
+			return GetTypeHash(GetOpAt<FRigVMSetupTraitsOp>(InInstruction));
 		}
 		case ERigVMOpCode::Invalid:
 		{
@@ -1070,9 +1070,9 @@ uint64 FRigVMByteCode::GetOpNumBytesAt(uint64 InByteCodeIndex, bool bIncludeOper
 		{
 			return (uint64)sizeof(FRigVMRunInstructionsOp);
 		}
-		case ERigVMOpCode::SetupDecorators:
+		case ERigVMOpCode::SetupTraits:
 		{
-			return (uint64)sizeof(FRigVMSetupDecoratorsOp);
+			return (uint64)sizeof(FRigVMSetupTraitsOp);
 		}
 		case ERigVMOpCode::Invalid:
 		{
@@ -1364,9 +1364,9 @@ FString FRigVMByteCode::DumpToText() const
 				Line += FString::FromInt(Op.EndInstruction);
 				break;
 			}
-			case ERigVMOpCode::SetupDecorators:
+			case ERigVMOpCode::SetupTraits:
 			{
-				const FRigVMSetupDecoratorsOp& Op = GetOpAt<FRigVMSetupDecoratorsOp>(Instruction.ByteCodeIndex);
+				const FRigVMSetupTraitsOp& Op = GetOpAt<FRigVMSetupTraitsOp>(Instruction.ByteCodeIndex);
 				FString Arg;
 				FRigVMOperand::StaticStruct()->ExportText(Arg, &Op.Arg, &Op.Arg, nullptr, PPF_None, nullptr);
 				Line += Arg;
@@ -1420,9 +1420,9 @@ uint64 FRigVMByteCode::AddRunInstructionsOp(FRigVMOperand InExecuteStateArg, int
 	return AddOp(FRigVMRunInstructionsOp(InExecuteStateArg, InStartInstruction, InEndInstruction));
 }
 
-uint64 FRigVMByteCode::AddSetupDecoratorsOp(FRigVMOperand InDecoratorListArg)
+uint64 FRigVMByteCode::AddSetupTraitsOp(FRigVMOperand InTraitListArg)
 {
-	return AddOp(FRigVMSetupDecoratorsOp(InDecoratorListArg));
+	return AddOp(FRigVMSetupTraitsOp(InTraitListArg));
 }
 
 int32 FRigVMByteCode::AddBranchInfo(const FRigVMBranchInfo& InBranchInfo)
@@ -1529,9 +1529,9 @@ FRigVMOperandArray FRigVMByteCode::GetOperandsForOp(const FRigVMInstruction& InI
 			const FRigVMRunInstructionsOp& Op = GetOpAt<FRigVMRunInstructionsOp>(InInstruction.ByteCodeIndex);
 			return FRigVMOperandArray(&Op.Arg, 1);
 		}
-		case ERigVMOpCode::SetupDecorators:
+		case ERigVMOpCode::SetupTraits:
 		{
-			const FRigVMSetupDecoratorsOp& Op = GetOpAt<FRigVMSetupDecoratorsOp>(InInstruction.ByteCodeIndex);
+			const FRigVMSetupTraitsOp& Op = GetOpAt<FRigVMSetupTraitsOp>(InInstruction.ByteCodeIndex);
 			return FRigVMOperandArray(&Op.Arg, 1);
 		}
 		case ERigVMOpCode::JumpAbsolute:
@@ -1632,9 +1632,9 @@ uint64 FRigVMByteCode::GetFirstOperandByteIndex(const FRigVMInstruction& InInstr
 			const FRigVMRunInstructionsOp& Op = GetOpAt<FRigVMRunInstructionsOp>(InInstruction.ByteCodeIndex);
 			return InInstruction.ByteCodeIndex + ((uint64)&Op.Arg - (uint64)&Op);
 		}
-		case ERigVMOpCode::SetupDecorators:
+		case ERigVMOpCode::SetupTraits:
 		{
-			const FRigVMSetupDecoratorsOp& Op = GetOpAt<FRigVMSetupDecoratorsOp>(InInstruction.ByteCodeIndex);
+			const FRigVMSetupTraitsOp& Op = GetOpAt<FRigVMSetupTraitsOp>(InInstruction.ByteCodeIndex);
 			return InInstruction.ByteCodeIndex + ((uint64)&Op.Arg - (uint64)&Op);
 		}
 		case ERigVMOpCode::JumpAbsolute:
@@ -1775,9 +1775,9 @@ uint64 FRigVMByteCode::GetOpAlignment(ERigVMOpCode InOpCode) const
 			static const uint64 Alignment = FRigVMRunInstructionsOp::StaticStruct()->GetCppStructOps()->GetAlignment();
 			return Alignment;
 		}
-		case ERigVMOpCode::SetupDecorators:
+		case ERigVMOpCode::SetupTraits:
 		{
-			static const uint64 Alignment = FRigVMSetupDecoratorsOp::StaticStruct()->GetCppStructOps()->GetAlignment();
+			static const uint64 Alignment = FRigVMSetupTraitsOp::StaticStruct()->GetCppStructOps()->GetAlignment();
 			return Alignment;
 		}
 		case ERigVMOpCode::Invalid:
@@ -2127,62 +2127,62 @@ void FRigVMByteCode::SetOperandsForInstruction(int32 InInstructionIndex, const F
 
 #endif
 
-TMap<int32, TArray<FRigVMDecoratorScope>> FRigVMByteCode::GetDecorators(
+TMap<int32, TArray<FRigVMTraitScope>> FRigVMByteCode::GetTraits(
 	FRigVMMemoryStorageStruct& InLiteralMemory,
 	FRigVMMemoryStorageStruct& InWorkMemory,
 	const UScriptStruct* InScriptStruct) const
 {
-	TMap<int32, TArray<FRigVMDecoratorScope>> Result;
+	TMap<int32, TArray<FRigVMTraitScope>> Result;
 	const FRigVMInstructionArray Instructions = GetInstructions();
 	for(int32 InstructionIndex = 0; InstructionIndex < Instructions.Num(); InstructionIndex++)
 	{
 		const FRigVMInstruction& Instruction = Instructions[InstructionIndex];
-		if(Instruction.OpCode == ERigVMOpCode::SetupDecorators)
+		if(Instruction.OpCode == ERigVMOpCode::SetupTraits)
 		{
-			const TArray<FRigVMDecoratorScope> Decorators = GetDecoratorsForInstruction(Instruction, InLiteralMemory, InWorkMemory, InScriptStruct);
-			if(!Decorators.IsEmpty())
+			const TArray<FRigVMTraitScope> Traits = GetTraitsForInstruction(Instruction, InLiteralMemory, InWorkMemory, InScriptStruct);
+			if(!Traits.IsEmpty())
 			{
-				Result.Add(InstructionIndex, Decorators);
+				Result.Add(InstructionIndex, Traits);
 			}
 		}
 	}
 	return Result;
 }
 
-TArray<FRigVMDecoratorScope> FRigVMByteCode::GetDecoratorsForInstruction(
+TArray<FRigVMTraitScope> FRigVMByteCode::GetTraitsForInstruction(
 	const FRigVMInstruction& InInstruction,
 	FRigVMMemoryStorageStruct& InLiteralMemory,
 	FRigVMMemoryStorageStruct& InWorkMemory,
 	const UScriptStruct* InScriptStruct) const
 {
-	TArray<FRigVMDecoratorScope> Decorators;
+	TArray<FRigVMTraitScope> Traits;
 
-	if(InInstruction.OpCode == ERigVMOpCode::SetupDecorators)
+	if(InInstruction.OpCode == ERigVMOpCode::SetupTraits)
 	{
 		if(InScriptStruct == nullptr)
 		{
-			InScriptStruct = FRigVMDecorator::StaticStruct();
+			InScriptStruct = FRigVMTrait::StaticStruct();
 		}
 
-		const FRigVMSetupDecoratorsOp& Op = GetOpAt<FRigVMSetupDecoratorsOp>(InInstruction);
+		const FRigVMSetupTraitsOp& Op = GetOpAt<FRigVMSetupTraitsOp>(InInstruction);
 		check(Op.Arg.GetMemoryType() == ERigVMMemoryType::Literal);
 
 		if(InLiteralMemory.GetProperties().IsValidIndex(Op.Arg.GetRegisterIndex()))
 		{
-			if(const FArrayProperty* DecoratorIndicesProperty = CastField<FArrayProperty>(InLiteralMemory.GetProperty(Op.Arg.GetRegisterIndex())))
+			if(const FArrayProperty* TraitIndicesProperty = CastField<FArrayProperty>(InLiteralMemory.GetProperty(Op.Arg.GetRegisterIndex())))
 			{
-				if(DecoratorIndicesProperty->Inner->IsA<FIntProperty>())
+				if(TraitIndicesProperty->Inner->IsA<FIntProperty>())
 				{
-					const TArray<int32>& DecoratorIndices = *InLiteralMemory.GetData<TArray<int32>>(DecoratorIndicesProperty);
-					for(const int32 DecoratorIndex : DecoratorIndices)
+					const TArray<int32>& TraitIndices = *InLiteralMemory.GetData<TArray<int32>>(TraitIndicesProperty);
+					for(const int32 TraitIndex : TraitIndices)
 					{
-						const FProperty* Property = InWorkMemory.GetProperties()[DecoratorIndex];
+						const FProperty* Property = InWorkMemory.GetProperties()[TraitIndex];
 						if(const FStructProperty* StructProperty = CastField<FStructProperty>(Property))
 						{
 							if(StructProperty->Struct && StructProperty->Struct->IsChildOf(InScriptStruct))
 							{
-								Decorators.Emplace(
-									InWorkMemory.GetData<FRigVMDecorator>(StructProperty),
+								Traits.Emplace(
+									InWorkMemory.GetData<FRigVMTrait>(StructProperty),
 									Cast<UScriptStruct>(StructProperty->Struct));
 							}
 						}
@@ -2192,7 +2192,7 @@ TArray<FRigVMDecoratorScope> FRigVMByteCode::GetDecoratorsForInstruction(
 		}
 	}
 
-	return Decorators;
+	return Traits;
 }
 
 const FRigVMBranchInfo* FRigVMByteCode::GetBranchInfo(const FRigVMBranchInfoKey& InBranchInfoKey) const
diff --git a/Engine/Plugins/Runtime/RigVM/Source/RigVM/Private/RigVMCore/RigVMRegistry.cpp b/Engine/Plugins/Runtime/RigVM/Source/RigVM/Private/RigVMCore/RigVMRegistry.cpp
--- a/Engine/Plugins/Runtime/RigVM/Source/RigVM/Private/RigVMCore/RigVMRegistry.cpp
+++ b/Engine/Plugins/Runtime/RigVM/Source/RigVM/Private/RigVMCore/RigVMRegistry.cpp
@@ -3,7 +3,7 @@
 #include "RigVMCore/RigVMRegistry.h"
 #include "AssetRegistry/AssetData.h"
 #include "RigVMCore/RigVMStruct.h"
-#include <RigVMCore/RigVMDecorator.h>
+#include <RigVMCore/RigVMTrait.h>
 #include "RigVMTypeUtils.h"
 #include "RigVMModule.h"
 #include "Animation/AttributeTypes.h"
@@ -1279,7 +1279,7 @@ bool FRigVMRegistry_NoLock::IsAllowedType_NoLock(const UStruct* InStruct) const
 		return false;
 	}
 	if(InStruct->IsChildOf(FRigVMStruct::StaticStruct()) &&
-		!InStruct->IsChildOf(FRigVMDecorator::StaticStruct()))
+		!InStruct->IsChildOf(FRigVMTrait::StaticStruct()))
 	{
 		return false;
 	}
diff --git a/Engine/Plugins/Runtime/RigVM/Source/RigVM/Public/RigVMCore/RigVM.h b/Engine/Plugins/Runtime/RigVM/Source/RigVM/Public/RigVMCore/RigVM.h
--- a/Engine/Plugins/Runtime/RigVM/Source/RigVM/Public/RigVMCore/RigVM.h
+++ b/Engine/Plugins/Runtime/RigVM/Source/RigVM/Public/RigVMCore/RigVM.h
@@ -428,30 +428,30 @@ public:
 	// returns false if an entry can not be executed
 	bool CanExecuteEntry(const FRigVMExtendedExecuteContext& Context, const FName& InEntryName, bool bLogErrorForMissingEntry = true) const;
 
-	// returns the decorators for this VM's bytecode and Context
-	TMap<int32, TArray<FRigVMDecoratorScope>> GetDecorators(FRigVMExtendedExecuteContext& InContext, const UScriptStruct* InScriptStruct = nullptr)
+	// returns the traits for this VM's bytecode and Context
+	TMap<int32, TArray<FRigVMTraitScope>> GetTraits(FRigVMExtendedExecuteContext& InContext, const UScriptStruct* InScriptStruct = nullptr)
 	{
-		return GetByteCode().GetDecorators(*GetLiteralMemory(), InContext.WorkMemoryStorage, InScriptStruct);
+		return GetByteCode().GetTraits(*GetLiteralMemory(), InContext.WorkMemoryStorage, InScriptStruct);
 	}
 
-	// returns the decorators of a given type for this VM's bytecode and Context
+	// returns the traits of a given type for this VM's bytecode and Context
 	template<typename T>
-	TMap<int32, TArray<FRigVMDecoratorScope>> GetDecorators(FRigVMExtendedExecuteContext& InContext)
+	TMap<int32, TArray<FRigVMTraitScope>> GetTraits(FRigVMExtendedExecuteContext& InContext)
 	{
-		return GetByteCode().GetDecorators<T>(*GetLiteralMemory(), InContext.WorkMemoryStorage);
+		return GetByteCode().GetTraits<T>(*GetLiteralMemory(), InContext.WorkMemoryStorage);
 	}
 
-	// returns the decorators for the provided memory for a single instruction
-	TArray<FRigVMDecoratorScope> GetDecoratorsForInstruction(const FRigVMInstruction& InInstruction, FRigVMExtendedExecuteContext& InContext, const UScriptStruct* InScriptStruct = nullptr)
+	// returns the traits for the provided memory for a single instruction
+	TArray<FRigVMTraitScope> GetTraitsForInstruction(const FRigVMInstruction& InInstruction, FRigVMExtendedExecuteContext& InContext, const UScriptStruct* InScriptStruct = nullptr)
 	{
-		return GetByteCode().GetDecoratorsForInstruction(InInstruction, *GetLiteralMemory(), InContext.WorkMemoryStorage, InScriptStruct);
+		return GetByteCode().GetTraitsForInstruction(InInstruction, *GetLiteralMemory(), InContext.WorkMemoryStorage, InScriptStruct);
 	}
 
-	// returns the decorators of a given type for the provided memory for a single instruction
+	// returns the traits of a given type for the provided memory for a single instruction
 	template<typename T>
-	TArray<FRigVMDecoratorScope> GetDecoratorsForInstruction(const FRigVMInstruction& InInstruction, FRigVMExtendedExecuteContext& InContext)
+	TArray<FRigVMTraitScope> GetTraitsForInstruction(const FRigVMInstruction& InInstruction, FRigVMExtendedExecuteContext& InContext)
 	{
-		return GetByteCode().GetDecoratorsForInstruction<T>(InInstruction, *GetLiteralMemory(), InContext.WorkMemoryStorage);
+		return GetByteCode().GetTraitsForInstruction<T>(InInstruction, *GetLiteralMemory(), InContext.WorkMemoryStorage);
 	}
 
 #if WITH_EDITOR
diff --git a/Engine/Plugins/Runtime/RigVM/Source/RigVM/Public/RigVMCore/RigVMByteCode.h b/Engine/Plugins/Runtime/RigVM/Source/RigVM/Public/RigVMCore/RigVMByteCode.h
--- a/Engine/Plugins/Runtime/RigVM/Source/RigVM/Public/RigVMCore/RigVMByteCode.h
+++ b/Engine/Plugins/Runtime/RigVM/Source/RigVM/Public/RigVMCore/RigVMByteCode.h
@@ -221,7 +221,7 @@ enum class ERigVMOpCode : uint8
 	JumpToBranch, // jumps to a branch based on a name operand
 	Execute, // single execute op (formerly Execute_0_Operands to Execute_64_Operands)
 	RunInstructions, // runs a set of instructions lazily
-	SetupDecorators, // sets up a list of decorators on executecontext
+	SetupTraits, // sets up a list of traits on executecontext
 	Invalid,
 	FirstArrayOpCode = ArrayReset,
 	LastArrayOpCode = ArrayReverse,
@@ -336,7 +336,7 @@ struct RIGVM_API FRigVMUnaryOp : public FRigVMBaseOp
 			uint8(InOpCode) == uint8(ERigVMOpCode::ChangeType) ||
 			uint8(InOpCode) == uint8(ERigVMOpCode::JumpToBranch) ||
 			uint8(InOpCode) == uint8(ERigVMOpCode::RunInstructions) ||
-			uint8(InOpCode) == uint8(ERigVMOpCode::SetupDecorators)
+			uint8(InOpCode) == uint8(ERigVMOpCode::SetupTraits)
 		);
 	}
 
@@ -1016,19 +1016,19 @@ struct RIGVM_API FRigVMRunInstructionsOp : public FRigVMUnaryOp
 	}
 };
 
-// sets up a list of decorators in the execute context
+// sets up a list of traits in the execute context
 USTRUCT()
-struct RIGVM_API FRigVMSetupDecoratorsOp : public FRigVMUnaryOp
+struct RIGVM_API FRigVMSetupTraitsOp : public FRigVMUnaryOp
 {
 	GENERATED_USTRUCT_BODY()
 
-	FRigVMSetupDecoratorsOp()
+	FRigVMSetupTraitsOp()
 		: FRigVMUnaryOp()
 	{
 	}
 
-	FRigVMSetupDecoratorsOp(FRigVMOperand InDecoratorListArg)
-		: FRigVMUnaryOp(ERigVMOpCode::SetupDecorators, InDecoratorListArg)
+	FRigVMSetupTraitsOp(FRigVMOperand InTraitListArg)
+		: FRigVMUnaryOp(ERigVMOpCode::SetupTraits, InTraitListArg)
 	{
 	}
 };
@@ -1235,8 +1235,8 @@ public:
 	// adds a run instructions op
 	uint64 AddRunInstructionsOp(FRigVMOperand InExecuteStateArg, int32 InStartInstruction, int32 InEndInstruction);
 
-	// adds a setup decorators op
-	uint64 AddSetupDecoratorsOp(FRigVMOperand InDecoratorListArg);
+	// adds a setup traits op
+	uint64 AddSetupTraitsOp(FRigVMOperand InTraitListArg);
 
 	// adds information about a branch for an instruction's argument
 	int32 AddBranchInfo(const FRigVMBranchInfo& InBranchInfo);
@@ -1430,24 +1430,24 @@ public:
 
 #endif
 
-	// returns the decorators for the provided memory
-	TMap<int32, TArray<FRigVMDecoratorScope>> GetDecorators(FRigVMMemoryStorageStruct& InLiteralMemory, FRigVMMemoryStorageStruct& InWorkMemory, const UScriptStruct* InScriptStruct = nullptr) const;
+	// returns the traits for the provided memory
+	TMap<int32, TArray<FRigVMTraitScope>> GetTraits(FRigVMMemoryStorageStruct& InLiteralMemory, FRigVMMemoryStorageStruct& InWorkMemory, const UScriptStruct* InScriptStruct = nullptr) const;
 
-	// returns the decorators of a given type for the provided memory
+	// returns the traits of a given type for the provided memory
 	template<typename T>
-	TMap<int32, TArray<FRigVMDecoratorScope>> GetDecorators(FRigVMMemoryStorageStruct& InLiteralMemory, FRigVMMemoryStorageStruct& InWorkMemory) const
+	TMap<int32, TArray<FRigVMTraitScope>> GetTraits(FRigVMMemoryStorageStruct& InLiteralMemory, FRigVMMemoryStorageStruct& InWorkMemory) const
 	{
-		return GetDecorators(InLiteralMemory, InWorkMemory, T::StaticStruct());
+		return GetTraits(InLiteralMemory, InWorkMemory, T::StaticStruct());
 	}
 
-	// returns the decorators for the provided memory for a single instruction
-	TArray<FRigVMDecoratorScope> GetDecoratorsForInstruction(const FRigVMInstruction& InInstruction, FRigVMMemoryStorageStruct& InLiteralMemory, FRigVMMemoryStorageStruct& InWorkMemory, const UScriptStruct* InScriptStruct = nullptr) const;
+	// returns the traits for the provided memory for a single instruction
+	TArray<FRigVMTraitScope> GetTraitsForInstruction(const FRigVMInstruction& InInstruction, FRigVMMemoryStorageStruct& InLiteralMemory, FRigVMMemoryStorageStruct& InWorkMemory, const UScriptStruct* InScriptStruct = nullptr) const;
 
-	// returns the decorators of a given type for the provided memory for a single instruction
+	// returns the traits of a given type for the provided memory for a single instruction
 	template<typename T>
-	TArray<FRigVMDecoratorScope> GetDecoratorsForInstruction(const FRigVMInstruction& InInstruction, FRigVMMemoryStorageStruct& InLiteralMemory, FRigVMMemoryStorageStruct& InWorkMemory) const
+	TArray<FRigVMTraitScope> GetTraitsForInstruction(const FRigVMInstruction& InInstruction, FRigVMMemoryStorageStruct& InLiteralMemory, FRigVMMemoryStorageStruct& InWorkMemory) const
 	{
-		return GetDecoratorsForInstruction(InInstruction, InLiteralMemory, InWorkMemory, T::StaticStruct());
+		return GetTraitsForInstruction(InInstruction, InLiteralMemory, InWorkMemory, T::StaticStruct());
 	}
 
 private:
diff --git a/Engine/Plugins/Runtime/RigVM/Source/RigVM/Public/RigVMCore/RigVMDecorator.h b/Engine/Plugins/Runtime/RigVM/Source/RigVM/Public/RigVMCore/RigVMDecorator.h
deleted file mode 100644
diff --git a/Engine/Plugins/Runtime/RigVM/Source/RigVM/Public/RigVMCore/RigVMDecoratorScope.h b/Engine/Plugins/Runtime/RigVM/Source/RigVM/Public/RigVMCore/RigVMDecoratorScope.h
deleted file mode 100644
diff --git a/Engine/Plugins/Runtime/RigVM/Source/RigVM/Public/RigVMCore/RigVMExecuteContext.h b/Engine/Plugins/Runtime/RigVM/Source/RigVM/Public/RigVMCore/RigVMExecuteContext.h
--- a/Engine/Plugins/Runtime/RigVM/Source/RigVM/Public/RigVMCore/RigVMExecuteContext.h
+++ b/Engine/Plugins/Runtime/RigVM/Source/RigVM/Public/RigVMCore/RigVMExecuteContext.h
@@ -19,7 +19,7 @@
 #include "RigVMCore/RigVMProfilingInfo.h"
 #include "RigVMCore/RigVMNameCache.h"
 #include "RigVMCore/RigVMMemoryStorageStruct.h"
-#include "RigVMCore/RigVMDecoratorScope.h"
+#include "RigVMCore/RigVMTraitScope.h"
 #include "RigVMLog.h"
 #include "RigVMDrawInterface.h"
 #include "RigVMDrawContainer.h"
@@ -272,7 +272,7 @@ struct FRigVMExecuteContext
 		, OwningComponent(nullptr)
 		, OwningActor(nullptr)
 		, World(nullptr)
-		, Decorators()
+		, Traits()
 	{
 	}
 
@@ -418,18 +418,18 @@ struct FRigVMExecuteContext
 	FRigVMDrawContainer* GetDrawContainer() { return DrawContainerPtr; }
 	void SetDrawContainer(FRigVMDrawContainer* InDrawContainer) { DrawContainerPtr = InDrawContainer; }
 
-	TArrayView<const FRigVMDecoratorScope> GetDecorators() const
+	TArrayView<const FRigVMTraitScope> GetTraits() const
 	{
-		if(Decorators.IsEmpty())
+		if(Traits.IsEmpty())
 		{
-			return TArrayView<const FRigVMDecoratorScope>();
+			return TArrayView<const FRigVMTraitScope>();
 		}
-		return TArrayView<const FRigVMDecoratorScope>(Decorators.GetData(), Decorators.Num());
+		return TArrayView<const FRigVMTraitScope>(Traits.GetData(), Traits.Num());
 	}
 
-	TArrayView<FRigVMDecoratorScope> GetDecorators()
+	TArrayView<FRigVMTraitScope> GetTraits()
 	{
-		return TArrayView<FRigVMDecoratorScope>(Decorators.GetData(), Decorators.Num());
+		return TArrayView<FRigVMTraitScope>(Traits.GetData(), Traits.Num());
 	}
 
 	virtual void Initialize()
@@ -499,8 +499,8 @@ protected:
 	/** The world this VM is running in */
 	const UWorld* World;
 
-	/** The decorators accessible to the current instruction */
-	TArray<FRigVMDecoratorScope> Decorators;
+	/** The traits accessible to the current instruction */
+	TArray<FRigVMTraitScope> Traits;
 
 #if UE_RIGVM_DEBUG_EXECUTION
 public:
diff --git a/Engine/Plugins/Runtime/RigVM/Source/RigVM/Public/RigVMCore/RigVMTrait.h b/Engine/Plugins/Runtime/RigVM/Source/RigVM/Public/RigVMCore/RigVMTrait.h
new file mode 100644
--- /dev/null
+++ b/Engine/Plugins/Runtime/RigVM/Source/RigVM/Public/RigVMCore/RigVMTrait.h
@@ -0,0 +1,74 @@
+// Copyright Epic Games, Inc. All Rights Reserved.
+
+#pragma once
+
+#include "RigVMDefines.h"
+#include "RigVMCore/RigVMStruct.h"
+
+#include "RigVMTrait.generated.h"
+
+#if WITH_EDITOR
+
+class URigVMController;
+class URigVMNode;
+
+#endif
+
+/**
+ * The base class for all RigVM traits.
+ */
+USTRUCT()
+struct RIGVM_API FRigVMTrait : public FRigVMStruct
+{
+	GENERATED_BODY()
+
+public:
+
+	FRigVMTrait()
+	{}
+
+	virtual ~FRigVMTrait() {}
+
+	// returns the name of the trait (the instance of it on the node)
+	FString GetName() const { return Name; }
+
+	// returns the display name of the trait
+	virtual FString GetDisplayName() const
+	{
+		return FString();
+	}
+
+#if WITH_EDITOR
+
+	// returns true if this trait can be added to a given node
+	virtual bool CanBeAddedToNode(URigVMNode* InNode, FString* OutFailureReason) const { return true; }
+
+	// allows the trait to react when added to a node
+	virtual void OnTraitAdded(URigVMController* InController, URigVMNode* InNode) {}
+
+	// allows the trait to return dynamic pins (parent pin index must be INDEX_NONE or point to a valid index of the parent pin in the OutPinArray)
+	virtual void GetProgrammaticPins(URigVMController* InController, int32 InParentPinIndex, const FString& InDefaultValue, struct FRigVMPinInfoArray& OutPinArray) const {}
+
+	virtual UScriptStruct* GetTraitSharedDataStruct() const { return nullptr; }
+
+	virtual bool ShouldCreatePinForProperty(const FProperty* InProperty) const override
+	{
+		if(!Super::ShouldCreatePinForProperty(InProperty))
+		{
+			return false;
+		}
+		return InProperty->GetFName() != GET_MEMBER_NAME_CHECKED(FRigVMTrait, Name);
+	}
+
+#endif
+
+private:
+
+	// The name of the trait on the node
+	UPROPERTY()
+	FString Name;
+	
+	friend class URigVMNode;
+	friend class URigVMController;
+	friend class FRigVMTraitScope;
+};
diff --git a/Engine/Plugins/Runtime/RigVM/Source/RigVM/Public/RigVMCore/RigVMTraitScope.h b/Engine/Plugins/Runtime/RigVM/Source/RigVM/Public/RigVMCore/RigVMTraitScope.h
new file mode 100644
--- /dev/null
+++ b/Engine/Plugins/Runtime/RigVM/Source/RigVM/Public/RigVMCore/RigVMTraitScope.h
@@ -0,0 +1,80 @@
+// Copyright Epic Games, Inc. All Rights Reserved.
+
+#pragma once
+
+#include "CoreMinimal.h"
+#include "UObject/Class.h"
+
+struct FRigVMTrait;
+class UScriptStruct;
+
+class RIGVM_API FRigVMTraitScope
+{
+public:
+	FRigVMTraitScope()
+		: Trait(nullptr)
+		, ScriptStruct(nullptr)
+	{
+	}
+	
+	FRigVMTraitScope(FRigVMTrait* InTrait, const UScriptStruct* InScriptStruct)
+		: Trait(InTrait)
+		, ScriptStruct(InScriptStruct)
+	{
+	}
+
+	bool IsValid() const
+	{
+		return (Trait != nullptr) && (ScriptStruct != nullptr);
+	}
+
+	template<typename T>
+	bool IsA() const
+	{
+		return ScriptStruct->IsChildOf(T::StaticStruct());
+	}
+
+	template<typename T = FRigVMTrait>
+	const T* GetTrait() const
+	{
+		if(IsA<T>())
+		{
+			return static_cast<T*>(Trait);
+		}
+		return nullptr;
+	}
+
+	template<typename T = FRigVMTrait>
+	const T* GetTraitChecked() const
+	{
+		check(IsA<T>());
+		return static_cast<T*>(Trait);
+	}
+
+	template<typename T = FRigVMTrait>
+	T* GetTrait()
+	{
+		if(IsA<T>())
+		{
+			return static_cast<T*>(Trait);
+		}
+		return nullptr;
+	}
+
+	template<typename T = FRigVMTrait>
+	T* GetTraitChecked()
+	{
+		check(IsA<T>());
+		return static_cast<T*>(Trait);
+	}
+
+	const UScriptStruct* GetScriptStruct() const
+	{
+		return ScriptStruct;
+	}
+
+private:
+
+	FRigVMTrait* Trait;
+	const UScriptStruct* ScriptStruct;
+};
\ No newline at end of file
diff --git a/Engine/Plugins/Runtime/RigVM/Source/RigVMDeveloper/Private/RigVMCompiler/RigVMAST.cpp b/Engine/Plugins/Runtime/RigVM/Source/RigVMDeveloper/Private/RigVMCompiler/RigVMAST.cpp
--- a/Engine/Plugins/Runtime/RigVM/Source/RigVMDeveloper/Private/RigVMCompiler/RigVMAST.cpp
+++ b/Engine/Plugins/Runtime/RigVM/Source/RigVMDeveloper/Private/RigVMCompiler/RigVMAST.cpp
@@ -1266,8 +1266,8 @@ TArray<FRigVMExprAST*> FRigVMParserAST::TraversePins(const FRigVMASTProxy& InNod
 		{
 			for (URigVMPin* Pin : Node->GetPins())
 			{
-				// We skip decorator pins as we don't want to traverse them
-				if (!Pin->IsDecoratorPin()  || Settings.bSetupDecorators)
+				// We skip trait pins as we don't want to traverse them
+				if (!Pin->IsTraitPin()  || Settings.bSetupTraits)
 				{
 					Pins.AddUnique(Pin);
 				}
@@ -1357,7 +1357,7 @@ FRigVMExprAST* FRigVMParserAST::TraversePin(const FRigVMASTProxy& InPinProxy, FR
 		}
 	}
 
-	if(Pin->IsDecoratorPin())
+	if(Pin->IsTraitPin())
 	{
 		PinExpr = MakeExpr<FRigVMVarExprAST>(FRigVMExprAST::EType::Var, InPinProxy);
 	}
diff --git a/Engine/Plugins/Runtime/RigVM/Source/RigVMDeveloper/Private/RigVMCompiler/RigVMCodeGenerator.cpp b/Engine/Plugins/Runtime/RigVM/Source/RigVMDeveloper/Private/RigVMCompiler/RigVMCodeGenerator.cpp
--- a/Engine/Plugins/Runtime/RigVM/Source/RigVMDeveloper/Private/RigVMCompiler/RigVMCodeGenerator.cpp
+++ b/Engine/Plugins/Runtime/RigVM/Source/RigVMDeveloper/Private/RigVMCompiler/RigVMCodeGenerator.cpp
@@ -1574,9 +1574,9 @@ void FRigVMCodeGenerator::ParseInstructionGroups(const FRigVMExtendedExecuteCont
 					// todo
 					break;
 				}
-				case ERigVMOpCode::SetupDecorators:
+				case ERigVMOpCode::SetupTraits:
 				{
-					const FRigVMSetupDecoratorsOp& Op = ByteCode.GetOpAt<FRigVMSetupDecoratorsOp>(Instruction);
+					const FRigVMSetupTraitsOp& Op = ByteCode.GetOpAt<FRigVMSetupTraitsOp>(Instruction);
 					// todo
 					break;
 				}
diff --git a/Engine/Plugins/Runtime/RigVM/Source/RigVMDeveloper/Private/RigVMCompiler/RigVMCompiler.cpp b/Engine/Plugins/Runtime/RigVM/Source/RigVMDeveloper/Private/RigVMCompiler/RigVMCompiler.cpp
--- a/Engine/Plugins/Runtime/RigVM/Source/RigVMDeveloper/Private/RigVMCompiler/RigVMCompiler.cpp
+++ b/Engine/Plugins/Runtime/RigVM/Source/RigVMDeveloper/Private/RigVMCompiler/RigVMCompiler.cpp
@@ -2081,39 +2081,39 @@ bool URigVMCompiler::TraverseCallExtern(const FRigVMCallExternExprAST* InExpr, F
 			return false;
 		}
 
-		if(WorkData.Settings.ASTSettings.bSetupDecorators)
+		if(WorkData.Settings.ASTSettings.bSetupTraits)
 		{
-			const TArray<URigVMPin*> DecoratorPins = Node->GetDecoratorPins(); 
-			if(!DecoratorPins.IsEmpty())
+			const TArray<URigVMPin*> TraitPins = Node->GetTraitPins(); 
+			if(!TraitPins.IsEmpty())
 			{
-				// also take care of the empty decorator list
-				if(!WorkData.DecoratorListLiterals.Contains(nullptr))
+				// also take care of the empty trait list
+				if(!WorkData.TraitListLiterals.Contains(nullptr))
 				{
-					const FName Name = WorkData.GetUniquePropertyName(ERigVMMemoryType::Literal, TEXT("EmptyDecoratorList"));
+					const FName Name = WorkData.GetUniquePropertyName(ERigVMMemoryType::Literal, TEXT("EmptyTraitList"));
 					const FRigVMOperand& ListOperand = WorkData.AddProperty(ERigVMMemoryType::Literal, Name, RigVMTypeUtils::ArrayTypeFromBaseType(RigVMTypeUtils::Int32Type), nullptr, TEXT("()"));
-					WorkData.DecoratorListLiterals.Add(nullptr, ListOperand);
+					WorkData.TraitListLiterals.Add(nullptr, ListOperand);
 				}
 
-				if(!WorkData.DecoratorListLiterals.Contains(Node))
+				if(!WorkData.TraitListLiterals.Contains(Node))
 				{
 					TArray<FString> DefaultValues;
-					for(const URigVMPin* DecoratorPin : DecoratorPins)
+					for(const URigVMPin* TraitPin : TraitPins)
 					{
-						if(const FRigVMExprAST* DecoratorExpr = InExpr->FindVarWithPinName(DecoratorPin->GetFName()))
+						if(const FRigVMExprAST* TraitExpr = InExpr->FindVarWithPinName(TraitPin->GetFName()))
 						{
-							check(DecoratorExpr->IsVar());
-							const FRigVMOperand DecoratorOperand = FindOrAddRegister(DecoratorExpr->To<FRigVMVarExprAST>(), WorkData);
-							check(DecoratorOperand.GetMemoryType() == ERigVMMemoryType::Work);
-							DefaultValues.Add(FString::FromInt(DecoratorOperand.GetRegisterIndex()));
+							check(TraitExpr->IsVar());
+							const FRigVMOperand TraitOperand = FindOrAddRegister(TraitExpr->To<FRigVMVarExprAST>(), WorkData);
+							check(TraitOperand.GetMemoryType() == ERigVMMemoryType::Work);
+							DefaultValues.Add(FString::FromInt(TraitOperand.GetRegisterIndex()));
 						}
 					}
 
 					if(!DefaultValues.IsEmpty())
 					{
-						const FName Name = WorkData.GetUniquePropertyName(ERigVMMemoryType::Literal, TEXT("DecoratorList"));
+						const FName Name = WorkData.GetUniquePropertyName(ERigVMMemoryType::Literal, TEXT("TraitList"));
 						const FString DefaultValue = FString::Printf(TEXT("(%s)"), *FString::Join(DefaultValues, TEXT(",")));
 						const FRigVMOperand& ListOperand = WorkData.AddProperty(ERigVMMemoryType::Literal, Name, RigVMTypeUtils::ArrayTypeFromBaseType(RigVMTypeUtils::Int32Type), nullptr, DefaultValue);
-						WorkData.DecoratorListLiterals.Add(Node, ListOperand);
+						WorkData.TraitListLiterals.Add(Node, ListOperand);
 					}
 				}
 			}
@@ -2306,16 +2306,16 @@ bool URigVMCompiler::TraverseCallExtern(const FRigVMCallExternExprAST* InExpr, F
 			return false;
 		}
 
-		// setup the decorator list for the context
-		bool bSetupDecorators = false;
-		if(const FRigVMOperand* DecoratorListOperand = WorkData.DecoratorListLiterals.Find(Node))
+		// setup the trait list for the context
+		bool bSetupTraits = false;
+		if(const FRigVMOperand* TraitListOperand = WorkData.TraitListLiterals.Find(Node))
 		{
 			if (WorkData.Settings.SetupNodeInstructionIndex)
 			{
 				WorkData.VM->GetByteCode().SetSubject(WorkData.VM->GetByteCode().GetNumInstructions(), Callstack.GetCallPath(), Callstack.GetStack());
 			}
-			WorkData.VM->GetByteCode().AddSetupDecoratorsOp(*DecoratorListOperand);
-			bSetupDecorators = true;
+			WorkData.VM->GetByteCode().AddSetupTraitsOp(*TraitListOperand);
+			bSetupTraits = true;
 		}
 
 		// setup the instruction
@@ -2474,16 +2474,16 @@ bool URigVMCompiler::TraverseCallExtern(const FRigVMCallExternExprAST* InExpr, F
 			}
 		}
 
-		if(bSetupDecorators)
+		if(bSetupTraits)
 		{
-			// passing nullptr retrieves the empty decorator list
-			if(const FRigVMOperand* EmptyDecoratorListOperand = WorkData.DecoratorListLiterals.Find(nullptr))
+			// passing nullptr retrieves the empty trait list
+			if(const FRigVMOperand* EmptyTraitListOperand = WorkData.TraitListLiterals.Find(nullptr))
 			{
 				if (WorkData.Settings.SetupNodeInstructionIndex)
 				{
 					WorkData.VM->GetByteCode().SetSubject(WorkData.VM->GetByteCode().GetNumInstructions(), Callstack.GetCallPath(), Callstack.GetStack());
 				}
-				WorkData.VM->GetByteCode().AddSetupDecoratorsOp(*EmptyDecoratorListOperand);
+				WorkData.VM->GetByteCode().AddSetupTraitsOp(*EmptyTraitListOperand);
 			}
 		}
 	}
@@ -2522,20 +2522,20 @@ bool URigVMCompiler::TraverseInlineFunction(const FRigVMInlineFunctionExprAST* I
 			return false;
 		}
 
-		// create a map of all of the decorator setup lists (indices to decorator properties)
-		TMap<int32, FRigVMOperand> LiteralValueToSetupDecoratorArg;
+		// create a map of all of the trait setup lists (indices to trait properties)
+		TMap<int32, FRigVMOperand> LiteralValueToSetupTraitArg;
 		const FRigVMInstructionArray FunctionInstructions = FunctionByteCode.GetInstructions();
 		for(const FRigVMInstruction& Instruction : FunctionInstructions)
 		{
-			if(Instruction.OpCode == ERigVMOpCode::SetupDecorators)
+			if(Instruction.OpCode == ERigVMOpCode::SetupTraits)
 			{
-				const FRigVMSetupDecoratorsOp& Op = FunctionByteCode.GetOpAt<FRigVMSetupDecoratorsOp>(Instruction);
+				const FRigVMSetupTraitsOp& Op = FunctionByteCode.GetOpAt<FRigVMSetupTraitsOp>(Instruction);
 				check(Op.Arg.GetMemoryType() == ERigVMMemoryType::Literal);
 				check(FunctionCompilationData->LiteralPropertyDescriptions.IsValidIndex(Op.Arg.GetRegisterIndex()));
 				const FString& DefaultValue = FunctionCompilationData->LiteralPropertyDescriptions[Op.Arg.GetRegisterIndex()].DefaultValue;
 				if(!DefaultValue.IsEmpty() && DefaultValue != TEXT("()"))
 				{
-					LiteralValueToSetupDecoratorArg.Add(Op.Arg.GetRegisterIndex(), FRigVMOperand());
+					LiteralValueToSetupTraitArg.Add(Op.Arg.GetRegisterIndex(), FRigVMOperand());
 				}
 			}
 		}
@@ -2598,12 +2598,12 @@ bool URigVMCompiler::TraverseInlineFunction(const FRigVMInlineFunctionExprAST* I
 				}
 			}
 
-			auto FindOrAddProperty = [&WorkData, MemoryType, LiteralValueToSetupDecoratorArg]
+			auto FindOrAddProperty = [&WorkData, MemoryType, LiteralValueToSetupTraitArg]
 			(const FRigVMFunctionCompilationPropertyDescription& InProperty, const FString& InNewName, const bool bIsExecuteState) -> FRigVMOperand
 			{
 				// Sharing / reusing memory / operands happens as per following contract:
 				// 1. properties are only shared if their CPP type matches
-				// 2. Literal / constant memory is only shared if the constant values match (and it is not a decorator list)
+				// 2. Literal / constant memory is only shared if the constant values match (and it is not a trait list)
 				// 3. Work state is only shared if it is not internal work state private to the instruction referring to it
 				// 4. Work state of type FRigVMInstructionSetExecuteState (bIsExecuteState) is never shared either since it is work state private to a lazy branch.
 
@@ -2621,9 +2621,9 @@ bool URigVMCompiler::TraverseInlineFunction(const FRigVMInlineFunctionExprAST* I
 					{
 						if(MemoryType == ERigVMMemoryType::Literal)
 						{
-							// if the value is the same and this is not a decorator setup list
+							// if the value is the same and this is not a trait setup list
 							if(ExistingProperty.DefaultValue.Equals(InProperty.DefaultValue) &&
-								!LiteralValueToSetupDecoratorArg.Contains(Operand.GetRegisterIndex()))
+								!LiteralValueToSetupTraitArg.Contains(Operand.GetRegisterIndex()))
 							{
 								return Operand;
 							}
@@ -2667,8 +2667,8 @@ bool URigVMCompiler::TraverseInlineFunction(const FRigVMInlineFunctionExprAST* I
 				// instantiate function library specific work state as well as
 				// instruction set execute state - which is used for lazy blocks.
 				const bool bIsExecuteState = Description.CPPType.Equals(RigVMInstructionSetExecuteStateName);
-				const bool bIsDecorator = ScriptStruct && ScriptStruct->IsChildOf(FRigVMDecorator::StaticStruct());
-				if (NewName.StartsWith(FunctionLibraryPrefix) || bIsExecuteState || bIsDecorator)
+				const bool bIsTrait = ScriptStruct && ScriptStruct->IsChildOf(FRigVMTrait::StaticStruct());
+				if (NewName.StartsWith(FunctionLibraryPrefix) || bIsExecuteState || bIsTrait)
 				{
 					NewName = FString::Printf(TEXT("%s%s"), *FunctionReferenceNode->GetNodePath(), *NewName.RightChop(FunctionLibraryPrefix.Len()));
 					FRigVMPropertyDescription::SanitizeName(NewName);
@@ -2678,23 +2678,23 @@ bool URigVMCompiler::TraverseInlineFunction(const FRigVMInlineFunctionExprAST* I
 				FRigVMCompilerWorkData::FFunctionRegisterData Data = {FunctionReferenceNode, MemoryType, PropertyIndex};
 				WorkData.FunctionRegisterToOperand.Add(Data, Operand);
 
-				if(MemoryType == ERigVMMemoryType::Literal && LiteralValueToSetupDecoratorArg.Contains(PropertyIndex))
+				if(MemoryType == ERigVMMemoryType::Literal && LiteralValueToSetupTraitArg.Contains(PropertyIndex))
 				{
-					LiteralValueToSetupDecoratorArg.FindChecked(PropertyIndex) = Operand;
+					LiteralValueToSetupTraitArg.FindChecked(PropertyIndex) = Operand;
 				}
 			}
 		}
-		// For decorator setup lists we need to update the integer values (pointing to decorator property indices)
-		for(const TPair<int32, FRigVMOperand>& Pair : LiteralValueToSetupDecoratorArg)
+		// For trait setup lists we need to update the integer values (pointing to trait property indices)
+		for(const TPair<int32, FRigVMOperand>& Pair : LiteralValueToSetupTraitArg)
 		{
-			const FRigVMOperand DecoratorIndicesOperand = Pair.Value;
-			const int32 LiteralPropertyIndex = DecoratorIndicesOperand.GetRegisterIndex();
-			FString OriginalDecoratorIndicesString = WorkData.PropertyDescriptions[ERigVMMemoryType::Literal][LiteralPropertyIndex].DefaultValue;
-			if(!OriginalDecoratorIndicesString.IsEmpty() && OriginalDecoratorIndicesString != TEXT("()"))
+			const FRigVMOperand TraitIndicesOperand = Pair.Value;
+			const int32 LiteralPropertyIndex = TraitIndicesOperand.GetRegisterIndex();
+			FString OriginalTraitIndicesString = WorkData.PropertyDescriptions[ERigVMMemoryType::Literal][LiteralPropertyIndex].DefaultValue;
+			if(!OriginalTraitIndicesString.IsEmpty() && OriginalTraitIndicesString != TEXT("()"))
 			{
-				OriginalDecoratorIndicesString = OriginalDecoratorIndicesString.TrimChar(TEXT('('));
-				OriginalDecoratorIndicesString = OriginalDecoratorIndicesString.TrimChar(TEXT(')'));
-				FString PinPathRemaining = OriginalDecoratorIndicesString;
+				OriginalTraitIndicesString = OriginalTraitIndicesString.TrimChar(TEXT('('));
+				OriginalTraitIndicesString = OriginalTraitIndicesString.TrimChar(TEXT(')'));
+				FString PinPathRemaining = OriginalTraitIndicesString;
 				FString Left, Right;
 				TArray<FString> IndexStrings;
 				while(PinPathRemaining.Split(TEXT(","), &Left, &Right))
diff --git a/Engine/Plugins/Runtime/RigVM/Source/RigVMDeveloper/Private/RigVMModel/RigVMController.cpp b/Engine/Plugins/Runtime/RigVM/Source/RigVMDeveloper/Private/RigVMModel/RigVMController.cpp
--- a/Engine/Plugins/Runtime/RigVM/Source/RigVMDeveloper/Private/RigVMModel/RigVMController.cpp
+++ b/Engine/Plugins/Runtime/RigVM/Source/RigVMDeveloper/Private/RigVMModel/RigVMController.cpp
@@ -373,7 +373,7 @@ FRigVMPinInfo::FRigVMPinInfo()
 	, bIsExpanded(false)
 	, bIsConstant(false)
 	, bIsDynamicArray(false)
-	, bIsDecorator(false)
+	, bIsTrait(false)
 	, bIsLazy(false)
 {
 }
@@ -389,7 +389,7 @@ FRigVMPinInfo::FRigVMPinInfo(const URigVMPin* InPin, int32 InParentIndex, ERigVM
 	, bIsExpanded(InPin->IsExpanded())
 	, bIsConstant(InPin->IsDefinedAsConstant())
 	, bIsDynamicArray(InPin->IsDynamicArray())
-	, bIsDecorator(InPin->IsDecoratorPin() && InPin->IsRootPin())
+	, bIsTrait(InPin->IsTraitPin() && InPin->IsRootPin())
 	, bIsLazy(InPin->IsLazy() && InPin->IsRootPin())
 {
 	// this method describes the info as currently represented in the model.
@@ -415,7 +415,7 @@ FRigVMPinInfo::FRigVMPinInfo(FProperty* InProperty, ERigVMPinDirection InDirecti
 	, bIsExpanded(false)
 	, bIsConstant(false)
 	, bIsDynamicArray(false)
-	, bIsDecorator(false)
+	, bIsTrait(false)
 	, bIsLazy(false)
 {
 	// this method describes the info as needed based on the property structure
@@ -539,7 +539,7 @@ uint32 GetTypeHash(const FRigVMPinInfo& InPin)
 	Hash = HashCombine(Hash, GetTypeHash((int32)InPin.Direction));
 	Hash = HashCombine(Hash, GetTypeHash((int32)InPin.TypeIndex));
 	Hash = HashCombine(Hash, GetTypeHash(InPin.bIsArray));
-	Hash = HashCombine(Hash, GetTypeHash(InPin.bIsDecorator));
+	Hash = HashCombine(Hash, GetTypeHash(InPin.bIsTrait));
 	// we are not hashing the parent index,  pinpath, default value or the property since
 	// it doesn't matter for the structure validity of the node
 	return Hash;
@@ -1278,7 +1278,7 @@ TArray<FString> URigVMController::GetAddNodePythonCommands(URigVMNode* Node) con
 				continue;
 			}
 
-			if(Pin->IsDecoratorPin())
+			if(Pin->IsTraitPin())
 			{
 				continue;
 			}
@@ -1321,32 +1321,32 @@ TArray<FString> URigVMController::GetAddNodePythonCommands(URigVMNode* Node) con
 			}
 		}
 
-		for(const FString& DecoratorName : Node->GetDecoratorNames())
+		for(const FString& TraitName : Node->GetTraitNames())
 		{
-			Commands.Append(GetAddDecoratorPythonCommands(Node, *DecoratorName));
+			Commands.Append(GetAddTraitPythonCommands(Node, *TraitName));
 		}
 	}
 
 	return Commands;
 }
 
-TArray<FString> URigVMController::GetAddDecoratorPythonCommands(URigVMNode* Node, const FName& DecoratorName) const
+TArray<FString> URigVMController::GetAddTraitPythonCommands(URigVMNode* Node, const FName& TraitName) const
 {
 	TArray<FString> Commands;
 
 	const FString GraphName = GetSchema()->GetSanitizedGraphName(GetGraph()->GetGraphName());
 	const FString NodeName = GetSchema()->GetSanitizedNodeName(Node->GetName());
 
-	if(const URigVMPin* DecoratorPin = Node->FindDecorator(DecoratorName))
+	if(const URigVMPin* TraitPin = Node->FindTrait(TraitName))
 	{
-		const FString DecoratorStructPath = DecoratorPin->GetCPPTypeObject()->GetPathName();
-		const FString DefaultValue = DecoratorPin->GetDefaultValue();
+		const FString TraitStructPath = TraitPin->GetCPPTypeObject()->GetPathName();
+		const FString DefaultValue = TraitPin->GetDefaultValue();
 
-		Commands.Add(FString::Printf(TEXT("blueprint.get_controller_by_name('%s').add_decorator('%s', '%s', '%s', '%s')"),
+		Commands.Add(FString::Printf(TEXT("blueprint.get_controller_by_name('%s').add_trait('%s', '%s', '%s', '%s')"),
 			*GraphName,
 			*NodeName,
-			*DecoratorStructPath,
-			*DecoratorName.ToString(),
+			*TraitStructPath,
+			*TraitName.ToString(),
 			*DefaultValue));
 	}
 
@@ -9209,7 +9209,7 @@ bool URigVMController::RemovePin(URigVMPin* InPinToRemove, bool bSetupUndoRedo,
 		return false;
 	}
 
-	check(!(InPinToRemove->IsRootPin() && InPinToRemove->IsDecoratorPin()));
+	check(!(InPinToRemove->IsRootPin() && InPinToRemove->IsTraitPin()));
 
 	FRigVMControllerCompileBracketScope CompileScope(this);
 	if (bSetupUndoRedo || bForceBreakLinks)
@@ -10954,7 +10954,7 @@ bool URigVMController::RenameExposedPin(const FName& InOldPinName, const FName&
 	{
 		static bool RenamePin(const URigVMController* InController, URigVMPin* InPin, const FName& InNewName)
 		{
-			check(!InPin->IsDecoratorPin());
+			check(!InPin->IsTraitPin());
 			
 			URigVMController* PinController = InController->GetControllerForGraph(InPin->GetGraph());
 			if(PinController == nullptr)
@@ -14403,8 +14403,8 @@ URigVMInvokeEntryNode* URigVMController::AddInvokeEntryNode(const FName& InEntry
 	return Node;
 }
 
-FName URigVMController::AddDecorator(const FName& InNodeName, const FName& InDecoratorTypeObjectPath,
-	const FName& InDecoratorName, const FString& InDefaultValue, int32 InPinIndex, bool bSetupUndoRedo,
+FName URigVMController::AddTrait(const FName& InNodeName, const FName& InTraitTypeObjectPath,
+	const FName& InTraitName, const FString& InDefaultValue, int32 InPinIndex, bool bSetupUndoRedo,
 	bool bPrintPythonCommand)
 {
 	if (!IsValidGraph())
@@ -14422,41 +14422,41 @@ FName URigVMController::AddDecorator(const FName& InNodeName, const FName& InDec
 
 	if (URigVMNode* Node = Graph->FindNodeByName(InNodeName))
 	{
-		UObject* DecoratorCPPTypeObject = RigVMTypeUtils::FindObjectFromCPPTypeObjectPath(InDecoratorTypeObjectPath.ToString());
-		if(DecoratorCPPTypeObject == nullptr)
+		UObject* TraitCPPTypeObject = RigVMTypeUtils::FindObjectFromCPPTypeObjectPath(InTraitTypeObjectPath.ToString());
+		if(TraitCPPTypeObject == nullptr)
 		{
-			DecoratorCPPTypeObject = URigVMCompiler::GetScriptStructForCPPType(InDecoratorTypeObjectPath.ToString());
+			TraitCPPTypeObject = URigVMCompiler::GetScriptStructForCPPType(InTraitTypeObjectPath.ToString());
 		}
-		if(DecoratorCPPTypeObject == nullptr)
+		if(TraitCPPTypeObject == nullptr)
 		{
-			ReportErrorf(TEXT("Cannot find decorator script struct '%s'."), *InDecoratorTypeObjectPath.ToString());
+			ReportErrorf(TEXT("Cannot find trait script struct '%s'."), *InTraitTypeObjectPath.ToString());
 			return NAME_None;
 		}
 
-		UScriptStruct* DecoratorScriptStruct = Cast<UScriptStruct>(DecoratorCPPTypeObject);
-		if(DecoratorScriptStruct == nullptr)
+		UScriptStruct* TraitScriptStruct = Cast<UScriptStruct>(TraitCPPTypeObject);
+		if(TraitScriptStruct == nullptr)
 		{
-			ReportErrorf(TEXT("CPP Type Object '%s' is not a struct."), *InDecoratorTypeObjectPath.ToString());
+			ReportErrorf(TEXT("CPP Type Object '%s' is not a struct."), *InTraitTypeObjectPath.ToString());
 			return NAME_None;
 		}
 
-		const FName DecoratorName = AddDecorator(Node, DecoratorScriptStruct, InDecoratorName, InDefaultValue, InPinIndex, bSetupUndoRedo);
-		if(!DecoratorName.IsNone() && bPrintPythonCommand)
+		const FName TraitName = AddTrait(Node, TraitScriptStruct, InTraitName, InDefaultValue, InPinIndex, bSetupUndoRedo);
+		if(!TraitName.IsNone() && bPrintPythonCommand)
 		{
-			const TArray<FString> DecoratorCommands = GetAddDecoratorPythonCommands(Node, DecoratorName);
-			for(const FString& DecoratorCommand : DecoratorCommands)
+			const TArray<FString> TraitCommands = GetAddTraitPythonCommands(Node, TraitName);
+			for(const FString& TraitCommand : TraitCommands)
 			{
-				RigVMPythonUtils::Print(GetSchema()->GetGraphOuterName(GetGraph()), DecoratorCommand);
+				RigVMPythonUtils::Print(GetSchema()->GetGraphOuterName(GetGraph()), TraitCommand);
 			}
 		}
-		return DecoratorName;
+		return TraitName;
 	}
 
 	ReportErrorf(TEXT("Cannot find node '%s'."), *InNodeName.ToString());
 	return NAME_None;
 }
 
-FName URigVMController::AddDecorator(URigVMNode* InNode, UScriptStruct* InDecoratorScriptStruct, const FName& InDecoratorName, const FString& InDefaultValue, int32 InPinIndex, bool bSetupUndoRedo)
+FName URigVMController::AddTrait(URigVMNode* InNode, UScriptStruct* InTraitScriptStruct, const FName& InTraitName, const FString& InDefaultValue, int32 InPinIndex, bool bSetupUndoRedo)
 {
 	if (!IsValidNodeForGraph(InNode))
 	{
@@ -14468,32 +14468,32 @@ FName URigVMController::AddDecorator(URigVMNode* InNode, UScriptStruct* InDecora
 		return NAME_None;
 	}
 
-	check(InDecoratorScriptStruct);
+	check(InTraitScriptStruct);
 
-	if(!InDecoratorScriptStruct->IsChildOf(FRigVMDecorator::StaticStruct()))
+	if(!InTraitScriptStruct->IsChildOf(FRigVMTrait::StaticStruct()))
 	{
-		ReportErrorf(TEXT("CPP Type Object '%s' is not a struct."), *InDecoratorScriptStruct->GetPathName());
+		ReportErrorf(TEXT("CPP Type Object '%s' is not a struct."), *InTraitScriptStruct->GetPathName());
 		return NAME_None;
 	}
 
-	const FRigVMTemplateArgumentType DecoratorType(InDecoratorScriptStruct);
-	const TRigVMTypeIndex DecoratorTypeIndex = FRigVMRegistry::Get().FindOrAddType(DecoratorType);
+	const FRigVMTemplateArgumentType TraitType(InTraitScriptStruct);
+	const TRigVMTypeIndex TraitTypeIndex = FRigVMRegistry::Get().FindOrAddType(TraitType);
 
 	if(const URigVMSchema* Schema = GetSchema())
 	{
-		if(!Schema->SupportsType(this, DecoratorTypeIndex))
+		if(!Schema->SupportsType(this, TraitTypeIndex))
 		{
-			ReportError(TEXT("Decorator cannot be added to node: Schema doesn't support the type."));
+			ReportError(TEXT("Trait cannot be added to node: Schema doesn't support the type."));
 			return NAME_None;
 		}
 	}
 
-	const FName ValidDecoratorName = URigVMSchema::GetUniqueName(InDecoratorName, [InNode](const FName& InName) {
+	const FName ValidTraitName = URigVMSchema::GetUniqueName(InTraitName, [InNode](const FName& InName) {
 		return InNode->FindPin(InName.ToString()) == nullptr;
 	}, false, false);
 
-	TSharedPtr<FStructOnScope> DecoratorScope(new FStructOnScope(InDecoratorScriptStruct));
-	FRigVMDecorator* Decorator = (FRigVMDecorator*)DecoratorScope->GetStructMemory();
+	TSharedPtr<FStructOnScope> TraitScope(new FStructOnScope(InTraitScriptStruct));
+	FRigVMTrait* Trait = (FRigVMTrait*)TraitScope->GetStructMemory();
 
 	if(!InDefaultValue.IsEmpty())
 	{
@@ -14501,44 +14501,44 @@ FName URigVMController::AddDecorator(URigVMNode* InNode, UScriptStruct* InDecora
 		{
 			// force logging to the error pipe for error detection
 			LOG_SCOPE_VERBOSITY_OVERRIDE(LogExec, ErrorPipe.GetMaxVerbosity()); 
-			InDecoratorScriptStruct->ImportText(*InDefaultValue, Decorator, nullptr, PPF_None, &ErrorPipe, InDecoratorScriptStruct->GetName()); 
+			InTraitScriptStruct->ImportText(*InDefaultValue, Trait, nullptr, PPF_None, &ErrorPipe, InTraitScriptStruct->GetName()); 
 		}
 	}
 
-	Decorator->Name = ValidDecoratorName.ToString();
+	Trait->Name = ValidTraitName.ToString();
 
 	FString FailureReason;
-	if(!Decorator->CanBeAddedToNode(InNode, &FailureReason))
+	if(!Trait->CanBeAddedToNode(InNode, &FailureReason))
 	{
-		ReportErrorf(TEXT("Decorator cannot be added to node: %s"), *FailureReason);
+		ReportErrorf(TEXT("Trait cannot be added to node: %s"), *FailureReason);
 		return NAME_None;
 	}
 
 	FRigVMBaseAction Action(this);
 	if (bSetupUndoRedo)
 	{
-		Action.SetTitle(FString::Printf(TEXT("Add Decorator")));
+		Action.SetTitle(FString::Printf(TEXT("Add Trait")));
 		GetActionStack()->BeginAction(Action);
-		GetActionStack()->AddAction(FRigVMAddDecoratorAction(this, InNode, ValidDecoratorName, InDecoratorScriptStruct, InDefaultValue, InPinIndex == INDEX_NONE ? InNode->GetPins().Num() : InPinIndex));
+		GetActionStack()->AddAction(FRigVMAddTraitAction(this, InNode, ValidTraitName, InTraitScriptStruct, InDefaultValue, InPinIndex == INDEX_NONE ? InNode->GetPins().Num() : InPinIndex));
 	}
 
-	InNode->DecoratorRootPinNames.Add(ValidDecoratorName.ToString());
+	InNode->TraitRootPinNames.Add(ValidTraitName.ToString());
 
-	URigVMPin* DecoratorPin = NewObject<URigVMPin>(InNode, ValidDecoratorName);
-	const FString DisplayName = Decorator->GetDisplayName();
-	DecoratorPin->DisplayName = DisplayName.IsEmpty() ? FName(NAME_None) : FName(*DisplayName); 
-	DecoratorPin->CPPType = InDecoratorScriptStruct->GetStructCPPName();
-	DecoratorPin->CPPTypeObject = InDecoratorScriptStruct;
-	DecoratorPin->CPPTypeObjectPath = *DecoratorPin->CPPTypeObject->GetPathName();
-	DecoratorPin->Direction = ERigVMPinDirection::Input;
+	URigVMPin* TraitPin = NewObject<URigVMPin>(InNode, ValidTraitName);
+	const FString DisplayName = Trait->GetDisplayName();
+	TraitPin->DisplayName = DisplayName.IsEmpty() ? FName(NAME_None) : FName(*DisplayName); 
+	TraitPin->CPPType = InTraitScriptStruct->GetStructCPPName();
+	TraitPin->CPPTypeObject = InTraitScriptStruct;
+	TraitPin->CPPTypeObjectPath = *TraitPin->CPPTypeObject->GetPathName();
+	TraitPin->Direction = ERigVMPinDirection::Input;
 
-	AddNodePin(InNode, DecoratorPin);
-	Notify(ERigVMGraphNotifType::PinAdded, DecoratorPin);
+	AddNodePin(InNode, TraitPin);
+	Notify(ERigVMGraphNotifType::PinAdded, TraitPin);
 	
-	AddPinsForStruct(InDecoratorScriptStruct, InNode, DecoratorPin, DecoratorPin->GetDirection(), InDefaultValue, true);
+	AddPinsForStruct(InTraitScriptStruct, InNode, TraitPin, TraitPin->GetDirection(), InDefaultValue, true);
 
 	FRigVMPinInfoArray ProgrammaticPins;
-	Decorator->GetProgrammaticPins(this, INDEX_NONE, InDefaultValue, ProgrammaticPins);
+	Trait->GetProgrammaticPins(this, INDEX_NONE, InDefaultValue, ProgrammaticPins);
 
 	const FRigVMRegistry& Registry = FRigVMRegistry::Get();
 	const FRigVMPinInfoArray PreviousPins;
@@ -14547,35 +14547,35 @@ FName URigVMController::AddDecorator(URigVMNode* InNode, UScriptStruct* InDecora
 	{
 		const FString& PinPath = ProgrammaticPins.GetPinPath(PinIndex);
 		FString ParentPinPath, PinName;
-		UObject* OuterForPin = DecoratorPin;
+		UObject* OuterForPin = TraitPin;
 		if (URigVMPin::SplitPinPathAtEnd(PinPath, ParentPinPath, PinName))
 		{
-			OuterForPin = DecoratorPin->FindSubPin(ParentPinPath);
+			OuterForPin = TraitPin->FindSubPin(ParentPinPath);
 		}
 
 		CreatePinFromPinInfo(Registry, PreviousPins, ProgrammaticPins[PinIndex], PinPath, OuterForPin);
 	}
 
 	// move the the pin to the right index as required
-	if(DecoratorPin->GetPinIndex() != InPinIndex &&
+	if(TraitPin->GetPinIndex() != InPinIndex &&
 		InPinIndex >=0 && InPinIndex < InNode->GetPins().Num())
 	{
 		URigVMPin* LastPin =  InNode->Pins.Pop();
 		InNode->Pins.Insert(LastPin, InPinIndex);
 	}
 
-	InNode->UpdateDecoratorRootPinNames();
-	Decorator->OnDecoratorAdded(this, InNode);
+	InNode->UpdateTraitRootPinNames();
+	Trait->OnTraitAdded(this, InNode);
 
 	if (bSetupUndoRedo)
 	{
 		GetActionStack()->EndAction(Action);
 	}
 
-	return ValidDecoratorName;
+	return ValidTraitName;
 }
 
-bool URigVMController::RemoveDecorator(const FName& InNodeName, const FName& InDecoratorName, bool bSetupUndoRedo, bool bPrintPythonCommand)
+bool URigVMController::RemoveTrait(const FName& InNodeName, const FName& InTraitName, bool bSetupUndoRedo, bool bPrintPythonCommand)
 {
 	if (!IsValidGraph())
 	{
@@ -14592,17 +14592,17 @@ bool URigVMController::RemoveDecorator(const FName& InNodeName, const FName& InD
 
 	if (URigVMNode* Node = Graph->FindNodeByName(InNodeName))
 	{
-		const bool bSuccess = RemoveDecorator(Node, InDecoratorName, bSetupUndoRedo);
+		const bool bSuccess = RemoveTrait(Node, InTraitName, bSetupUndoRedo);
 		if(bSuccess && bPrintPythonCommand)
 		{
 			const FString GraphName = GetSchema()->GetSanitizedGraphName(GetGraph()->GetGraphName());
         	const FString NodeName = GetSchema()->GetSanitizedNodeName(Node->GetName());
 
 			RigVMPythonUtils::Print(GetSchema()->GetGraphOuterName(GetGraph()),
-				FString::Printf(TEXT("blueprint.get_controller_by_name('%s').remove_decorator('%s', '%s')"),
+				FString::Printf(TEXT("blueprint.get_controller_by_name('%s').remove_trait('%s', '%s')"),
 				*GraphName,
 				*NodeName,
-				*InDecoratorName.ToString()));
+				*InTraitName.ToString()));
 		}
 		return bSuccess;
 	}
@@ -14612,7 +14612,7 @@ bool URigVMController::RemoveDecorator(const FName& InNodeName, const FName& InD
 
 }
 
-bool URigVMController::RemoveDecorator(URigVMNode* InNode, const FName& InDecoratorName, bool bSetupUndoRedo)
+bool URigVMController::RemoveTrait(URigVMNode* InNode, const FName& InTraitName, bool bSetupUndoRedo)
 {
 	if(!IsValidNodeForGraph(InNode))
 	{
@@ -14624,8 +14624,8 @@ bool URigVMController::RemoveDecorator(URigVMNode* InNode, const FName& InDecora
 		return false;
 	}
 
-	URigVMPin* DecoratorPin = InNode->FindDecorator(InDecoratorName);
-	if(DecoratorPin == nullptr)
+	URigVMPin* TraitPin = InNode->FindTrait(InTraitName);
+	if(TraitPin == nullptr)
 	{
 		return false;
 	}
@@ -14633,21 +14633,21 @@ bool URigVMController::RemoveDecorator(URigVMNode* InNode, const FName& InDecora
 	FRigVMBaseAction Action(this);
 	if (bSetupUndoRedo)
 	{
-		Action.SetTitle(FString::Printf(TEXT("Remove Decorator")));
+		Action.SetTitle(FString::Printf(TEXT("Remove Trait")));
 		GetActionStack()->BeginAction(Action);
 
-		const UScriptStruct* DecoratorScriptStruct = DecoratorPin->GetScriptStruct();
-		const FString DecoratorDefaultValue = DecoratorPin->GetDefaultValue();
-		GetActionStack()->AddAction(FRigVMRemoveDecoratorAction(this, InNode, InDecoratorName, DecoratorScriptStruct, DecoratorDefaultValue, DecoratorPin->GetPinIndex()));
+		const UScriptStruct* TraitScriptStruct = TraitPin->GetScriptStruct();
+		const FString TraitDefaultValue = TraitPin->GetDefaultValue();
+		GetActionStack()->AddAction(FRigVMRemoveTraitAction(this, InNode, InTraitName, TraitScriptStruct, TraitDefaultValue, TraitPin->GetPinIndex()));
 	}
 
-	const FString DecoratorNameString = InDecoratorName.ToString();
-	(void)InNode->DecoratorRootPinNames.RemoveAll([DecoratorNameString](const FString& DecoratorRootPinName) -> bool
+	const FString TraitNameString = InTraitName.ToString();
+	(void)InNode->TraitRootPinNames.RemoveAll([TraitNameString](const FString& TraitRootPinName) -> bool
 	{
-		return DecoratorNameString.Equals(DecoratorRootPinName, ESearchCase::CaseSensitive);
+		return TraitNameString.Equals(TraitRootPinName, ESearchCase::CaseSensitive);
 	});
 
-	RemovePin(DecoratorPin, bSetupUndoRedo, true);
+	RemovePin(TraitPin, bSetupUndoRedo, true);
 
 	if(bSetupUndoRedo)
 	{
@@ -16106,21 +16106,21 @@ bool URigVMController::GenerateNewPinInfos(const FRigVMRegistry& Registry, URigV
 		return false;
 	}
 
-	// make sure the new pin infos contains the decorator pins from the last run
+	// make sure the new pin infos contains the trait pins from the last run
 	for (int32 Index = 0; Index < PreviousPinInfos.Num(); Index++)
 	{
 		const FRigVMPinInfo& PreviousPin = PreviousPinInfos[Index];
-		if (PreviousPin.bIsDecorator)
+		if (PreviousPin.bIsTrait)
 		{
 			const int32 NewPinIndex = NewPinInfos.AddPin(this, INDEX_NONE, PreviousPin.Name, PreviousPin.Direction, PreviousPin.TypeIndex, PreviousPin.DefaultValue, PreviousPin.DefaultValueType, nullptr, &PreviousPinInfos, true);
-			NewPinInfos[NewPinIndex].bIsDecorator = true;
+			NewPinInfos[NewPinIndex].bIsTrait = true;
 
 			if (URigVMPin* Pin = InNode->FindPin(PreviousPin.PinPath))
 			{
-				TSharedPtr<FStructOnScope> DecoratorScope = Pin->GetDecoratorInstance();
-				FRigVMDecorator* VMDecorator = (FRigVMDecorator*)DecoratorScope->GetStructMemory();
+				TSharedPtr<FStructOnScope> TraitScope = Pin->GetTraitInstance();
+				FRigVMTrait* VMTrait = (FRigVMTrait*)TraitScope->GetStructMemory();
 
-				VMDecorator->GetProgrammaticPins(this, NewPinIndex, Pin->GetDefaultValue(), NewPinInfos);
+				VMTrait->GetProgrammaticPins(this, NewPinIndex, Pin->GetDefaultValue(), NewPinInfos);
 			}
 		}
 	}
@@ -16542,15 +16542,15 @@ void URigVMController::RepopulatePinsOnNode(const FRigVMRegistry& Registry, cons
 		ApplyPinStates(InNode, PinStates, RedirectedPinPaths);
 	}
 
-	InNode->DecoratorRootPinNames.Reset();
+	InNode->TraitRootPinNames.Reset();
 	for (int32 Index = 0; Index < NodeData.NewPinInfos.Num(); Index++)
 	{
-		if (NodeData.NewPinInfos[Index].bIsDecorator)
+		if (NodeData.NewPinInfos[Index].bIsTrait)
 		{
-			InNode->DecoratorRootPinNames.Add(NodeData.NewPinInfos[Index].Name.ToString());
+			InNode->TraitRootPinNames.Add(NodeData.NewPinInfos[Index].Name.ToString());
 		}
 	}
-	InNode->UpdateDecoratorRootPinNames();
+	InNode->UpdateTraitRootPinNames();
 
 	if (!LinkedPaths.IsEmpty())
 	{
@@ -18371,7 +18371,7 @@ bool URigVMController::ChangePinType(URigVMPin* InPin, TRigVMTypeIndex InTypeInd
 		return false;
 	}
 
-	if(InPin->IsDecoratorPin() && InPin->GetRootPin())
+	if(InPin->IsTraitPin() && InPin->GetRootPin())
 	{
 		return false;
 	}
diff --git a/Engine/Plugins/Runtime/RigVM/Source/RigVMDeveloper/Private/RigVMModel/RigVMControllerActions.cpp b/Engine/Plugins/Runtime/RigVM/Source/RigVMDeveloper/Private/RigVMModel/RigVMControllerActions.cpp
--- a/Engine/Plugins/Runtime/RigVM/Source/RigVMDeveloper/Private/RigVMModel/RigVMControllerActions.cpp
+++ b/Engine/Plugins/Runtime/RigVM/Source/RigVMDeveloper/Private/RigVMModel/RigVMControllerActions.cpp
@@ -2285,76 +2285,76 @@ bool FRigVMReplaceNodesAction::Redo()
 	return FRigVMBaseAction::Redo();
 }
 
-FRigVMAddDecoratorAction::FRigVMAddDecoratorAction()
+FRigVMAddTraitAction::FRigVMAddTraitAction()
 : FRigVMBaseAction()
 , NodeName(NAME_None)
-, DecoratorName(NAME_None)
+, TraitName(NAME_None)
 , ScriptStructPath()
-, DecoratorDefault()
+, TraitDefault()
 , PinIndex(INDEX_NONE)
 {
 }
 
-FRigVMAddDecoratorAction::FRigVMAddDecoratorAction(URigVMController* InController, const URigVMNode* InNode, const FName& InDecoratorName,
-	const UScriptStruct* InDecoratorScriptStruct, const FString& InDecoratorDefault, int32 InPinIndex)
+FRigVMAddTraitAction::FRigVMAddTraitAction(URigVMController* InController, const URigVMNode* InNode, const FName& InTraitName,
+	const UScriptStruct* InTraitScriptStruct, const FString& InTraitDefault, int32 InPinIndex)
 : FRigVMBaseAction(InController)
 , NodeName(InNode->GetFName())
-, DecoratorName(InDecoratorName)
-, ScriptStructPath(InDecoratorScriptStruct->GetPathName())
-, DecoratorDefault(InDecoratorDefault)
+, TraitName(InTraitName)
+, ScriptStructPath(InTraitScriptStruct->GetPathName())
+, TraitDefault(InTraitDefault)
 , PinIndex(InPinIndex)
 {
 }
 
-bool FRigVMAddDecoratorAction::Undo()
+bool FRigVMAddTraitAction::Undo()
 {
 	if (!FRigVMBaseAction::Undo())
 	{
 		return false;
 	}
-	return GetController()->RemoveDecorator(NodeName, DecoratorName, false);
+	return GetController()->RemoveTrait(NodeName, TraitName, false);
 }
 
-bool FRigVMAddDecoratorAction::Redo()
+bool FRigVMAddTraitAction::Redo()
 {
 	if(!CanUndoRedo())
 	{
 		return false;
 	}
-	if (GetController()->AddDecorator(NodeName, *ScriptStructPath, DecoratorName, DecoratorDefault, PinIndex, false, false) == DecoratorName)
+	if (GetController()->AddTrait(NodeName, *ScriptStructPath, TraitName, TraitDefault, PinIndex, false, false) == TraitName)
 	{
 		return FRigVMBaseAction::Redo();
 	}
 	return false;
 }
 
-FRigVMRemoveDecoratorAction::FRigVMRemoveDecoratorAction()
-: FRigVMAddDecoratorAction()
+FRigVMRemoveTraitAction::FRigVMRemoveTraitAction()
+: FRigVMAddTraitAction()
 {
 }
 
-FRigVMRemoveDecoratorAction::FRigVMRemoveDecoratorAction(URigVMController* InController, const URigVMNode* InNode,
-	const FName& InDecoratorName, const UScriptStruct* InDecoratorScriptStruct, const FString& InDecoratorDefault, int32 InPinIndex)
-: FRigVMAddDecoratorAction(InController, InNode, InDecoratorName, InDecoratorScriptStruct, InDecoratorDefault, InPinIndex)
+FRigVMRemoveTraitAction::FRigVMRemoveTraitAction(URigVMController* InController, const URigVMNode* InNode,
+	const FName& InTraitName, const UScriptStruct* InTraitScriptStruct, const FString& InTraitDefault, int32 InPinIndex)
+: FRigVMAddTraitAction(InController, InNode, InTraitName, InTraitScriptStruct, InTraitDefault, InPinIndex)
 {
 }
 
-bool FRigVMRemoveDecoratorAction::Undo()
+bool FRigVMRemoveTraitAction::Undo()
 {
 	if (!FRigVMBaseAction::Undo())
 	{
 		return false;
 	}
-	return GetController()->AddDecorator(NodeName, *ScriptStructPath, DecoratorName, DecoratorDefault, PinIndex, false, false) == DecoratorName;
+	return GetController()->AddTrait(NodeName, *ScriptStructPath, TraitName, TraitDefault, PinIndex, false, false) == TraitName;
 }
 
-bool FRigVMRemoveDecoratorAction::Redo()
+bool FRigVMRemoveTraitAction::Redo()
 {
 	if(!CanUndoRedo())
 	{
 		return false;
 	}
-	if (GetController()->RemoveDecorator(NodeName, DecoratorName, false))
+	if (GetController()->RemoveTrait(NodeName, TraitName, false))
 	{
 		return FRigVMBaseAction::Redo();
 	}
diff --git a/Engine/Plugins/Runtime/RigVM/Source/RigVMDeveloper/Private/RigVMModel/RigVMNode.cpp b/Engine/Plugins/Runtime/RigVM/Source/RigVMDeveloper/Private/RigVMModel/RigVMNode.cpp
--- a/Engine/Plugins/Runtime/RigVM/Source/RigVMDeveloper/Private/RigVMModel/RigVMNode.cpp
+++ b/Engine/Plugins/Runtime/RigVM/Source/RigVMDeveloper/Private/RigVMModel/RigVMNode.cpp
@@ -337,21 +337,21 @@ FString URigVMNode::GetOriginalDefaultValueForRootPin(const URigVMPin* InRootPin
 	return FString();
 }
 
-void URigVMNode::UpdateDecoratorRootPinNames()
+void URigVMNode::UpdateTraitRootPinNames()
 {
-	TArray<FString> NewDecoratorRootPinNames;
+	TArray<FString> NewTraitRootPinNames;
 	for(URigVMPin* Pin : GetPins())
 	{
-		if(Pin->IsDecoratorPin())
+		if(Pin->IsTraitPin())
 		{
 			if(URigVMPin* NamePin = Pin->FindSubPin(TEXT("Name")))
 			{
 				NamePin->DefaultValue = Pin->GetName();
 			}
-			NewDecoratorRootPinNames.Add(Pin->GetName());
+			NewTraitRootPinNames.Add(Pin->GetName());
 		}
 	}
-	DecoratorRootPinNames = NewDecoratorRootPinNames;
+	TraitRootPinNames = NewTraitRootPinNames;
 }
 
 bool URigVMNode::IsSelected() const
@@ -529,82 +529,82 @@ uint32 URigVMNode::GetStructureHash() const
 	return Hash;
 }
 
-TArray<URigVMPin*> URigVMNode::GetDecoratorPins() const
+TArray<URigVMPin*> URigVMNode::GetTraitPins() const
 {
-	TArray<URigVMPin*> DecoratorPins;
-	DecoratorPins.Reserve(DecoratorRootPinNames.Num());
+	TArray<URigVMPin*> TraitPins;
+	TraitPins.Reserve(TraitRootPinNames.Num());
 	
-	for(const FString& DecoratorRootPinName : DecoratorRootPinNames)
+	for(const FString& TraitRootPinName : TraitRootPinNames)
 	{
-		URigVMPin* DecoratorPin = FindPin(DecoratorRootPinName);
-		check(DecoratorPin);
+		URigVMPin* TraitPin = FindPin(TraitRootPinName);
+		check(TraitPin);
 
-		DecoratorPins.Add(DecoratorPin);
+		TraitPins.Add(TraitPin);
 	}
 
-	return DecoratorPins;
+	return TraitPins;
 }
 
-bool URigVMNode::IsDecoratorPin(FName InName) const
+bool URigVMNode::IsTraitPin(FName InName) const
 {
 	if(const URigVMPin* Pin = FindPin(InName.ToString()))
 	{
-		return IsDecoratorPin(Pin);
+		return IsTraitPin(Pin);
 	}
 	return false;
 }
 
-bool URigVMNode::IsDecoratorPin(const URigVMPin* InDecoratorPin) const
+bool URigVMNode::IsTraitPin(const URigVMPin* InTraitPin) const
 {
-	return FindDecorator(InDecoratorPin) != nullptr;
+	return FindTrait(InTraitPin) != nullptr;
 }
 
-URigVMPin* URigVMNode::FindDecorator(const FName& InName, const FString& InSubPinPath) const
+URigVMPin* URigVMNode::FindTrait(const FName& InName, const FString& InSubPinPath) const
 {
 	const FString NameString = InName.ToString();
-	for(const FString& DecoratorRootPinName : DecoratorRootPinNames)
+	for(const FString& TraitRootPinName : TraitRootPinNames)
 	{
-		if(DecoratorRootPinName.Equals(NameString, ESearchCase::CaseSensitive))
+		if(TraitRootPinName.Equals(NameString, ESearchCase::CaseSensitive))
 		{
 			if(InSubPinPath.IsEmpty())
 			{
-				return FindPin(DecoratorRootPinName);
+				return FindPin(TraitRootPinName);
 			}
-			return FindPin(URigVMPin::JoinPinPath(DecoratorRootPinName, InSubPinPath));
+			return FindPin(URigVMPin::JoinPinPath(TraitRootPinName, InSubPinPath));
 		}
 	}
 	return nullptr;
 }
 
-URigVMPin* URigVMNode::FindDecorator(const URigVMPin* InDecoratorPin) const
+URigVMPin* URigVMNode::FindTrait(const URigVMPin* InTraitPin) const
 {
-	if(InDecoratorPin)
+	if(InTraitPin)
 	{
-		const URigVMPin* RootPin = InDecoratorPin->GetRootPin();
+		const URigVMPin* RootPin = InTraitPin->GetRootPin();
 		if(RootPin->GetNode() == this)
 		{
-			return FindDecorator(RootPin->GetFName());
+			return FindTrait(RootPin->GetFName());
 		}
 	}
 	return nullptr;
 }
 
-TSharedPtr<FStructOnScope> URigVMNode::GetDecoratorInstance(const FName& InName, bool bUseDefaultValueFromPin) const
+TSharedPtr<FStructOnScope> URigVMNode::GetTraitInstance(const FName& InName, bool bUseDefaultValueFromPin) const
 {
-	return GetDecoratorInstance(FindPin(InName.ToString()), bUseDefaultValueFromPin);
+	return GetTraitInstance(FindPin(InName.ToString()), bUseDefaultValueFromPin);
 }
 
-TSharedPtr<FStructOnScope> URigVMNode::GetDecoratorInstance(const URigVMPin* InDecoratorPin, bool bUseDefaultValueFromPin) const
+TSharedPtr<FStructOnScope> URigVMNode::GetTraitInstance(const URigVMPin* InTraitPin, bool bUseDefaultValueFromPin) const
 {
-	if(const URigVMPin* RootPin = FindDecorator(InDecoratorPin))
+	if(const URigVMPin* RootPin = FindTrait(InTraitPin))
 	{
 		check(RootPin->IsStruct());
 
 		UScriptStruct* ScriptStruct = RootPin->GetScriptStruct();
-		check(ScriptStruct->IsChildOf(FRigVMDecorator::StaticStruct()));
+		check(ScriptStruct->IsChildOf(FRigVMTrait::StaticStruct()));
 
 		TSharedPtr<FStructOnScope> Scope(new FStructOnScope(ScriptStruct));
-		FRigVMDecorator* Decorator = (FRigVMDecorator*)Scope->GetStructMemory();
+		FRigVMTrait* Trait = (FRigVMTrait*)Scope->GetStructMemory();
 
 		if(bUseDefaultValueFromPin)
 		{
@@ -615,12 +615,12 @@ TSharedPtr<FStructOnScope> URigVMNode::GetDecoratorInstance(const URigVMPin* InD
 				{
 					// force logging to the error pipe for error detection
 					LOG_SCOPE_VERBOSITY_OVERRIDE(LogExec, ErrorPipe.GetMaxVerbosity());
-					ScriptStruct->ImportText(*DefaultValue, Decorator, nullptr, PPF_SerializedAsImportText, &ErrorPipe, ScriptStruct->GetName());
+					ScriptStruct->ImportText(*DefaultValue, Trait, nullptr, PPF_SerializedAsImportText, &ErrorPipe, ScriptStruct->GetName());
 				}
 			}
 		}
 
-		Decorator->Name = RootPin->GetName();
+		Trait->Name = RootPin->GetName();
 		
 		return Scope;
 	}
@@ -629,19 +629,19 @@ TSharedPtr<FStructOnScope> URigVMNode::GetDecoratorInstance(const URigVMPin* InD
 	return EmptyScope;
 }
 
-UScriptStruct* URigVMNode::GetDecoratorScriptStruct(const FName& InName) const
+UScriptStruct* URigVMNode::GetTraitScriptStruct(const FName& InName) const
 {
-	return GetDecoratorScriptStruct(FindPin(InName.ToString()));
+	return GetTraitScriptStruct(FindPin(InName.ToString()));
 }
 
-UScriptStruct* URigVMNode::GetDecoratorScriptStruct(const URigVMPin* InDecoratorPin) const
+UScriptStruct* URigVMNode::GetTraitScriptStruct(const URigVMPin* InTraitPin) const
 {
-	if(const URigVMPin* RootPin = FindDecorator(InDecoratorPin))
+	if(const URigVMPin* RootPin = FindTrait(InTraitPin))
 	{
 		check(RootPin->IsStruct());
 
 		UScriptStruct* ScriptStruct = RootPin->GetScriptStruct();
-		check(ScriptStruct->IsChildOf(FRigVMDecorator::StaticStruct()));
+		check(ScriptStruct->IsChildOf(FRigVMTrait::StaticStruct()));
 		return ScriptStruct;
 	}
 
diff --git a/Engine/Plugins/Runtime/RigVM/Source/RigVMDeveloper/Private/RigVMModel/RigVMPin.cpp b/Engine/Plugins/Runtime/RigVM/Source/RigVMDeveloper/Private/RigVMModel/RigVMPin.cpp
--- a/Engine/Plugins/Runtime/RigVM/Source/RigVMDeveloper/Private/RigVMModel/RigVMPin.cpp
+++ b/Engine/Plugins/Runtime/RigVM/Source/RigVMDeveloper/Private/RigVMModel/RigVMPin.cpp
@@ -810,10 +810,10 @@ FString URigVMPin::GetDefaultValue(const URigVMPin::FPinOverride& InOverride, bo
 	}
 	else if (IsStruct())
 	{
-		if (SubPins.Num() > 0 || IsDecoratorPin())
+		if (SubPins.Num() > 0 || IsTraitPin())
 		{
 			TArray<FString> MemberDefaultValues;
-			if(IsDecoratorPin())
+			if(IsTraitPin())
 			{
 				MemberDefaultValues.Add(FString::Printf(TEXT("Name=\"%s\""), *GetName()));
 			}
@@ -1206,9 +1206,9 @@ FString URigVMPin::GetMetaData(FName InKey) const
 #if WITH_EDITOR
 	if(const URigVMUnitNode* UnitNode = Cast<URigVMUnitNode>(GetNode()))
 	{
-		if(IsDecoratorPin())
+		if(IsTraitPin())
 		{
-			if(const UScriptStruct* Struct = GetDecoratorScriptStruct())
+			if(const UScriptStruct* Struct = GetTraitScriptStruct())
 			{
 				if(const FProperty* Property = Struct->FindPropertyByName(GetFName()))
 				{
@@ -1220,12 +1220,12 @@ FString URigVMPin::GetMetaData(FName InKey) const
 				}
 				else
 				{
-					// Possible the pin was programmatically generated from the decorator's shared struct
-					TSharedPtr<FStructOnScope> DecoratorScope = GetDecoratorInstance();
-					if(DecoratorScope.IsValid())
+					// Possible the pin was programmatically generated from the trait's shared struct
+					TSharedPtr<FStructOnScope> TraitScope = GetTraitInstance();
+					if(TraitScope.IsValid())
 					{
-						const FRigVMDecorator* VMDecorator = (FRigVMDecorator*)DecoratorScope->GetStructMemory();
-						Struct = VMDecorator->GetDecoratorSharedDataStruct();
+						const FRigVMTrait* VMTrait = (FRigVMTrait*)TraitScope->GetStructMemory();
+						Struct = VMTrait->GetTraitSharedDataStruct();
 						Property = Struct != nullptr ? Struct->FindPropertyByName(GetFName()) : nullptr;
 						if(Property)
 						{
@@ -1529,31 +1529,31 @@ uint32 URigVMPin::GetStructureHash() const
 	return Hash;
 }
 
-bool URigVMPin::IsDecoratorPin() const
+bool URigVMPin::IsTraitPin() const
 {
 	if(const URigVMNode* Node = GetNode())
 	{
-		return Node->IsDecoratorPin(GetRootPin());
+		return Node->IsTraitPin(GetRootPin());
 	}
 	return false;
 }
 
-TSharedPtr<FStructOnScope> URigVMPin::GetDecoratorInstance(bool bUseDefaultValueFromPin) const
+TSharedPtr<FStructOnScope> URigVMPin::GetTraitInstance(bool bUseDefaultValueFromPin) const
 {
 	if(const URigVMNode* Node = GetNode())
 	{
-		return Node->GetDecoratorInstance(GetRootPin(), bUseDefaultValueFromPin);
+		return Node->GetTraitInstance(GetRootPin(), bUseDefaultValueFromPin);
 	}
 
-	static const TSharedPtr<FStructOnScope> EmptyDecorator;
-	return EmptyDecorator;
+	static const TSharedPtr<FStructOnScope> EmptyTrait;
+	return EmptyTrait;
 }
 
-UScriptStruct* URigVMPin::GetDecoratorScriptStruct() const
+UScriptStruct* URigVMPin::GetTraitScriptStruct() const
 {
 	if(const URigVMNode* Node = GetNode())
 	{
-		return Node->GetDecoratorScriptStruct(GetRootPin());
+		return Node->GetTraitScriptStruct(GetRootPin());
 	}
 
 	return nullptr;
@@ -1922,12 +1922,12 @@ bool URigVMPin::CanLink(const URigVMPin* InSourcePin, const URigVMPin* InTargetP
 		return false;
 	}
 
-	if((InSourcePin->IsDecoratorPin() && InSourcePin->IsRootPin()) ||
-		(InTargetPin->IsDecoratorPin() && InTargetPin->IsRootPin()))
+	if((InSourcePin->IsTraitPin() && InSourcePin->IsRootPin()) ||
+		(InTargetPin->IsTraitPin() && InTargetPin->IsRootPin()))
 	{
 		if(OutFailureReason)
 		{
-			*OutFailureReason = TEXT("Cannot add link to root decorator pins.");
+			*OutFailureReason = TEXT("Cannot add link to root trait pins.");
 		}
 		return false;
 	}
diff --git a/Engine/Plugins/Runtime/RigVM/Source/RigVMDeveloper/Public/RigVMCompiler/RigVMAST.h b/Engine/Plugins/Runtime/RigVM/Source/RigVMDeveloper/Public/RigVMCompiler/RigVMAST.h
--- a/Engine/Plugins/Runtime/RigVM/Source/RigVMDeveloper/Public/RigVMCompiler/RigVMAST.h
+++ b/Engine/Plugins/Runtime/RigVM/Source/RigVMDeveloper/Public/RigVMCompiler/RigVMAST.h
@@ -1149,7 +1149,7 @@ struct RIGVMDEVELOPER_API FRigVMParserASTSettings
 	bool bFoldLiterals = false;
 
 	UPROPERTY(EditAnywhere, Transient, BlueprintReadWrite, Category = "AST")
-	bool bSetupDecorators = true;
+	bool bSetupTraits = true;
 
 	// links to be ignored during the parse
 	UPROPERTY()
diff --git a/Engine/Plugins/Runtime/RigVM/Source/RigVMDeveloper/Public/RigVMCompiler/RigVMCompiler.h b/Engine/Plugins/Runtime/RigVM/Source/RigVMDeveloper/Public/RigVMCompiler/RigVMCompiler.h
--- a/Engine/Plugins/Runtime/RigVM/Source/RigVMDeveloper/Public/RigVMCompiler/RigVMCompiler.h
+++ b/Engine/Plugins/Runtime/RigVM/Source/RigVMDeveloper/Public/RigVMCompiler/RigVMCompiler.h
@@ -141,7 +141,7 @@ public:
 	TArray<const FRigVMExprAST*> ExprToSkip;
 	TArray<const FRigVMExprAST*> TraversalExpressions;
 	TMap<FString, int32> ProcessedLinks;
-	TMap<const URigVMNode*, FRigVMOperand> DecoratorListLiterals;
+	TMap<const URigVMNode*, FRigVMOperand> TraitListLiterals;
 	FRigVMOperand ComparisonOperand;
 
 	using FRigVMASTProxyArray = TArray<FRigVMASTProxy, TInlineAllocator<3>>; 
diff --git a/Engine/Plugins/Runtime/RigVM/Source/RigVMDeveloper/Public/RigVMModel/RigVMController.h b/Engine/Plugins/Runtime/RigVM/Source/RigVMDeveloper/Public/RigVMModel/RigVMController.h
--- a/Engine/Plugins/Runtime/RigVM/Source/RigVMDeveloper/Public/RigVMModel/RigVMController.h
+++ b/Engine/Plugins/Runtime/RigVM/Source/RigVMDeveloper/Public/RigVMModel/RigVMController.h
@@ -188,7 +188,7 @@ struct RIGVMDEVELOPER_API FRigVMPinInfo
 	bool bIsExpanded;
 	bool bIsConstant;
 	bool bIsDynamicArray;
-	bool bIsDecorator;
+	bool bIsTrait;
 	bool bIsLazy;
 	TArray<int32> SubPins;
 
@@ -323,7 +323,7 @@ public:
 	TArray<FString> GeneratePythonCommands();
 
 	TArray<FString> GetAddNodePythonCommands(URigVMNode* Node) const;
-	TArray<FString> GetAddDecoratorPythonCommands(URigVMNode* Node, const FName& DecoratorName) const;
+	TArray<FString> GetAddTraitPythonCommands(URigVMNode* Node, const FName& TraitName) const;
 
 	FRigVMGraphFunctionStore* GetGraphFunctionStore() const;
 	FRigVMGraphFunctionData* FindFunctionData(const FName& InFunctionName) const;
@@ -565,13 +565,13 @@ public:
 	UFUNCTION(BlueprintCallable, Category = RigVMController)
 	URigVMInvokeEntryNode* AddInvokeEntryNode(const FName& InEntryName, const FVector2D& InPosition = FVector2D::ZeroVector, const FString& InNodeName = TEXT(""), bool bSetupUndoRedo = true, bool bPrintPythonCommand = false);
 
-	// Adds a decorator to a node
+	// Adds a trait to a node
 	UFUNCTION(BlueprintCallable, Category = RigVMController)
-	FName AddDecorator(const FName& InNodeName, const FName& InDecoratorTypeObjectPath, const FName& InDecoratorName = NAME_None, const FString& InDefaultValue = TEXT(""), int32 InPinIndex = -1, bool bSetupUndoRedo = true, bool bPrintPythonCommand = false);
+	FName AddTrait(const FName& InNodeName, const FName& InTraitTypeObjectPath, const FName& InTraitName = NAME_None, const FString& InDefaultValue = TEXT(""), int32 InPinIndex = -1, bool bSetupUndoRedo = true, bool bPrintPythonCommand = false);
 
-	// Removes a decorator from a node
+	// Removes a trait from a node
 	UFUNCTION(BlueprintCallable, Category = RigVMController)
-	bool RemoveDecorator(const FName& InNodeName, const FName& InDecoratorName, bool bSetupUndoRedo = true, bool bPrintPythonCommand = false);
+	virtual bool RemoveTrait(const FName& InNodeName, const FName& InTraitName, bool bSetupUndoRedo = true, bool bPrintPythonCommand = false);
 
 	// Un-does the last action on the stack.
 	// Note: This should really only be used for unit tests,
@@ -1198,8 +1198,8 @@ private:
 	URigVMInjectionInfo* InjectNodeIntoPin(URigVMPin* InPin, bool bAsInput, const FName& InInputPinName, const FName& InOutputPinName, bool bSetupUndoRedo = true);
 	URigVMNode* EjectNodeFromPin(URigVMPin* InPin, bool bSetupUndoRedo = true, bool bPrintPythonCommands = false);
 	bool EjectAllInjectedNodes(URigVMNode* InNode, bool bSetupUndoRedo = true, bool bPrintPythonCommands = false);
-	FName AddDecorator(URigVMNode* InNode, UScriptStruct* InDecoratorScriptStruct, const FName& InDecoratorName, const FString& InDefaultValue, int32 InPinIndex = -1, bool bSetupUndoRedo = true);
-	bool RemoveDecorator(URigVMNode* InNode, const FName& InDecoratorName, bool bSetupUndoRedo = true);
+	FName AddTrait(URigVMNode* InNode, UScriptStruct* InTraitScriptStruct, const FName& InTraitName, const FString& InDefaultValue, int32 InPinIndex = -1, bool bSetupUndoRedo = true);
+	bool RemoveTrait(URigVMNode* InNode, const FName& InTraitName, bool bSetupUndoRedo = true);
 
 protected:
 
diff --git a/Engine/Plugins/Runtime/RigVM/Source/RigVMDeveloper/Public/RigVMModel/RigVMControllerActions.h b/Engine/Plugins/Runtime/RigVM/Source/RigVMDeveloper/Public/RigVMModel/RigVMControllerActions.h
--- a/Engine/Plugins/Runtime/RigVM/Source/RigVMDeveloper/Public/RigVMModel/RigVMControllerActions.h
+++ b/Engine/Plugins/Runtime/RigVM/Source/RigVMDeveloper/Public/RigVMModel/RigVMControllerActions.h
@@ -1549,19 +1549,19 @@ public:
 };
 
 /**
- * An action to add a decorator to a node
+ * An action to add a trait to a node
  */
 USTRUCT()
-struct FRigVMAddDecoratorAction : public FRigVMBaseAction
+struct FRigVMAddTraitAction : public FRigVMBaseAction
 {
 	GENERATED_BODY()
 
 public:
 
-	FRigVMAddDecoratorAction();
-	FRigVMAddDecoratorAction(URigVMController* InController, const URigVMNode* InNode, const FName& InDecoratorName, const UScriptStruct* InDecoratorScriptStruct, const FString& InDecoratorDefault, int32 InPinIndex);
-	virtual ~FRigVMAddDecoratorAction() {};
-	virtual UScriptStruct* GetScriptStruct() const override { return FRigVMAddDecoratorAction::StaticStruct(); }
+	FRigVMAddTraitAction();
+	FRigVMAddTraitAction(URigVMController* InController, const URigVMNode* InNode, const FName& InTraitName, const UScriptStruct* InTraitScriptStruct, const FString& InTraitDefault, int32 InPinIndex);
+	virtual ~FRigVMAddTraitAction() {};
+	virtual UScriptStruct* GetScriptStruct() const override { return FRigVMAddTraitAction::StaticStruct(); }
 	virtual bool Undo() override;
 	virtual bool Redo() override;
 
@@ -1569,32 +1569,32 @@ public:
 	FName NodeName;
 
 	UPROPERTY()
-	FName DecoratorName;
+	FName TraitName;
 
 	UPROPERTY()
 	FString ScriptStructPath;
 
 	UPROPERTY()
-	FString DecoratorDefault;
+	FString TraitDefault;
 
 	UPROPERTY()
 	int32 PinIndex;
 };
 
 /**
- * An action to remove a decorator from a node
+ * An action to remove a trait from a node
  */
 USTRUCT()
-struct FRigVMRemoveDecoratorAction : public FRigVMAddDecoratorAction
+struct FRigVMRemoveTraitAction : public FRigVMAddTraitAction
 {
 	GENERATED_BODY()
 
 public:
 
-	FRigVMRemoveDecoratorAction();
-	FRigVMRemoveDecoratorAction(URigVMController* InController, const URigVMNode* InNode, const FName& InDecoratorName, const UScriptStruct* InDecoratorScriptStruct, const FString& InDecoratorDefault, int32 InPinIndex);
-	virtual ~FRigVMRemoveDecoratorAction() {};
-	virtual UScriptStruct* GetScriptStruct() const override { return FRigVMRemoveDecoratorAction::StaticStruct(); }
+	FRigVMRemoveTraitAction();
+	FRigVMRemoveTraitAction(URigVMController* InController, const URigVMNode* InNode, const FName& InTraitName, const UScriptStruct* InTraitScriptStruct, const FString& InTraitDefault, int32 InPinIndex);
+	virtual ~FRigVMRemoveTraitAction() {};
+	virtual UScriptStruct* GetScriptStruct() const override { return FRigVMRemoveTraitAction::StaticStruct(); }
 	virtual bool Undo() override;
 	virtual bool Redo() override;
 };
\ No newline at end of file
diff --git a/Engine/Plugins/Runtime/RigVM/Source/RigVMDeveloper/Public/RigVMModel/RigVMNode.h b/Engine/Plugins/Runtime/RigVM/Source/RigVMDeveloper/Public/RigVMModel/RigVMNode.h
--- a/Engine/Plugins/Runtime/RigVM/Source/RigVMDeveloper/Public/RigVMModel/RigVMNode.h
+++ b/Engine/Plugins/Runtime/RigVM/Source/RigVMDeveloper/Public/RigVMModel/RigVMNode.h
@@ -5,7 +5,7 @@
 #include "RigVMPin.h"
 #include "RigVMCore/RigVM.h"
 #include "RigVMCore/RigVMStruct.h"
-#include "RigVMCore/RigVMDecorator.h"
+#include "RigVMCore/RigVMTrait.h"
 #include "RigVMCore/RigVMUserWorkflow.h"
 #include "RigVMCore/RigVMExecuteContext.h"
 #include "UObject/StructOnScope.h"
@@ -302,29 +302,29 @@ public:
 	virtual uint32 GetStructureHash() const;
 
 	// allows the node to support non-native pins
-	virtual bool HasNonNativePins() const { return !DecoratorRootPinNames.IsEmpty(); }
+	virtual bool HasNonNativePins() const { return !TraitRootPinNames.IsEmpty(); }
 
 	UFUNCTION(BlueprintPure, Category = RigVMNode)
-	TArray<URigVMPin*> GetDecoratorPins() const;
+	TArray<URigVMPin*> GetTraitPins() const;
 
-	const TArray<FString>& GetDecoratorNames() const { return DecoratorRootPinNames; }
+	const TArray<FString>& GetTraitNames() const { return TraitRootPinNames; }
 
 	UFUNCTION(BlueprintPure, Category = RigVMNode)
-	bool IsDecoratorPin(FName InName) const;
+	bool IsTraitPin(FName InName) const;
 
-	bool IsDecoratorPin(const URigVMPin* InDecoratorPin) const;
+	bool IsTraitPin(const URigVMPin* InTraitPin) const;
 
-	URigVMPin* FindDecorator(const FName& InName, const FString& InSubPinPath = FString()) const;
+	URigVMPin* FindTrait(const FName& InName, const FString& InSubPinPath = FString()) const;
 
-	URigVMPin* FindDecorator(const URigVMPin* InDecoratorPin) const;
+	URigVMPin* FindTrait(const URigVMPin* InTraitPin) const;
 
-	TSharedPtr<FStructOnScope> GetDecoratorInstance(const FName& InName, bool bUseDefaultValueFromPin = true) const;
+	TSharedPtr<FStructOnScope> GetTraitInstance(const FName& InName, bool bUseDefaultValueFromPin = true) const;
 
-	TSharedPtr<FStructOnScope> GetDecoratorInstance(const URigVMPin* InDecoratorPin, bool bUseDefaultValueFromPin = true) const;
+	TSharedPtr<FStructOnScope> GetTraitInstance(const URigVMPin* InTraitPin, bool bUseDefaultValueFromPin = true) const;
 
-	UScriptStruct* GetDecoratorScriptStruct(const FName& InName) const;
+	UScriptStruct* GetTraitScriptStruct(const FName& InName) const;
 
-	UScriptStruct* GetDecoratorScriptStruct(const URigVMPin* InDecoratorPin) const;
+	UScriptStruct* GetTraitScriptStruct(const URigVMPin* InTraitPin) const;
 	
 	virtual bool IsOutDated() const
 	{
@@ -352,7 +352,7 @@ protected:
 	virtual bool ShouldInputPinComputeLazily(const URigVMPin* InPin) const { return false; }
 	virtual FString GetOriginalDefaultValueForRootPin(const URigVMPin* InRootPin) const;
 
-	void UpdateDecoratorRootPinNames();
+	void UpdateTraitRootPinNames();
 
 	UPROPERTY()
 	FString NodeTitle;
@@ -376,7 +376,7 @@ protected:
 	bool bHaltedAtThisNode;
 
 	UPROPERTY()
-	TArray<FString> DecoratorRootPinNames;
+	TArray<FString> TraitRootPinNames;
 
 private:
 
diff --git a/Engine/Plugins/Runtime/RigVM/Source/RigVMDeveloper/Public/RigVMModel/RigVMPin.h b/Engine/Plugins/Runtime/RigVM/Source/RigVMDeveloper/Public/RigVMModel/RigVMPin.h
--- a/Engine/Plugins/Runtime/RigVM/Source/RigVMDeveloper/Public/RigVMModel/RigVMPin.h
+++ b/Engine/Plugins/Runtime/RigVM/Source/RigVMDeveloper/Public/RigVMModel/RigVMPin.h
@@ -10,7 +10,7 @@
 #include "RigVMCore/RigVMTemplate.h"
 #include "RigVMCompiler/RigVMASTProxy.h"
 #include "UObject/StructOnScope.h"
-#include <RigVMCore/RigVMDecorator.h>
+#include <RigVMCore/RigVMTrait.h>
 #include "RigVMPin.generated.h"
 
 class URigVMGraph;
@@ -493,15 +493,15 @@ public:
 
 	uint32 GetStructureHash() const;
 
-	// Returns true if this pin represents a decorator 
+	// Returns true if this pin represents a trait 
 	UFUNCTION(BlueprintPure, Category = RigVMPin)
-	bool IsDecoratorPin() const;
+	bool IsTraitPin() const;
 
-	// Returns the decorator backing up this pin
-	TSharedPtr<FStructOnScope> GetDecoratorInstance(bool bUseDefaultValueFromPin = true) const;
+	// Returns the trait backing up this pin
+	TSharedPtr<FStructOnScope> GetTraitInstance(bool bUseDefaultValueFromPin = true) const;
 
-	// Returns the struct of the decorator backing up this pin
-	UScriptStruct* GetDecoratorScriptStruct() const;
+	// Returns the struct of the trait backing up this pin
+	UScriptStruct* GetTraitScriptStruct() const;
 
 private:
 
diff --git a/Engine/Plugins/Runtime/RigVM/Source/RigVMEditor/Private/RigVMEditorModule.cpp b/Engine/Plugins/Runtime/RigVM/Source/RigVMEditor/Private/RigVMEditorModule.cpp
--- a/Engine/Plugins/Runtime/RigVM/Source/RigVMEditor/Private/RigVMEditorModule.cpp
+++ b/Engine/Plugins/Runtime/RigVM/Source/RigVMEditor/Private/RigVMEditorModule.cpp
@@ -1244,9 +1244,9 @@ void FRigVMEditorModule::GetNodeVariantContextMenuActions(IRigVMClientHost* RigV
 void FRigVMEditorModule::GetNodeTestContextMenuActions(IRigVMClientHost* RigVMClientHost, const URigVMEdGraphNode* EdGraphNode, URigVMNode* ModelNode, UToolMenu* Menu) const
 {
 	// this struct is only available in EngineTest for now
-	static const FString DecoratorObjectPath = TEXT("/Script/EngineTestEditor.EngineTestRigVM_SimpleDecorator");
-	const UScriptStruct* SimpleDecoratorStruct = Cast<UScriptStruct>(RigVMTypeUtils::FindObjectFromCPPTypeObjectPath(DecoratorObjectPath));
-	if(SimpleDecoratorStruct == nullptr)
+	static const FString TraitObjectPath = TEXT("/Script/EngineTestEditor.EngineTestRigVM_SimpleTrait");
+	const UScriptStruct* SimpleTraitStruct = Cast<UScriptStruct>(RigVMTypeUtils::FindObjectFromCPPTypeObjectPath(TraitObjectPath));
+	if(SimpleTraitStruct == nullptr)
 	{
 		return;
 	}
@@ -1256,16 +1256,16 @@ void FRigVMEditorModule::GetNodeTestContextMenuActions(IRigVMClientHost* RigVMCl
 
 	FToolMenuSection& EngineTestSection = Menu->AddSection("RigVMEditorContextMenuEngineTest", LOCTEXT("EngineTestHeader", "EngineTest"));
 	EngineTestSection.AddMenuEntry(
-		"Add simple decorator",
-		LOCTEXT("AddSimpleDecorator", "Add simple decorator"),
-		LOCTEXT("AddSimpleDecorator_Tooltip", "Adds a simple test decorator to the node"),
+		"Add simple trait",
+		LOCTEXT("AddSimpleTrait", "Add simple trait"),
+		LOCTEXT("AddSimpleTrait_Tooltip", "Adds a simple test trait to the node"),
 		FSlateIcon(),
-		FUIAction(FExecuteAction::CreateLambda([Controller, ModelNode, SimpleDecoratorStruct]()
+		FUIAction(FExecuteAction::CreateLambda([Controller, ModelNode, SimpleTraitStruct]()
 		{
-			(void)Controller->AddDecorator(
+			(void)Controller->AddTrait(
 				ModelNode->GetFName(),
-				*SimpleDecoratorStruct->GetPathName(),
-				TEXT("Decorator"),
+				*SimpleTraitStruct->GetPathName(),
+				TEXT("Trait"),
 				FString(), INDEX_NONE, true, true);
 		}))
 	);
diff --git a/Engine/Plugins/Runtime/RigVM/Source/RigVMEditor/Private/Widgets/SRigVMExecutionStackView.cpp b/Engine/Plugins/Runtime/RigVM/Source/RigVMEditor/Private/Widgets/SRigVMExecutionStackView.cpp
--- a/Engine/Plugins/Runtime/RigVM/Source/RigVMEditor/Private/Widgets/SRigVMExecutionStackView.cpp
+++ b/Engine/Plugins/Runtime/RigVM/Source/RigVMEditor/Private/Widgets/SRigVMExecutionStackView.cpp
@@ -585,9 +585,9 @@ void SRigVMExecutionStackView::PopulateStackView(URigVM* InVM, FRigVMExtendedExe
 						Label = FString::Printf(TEXT("Run Instructions %d-%d"), Op.StartInstruction, Op.EndInstruction);
 						break;
 					}
-					case ERigVMOpCode::SetupDecorators:
+					case ERigVMOpCode::SetupTraits:
 					{
-						Label = TEXT("Setup Decorators");
+						Label = TEXT("Setup Traits");
 						break;
 					}
 					case ERigVMOpCode::Exit:
