AnimNext : Fixed required interfaces of a base trait being flagged as missing if the derivated trait implements the interface

#rnx

diff --git a/Engine/Plugins/Experimental/AnimNext/Source/AnimNext/Public/TraitCore/Trait.h b/Engine/Plugins/Experimental/AnimNext/Source/AnimNext/Public/TraitCore/Trait.h
--- a/Engine/Plugins/Experimental/AnimNext/Source/AnimNext/Public/TraitCore/Trait.h
+++ b/Engine/Plugins/Experimental/AnimNext/Source/AnimNext/Public/TraitCore/Trait.h
@@ -119,14 +119,14 @@ class FArchive;
 	}
 
 // Implements GetTraitRequiredInterfaces()
-#define ANIM_NEXT_IMPL_DEFINE_ANIM_TRAIT_GET_REQUIRED_INTERFACES(TraitName, InterfaceEnumeratorMacro) \
+#define ANIM_NEXT_IMPL_DEFINE_ANIM_TRAIT_GET_REQUIRED_INTERFACES(TraitName, RequiredInterfaceEnumeratorMacro) \
 	TConstArrayView<UE::AnimNext::FTraitInterfaceUID> TraitName::GetTraitRequiredInterfaces() const \
 	{ \
 		/* Thread safe cache initialization */ \
 		static TArray<UE::AnimNext::FTraitInterfaceUID> CachedInterfaceList = FTrait::BuildTraitInterfaceList( \
 			TraitSuper::GetTraitRequiredInterfaces(), \
 			{ \
-				InterfaceEnumeratorMacro(ANIM_NEXT_IMPL_GET_INTERFACES_IMPL_FOR_INTERFACE) \
+				RequiredInterfaceEnumeratorMacro(ANIM_NEXT_IMPL_GET_INTERFACES_IMPL_FOR_INTERFACE) \
 			}); \
 		return CachedInterfaceList; \
 	}
diff --git a/Engine/Plugins/Experimental/AnimNext/Source/AnimNextEditor/Private/Graph/STraitEditorView.cpp b/Engine/Plugins/Experimental/AnimNext/Source/AnimNextEditor/Private/Graph/STraitEditorView.cpp
--- a/Engine/Plugins/Experimental/AnimNext/Source/AnimNextEditor/Private/Graph/STraitEditorView.cpp
+++ b/Engine/Plugins/Experimental/AnimNext/Source/AnimNextEditor/Private/Graph/STraitEditorView.cpp
@@ -361,9 +361,10 @@ void STraitEditorView::UpdateTraitStatusInStack(const TArray<TSharedPtr<FTraitDa
 				const FTraitInterfaceUID& RequiredInterface = RequiredInterfaces[RequiredIndex];
 
 				bool bFound = false;
-				for (int32 ParentIndex = TraitIndex - 1; ParentIndex >= 0 && !bFound; ParentIndex--)
+				// Start search from current trait to enable traits that inherit from a trait with a required interface that is implemented in the derived class
+				for (int32 SearchStartIndex = TraitIndex; SearchStartIndex >= 0 && !bFound; SearchStartIndex--)
 				{
-					const TSharedPtr<FTraitDataEditorDef>& ParentTraitData = CurrentTraitsData[ParentIndex];
+					const TSharedPtr<FTraitDataEditorDef>& ParentTraitData = CurrentTraitsData[SearchStartIndex];
 
 					const TArray<FTraitInterfaceUID>& ParentImplementedInterfaces = ParentTraitData->ImplementedInterfaces;
 
diff --git a/Engine/Plugins/Experimental/AnimNextAnimGraph/Source/AnimNextAnimGraph/Private/TraitInterfaces/ISmoothBlend.cpp b/Engine/Plugins/Experimental/AnimNextAnimGraph/Source/AnimNextAnimGraph/Private/TraitInterfaces/ISmoothBlend.cpp
--- a/Engine/Plugins/Experimental/AnimNextAnimGraph/Source/AnimNextAnimGraph/Private/TraitInterfaces/ISmoothBlend.cpp
+++ b/Engine/Plugins/Experimental/AnimNextAnimGraph/Source/AnimNextAnimGraph/Private/TraitInterfaces/ISmoothBlend.cpp
@@ -11,7 +11,7 @@ namespace UE::AnimNext
 #if WITH_EDITOR
 	const FText& ISmoothBlend::GetDisplayName() const
 	{
-		static FText InterfaceName = NSLOCTEXT("TraitInterfaces", "TraitInterface_ISmoothBlend_Name", "Group Synchronization");
+		static FText InterfaceName = NSLOCTEXT("TraitInterfaces", "TraitInterface_ISmoothBlend_Name", "Smooth Blend");
 		return InterfaceName;
 	}
 	const FText& ISmoothBlend::GetDisplayShortName() const
diff --git a/Engine/Plugins/Experimental/AnimNextAnimGraph/Source/AnimNextAnimGraph/Private/Traits/BlendInertializer.cpp b/Engine/Plugins/Experimental/AnimNextAnimGraph/Source/AnimNextAnimGraph/Private/Traits/BlendInertializer.cpp
--- a/Engine/Plugins/Experimental/AnimNextAnimGraph/Source/AnimNextAnimGraph/Private/Traits/BlendInertializer.cpp
+++ b/Engine/Plugins/Experimental/AnimNextAnimGraph/Source/AnimNextAnimGraph/Private/Traits/BlendInertializer.cpp
@@ -16,9 +16,15 @@ namespace UE::AnimNext
 	#define TRAIT_INTERFACE_ENUMERATOR(GeneratorMacro) \
 		GeneratorMacro(IDiscreteBlend) \
 		GeneratorMacro(ISmoothBlend) \
+	
+	// Trait required interfaces implementation boilerplate
+	#define TRAIT_REQUIRED_INTERFACE_ENUMERATOR(GeneratorMacro) \
+		GeneratorMacro(IInertializerBlend) \
 
-	GENERATE_ANIM_TRAIT_IMPLEMENTATION(FBlendInertializerCoreTrait, TRAIT_INTERFACE_ENUMERATOR, NULL_ANIM_TRAIT_INTERFACE_ENUMERATOR, NULL_ANIM_TRAIT_EVENT_ENUMERATOR)
+	GENERATE_ANIM_TRAIT_IMPLEMENTATION(FBlendInertializerCoreTrait, TRAIT_INTERFACE_ENUMERATOR, TRAIT_REQUIRED_INTERFACE_ENUMERATOR, NULL_ANIM_TRAIT_EVENT_ENUMERATOR)
 	#undef TRAIT_INTERFACE_ENUMERATOR
+	#undef TRAIT_REQUIRED_INTERFACE_ENUMERATOR
+
 
 	// Trait implementation boilerplate
 	#define TRAIT_INTERFACE_ENUMERATOR(GeneratorMacro) \
diff --git a/Engine/Plugins/Experimental/AnimNextAnimGraph/Source/AnimNextAnimGraph/Private/Traits/BlendSmoother.cpp b/Engine/Plugins/Experimental/AnimNextAnimGraph/Source/AnimNextAnimGraph/Private/Traits/BlendSmoother.cpp
--- a/Engine/Plugins/Experimental/AnimNextAnimGraph/Source/AnimNextAnimGraph/Private/Traits/BlendSmoother.cpp
+++ b/Engine/Plugins/Experimental/AnimNextAnimGraph/Source/AnimNextAnimGraph/Private/Traits/BlendSmoother.cpp
@@ -18,9 +18,14 @@ namespace UE::AnimNext
 		GeneratorMacro(IDiscreteBlend) \
 		GeneratorMacro(IEvaluate) \
 		GeneratorMacro(IUpdate) \
+	
+	// Trait required interfaces implementation boilerplate
+	#define TRAIT_REQUIRED_INTERFACE_ENUMERATOR(GeneratorMacroRerquired) \
+		GeneratorMacroRerquired(ISmoothBlend) \
 
-	GENERATE_ANIM_TRAIT_IMPLEMENTATION(FBlendSmootherCoreTrait, TRAIT_INTERFACE_ENUMERATOR, NULL_ANIM_TRAIT_INTERFACE_ENUMERATOR, NULL_ANIM_TRAIT_EVENT_ENUMERATOR)
+	GENERATE_ANIM_TRAIT_IMPLEMENTATION(FBlendSmootherCoreTrait, TRAIT_INTERFACE_ENUMERATOR, TRAIT_REQUIRED_INTERFACE_ENUMERATOR, NULL_ANIM_TRAIT_EVENT_ENUMERATOR)
 	#undef TRAIT_INTERFACE_ENUMERATOR
+	#undef TRAIT_REQUIRED_INTERFACE_ENUMERATOR
 
 	// Trait implementation boilerplate
 	#define TRAIT_INTERFACE_ENUMERATOR(GeneratorMacro) \
