AnimNext: Module refactor

diff --git a/Engine/Plugins/Experimental/AnimNextAnimGraph/Source/AnimNextAnimGraph/Internal/Graph/AnimGraph/AnimNode_AnimNextGraph.h b/Engine/Plugins/Experimental/AnimNextAnimGraph/Source/AnimNextAnimGraph/Internal/Graph/AnimGraph/AnimNode_AnimNextGraph.h
--- a/Engine/Plugins/Experimental/AnimNextAnimGraph/Source/AnimNextAnimGraph/Internal/Graph/AnimGraph/AnimNode_AnimNextGraph.h
+++ b/Engine/Plugins/Experimental/AnimNextAnimGraph/Source/AnimNextAnimGraph/Internal/Graph/AnimGraph/AnimNode_AnimNextGraph.h
@@ -3,10 +3,9 @@
 #pragma once
 
 #include "Animation/AnimNode_CustomProperty.h"
-#include "Module/AnimNextModule.h"
+#include "Graph/AnimNextAnimationGraph.h"
 #include "Graph/AnimNextGraphInstancePtr.h"
 #include "TraitCore/TraitPtr.h"
-#include "Context.h"
 #include "AnimNode_AnimNextGraph.generated.h"
 
 /**
@@ -42,7 +41,7 @@ private:
 	FPoseLink SourceLink;
 
 	UPROPERTY(EditAnywhere, Category = Settings, meta = (PinHiddenByDefault))
-	TObjectPtr<UAnimNextModule> Module;
+	TObjectPtr<UAnimNextAnimationGraph> AnimationGraph;
 
 	// Our graph instance, we own it
 	UPROPERTY()
@@ -58,7 +57,7 @@ private:
 	int32 LODThreshold;
 
 protected:
-	virtual UClass* GetTargetClass() const override { return Module ? Module->StaticClass() : nullptr; }
+	virtual UClass* GetTargetClass() const override { return AnimationGraph ? AnimationGraph->StaticClass() : nullptr; }
 	
 public:
 
diff --git a/Engine/Plugins/Experimental/AnimNextAnimGraph/Source/AnimNextAnimGraph/Private/Graph/AnimGraph/AnimNode_AnimNextGraph.cpp b/Engine/Plugins/Experimental/AnimNextAnimGraph/Source/AnimNextAnimGraph/Private/Graph/AnimGraph/AnimNode_AnimNextGraph.cpp
--- a/Engine/Plugins/Experimental/AnimNextAnimGraph/Source/AnimNextAnimGraph/Private/Graph/AnimGraph/AnimNode_AnimNextGraph.cpp
+++ b/Engine/Plugins/Experimental/AnimNextAnimGraph/Source/AnimNextAnimGraph/Private/Graph/AnimGraph/AnimNode_AnimNextGraph.cpp
@@ -10,8 +10,6 @@
 #include "GenerationTools.h"
 #include "Graph/AnimNext_LODPose.h"
 #include "Engine/SkeletalMesh.h"
-#include "Param/ParamStack.h"
-#include "Graph/AnimGraphParamStackScope.h"
 #include "TraitInterfaces/IEvaluate.h"
 #include "TraitInterfaces/IUpdate.h"
 #include "EvaluationVM/EvaluationVM.h"
@@ -26,7 +24,7 @@ TAutoConsoleVariable<int32> CVarAnimNextForceAnimBP(TEXT("a.AnimNextForceAnimBP"
 
 FAnimNode_AnimNextGraph::FAnimNode_AnimNextGraph()
 	: FAnimNode_CustomProperty()
-	, Module(nullptr)
+	, AnimationGraph(nullptr)
 	, LODThreshold(INDEX_NONE)
 {
 }
@@ -79,9 +77,9 @@ void FAnimNode_AnimNextGraph::Initialize_AnyThread(const FAnimationInitializeCon
 
 	SourceLink.Initialize(Context);
 
-	if (!GraphInstance.IsValid() && Module)
+	if (!GraphInstance.IsValid() && AnimationGraph)
 	{
-		Module->AllocateInstance(GraphInstance);
+		AnimationGraph->AllocateInstance(GraphInstance);
 	}
 
 	FAnimNode_CustomProperty::Initialize_AnyThread(Context);
@@ -163,7 +161,7 @@ void FAnimNode_AnimNextGraph::PostSerialize(const FArchive& Ar)
 	// since memory has changed
 	if (Ar.IsObjectReferenceCollector())
 	{
-		if (Module)
+		if (AnimationGraph)
 		{
 			//AnimNextGraph->Initialize();
 		}
diff --git a/Engine/Plugins/Experimental/AnimNextAnimGraph/Source/AnimNextAnimGraph/Private/Module.cpp b/Engine/Plugins/Experimental/AnimNextAnimGraph/Source/AnimNextAnimGraph/Private/Module.cpp
--- a/Engine/Plugins/Experimental/AnimNextAnimGraph/Source/AnimNextAnimGraph/Private/Module.cpp
+++ b/Engine/Plugins/Experimental/AnimNextAnimGraph/Source/AnimNextAnimGraph/Private/Module.cpp
@@ -43,12 +43,12 @@ public:
 		IAnimNextModuleInterface::Get().UnregisterAnimNextAnimGraph();
 	}
 
-	virtual void UpdateGraph(FAnimNextGraphInstancePtr& GraphInstance, float DeltaTime, UE::AnimNext::FTraitEventList& InputEventList, FTraitEventList& OutputEventList) const override
+	virtual void UpdateGraph(const FAnimNextGraphInstancePtr& GraphInstance, float DeltaTime, UE::AnimNext::FTraitEventList& InputEventList, FTraitEventList& OutputEventList) const override
 	{
 		UE::AnimNext::UpdateGraph(GraphInstance, DeltaTime, InputEventList, OutputEventList);
 	}
 
-	virtual void EvaluateGraph(FAnimNextGraphInstancePtr& GraphInstance, const UE::AnimNext::FReferencePose& RefPose, int32 GraphLODLevel, FLODPoseHeap& OutputPose) const override
+	virtual void EvaluateGraph(const FAnimNextGraphInstancePtr& GraphInstance, const UE::AnimNext::FReferencePose& RefPose, int32 GraphLODLevel, FLODPoseHeap& OutputPose) const override
 	{
 		const FEvaluationProgram EvaluationProgram = UE::AnimNext::EvaluateGraph(GraphInstance);
 
diff --git a/Engine/Plugins/Experimental/AnimNextAnimGraph/Source/AnimNextAnimGraph/Private/PlayAnim/PlayAnimEvents.cpp b/Engine/Plugins/Experimental/AnimNextAnimGraph/Source/AnimNextAnimGraph/Private/PlayAnim/PlayAnimEvents.cpp
--- a/Engine/Plugins/Experimental/AnimNextAnimGraph/Source/AnimNextAnimGraph/Private/PlayAnim/PlayAnimEvents.cpp
+++ b/Engine/Plugins/Experimental/AnimNextAnimGraph/Source/AnimNextAnimGraph/Private/PlayAnim/PlayAnimEvents.cpp
@@ -8,7 +8,7 @@ namespace UE::AnimNext
 {
 	void FPlayAnim_PlayEvent::OnExpired(FTraitEventList& OutputEventList)
 	{
-		auto ActionEvent = MakeTraitEvent<FAnimNextSchedule_ActionEvent>();
+		auto ActionEvent = MakeTraitEvent<FAnimNextModule_ActionEvent>();
 		ActionEvent->ActionFunction = [Request = Request]()
 			{
 				Request->OnStatusUpdate(EPlayAnimStatus::Expired);
diff --git a/Engine/Plugins/Experimental/AnimNextAnimGraph/Source/AnimNextAnimGraph/Private/PlayAnim/PlayAnimEvents.h b/Engine/Plugins/Experimental/AnimNextAnimGraph/Source/AnimNextAnimGraph/Private/PlayAnim/PlayAnimEvents.h
--- a/Engine/Plugins/Experimental/AnimNextAnimGraph/Source/AnimNextAnimGraph/Private/PlayAnim/PlayAnimEvents.h
+++ b/Engine/Plugins/Experimental/AnimNextAnimGraph/Source/AnimNextAnimGraph/Private/PlayAnim/PlayAnimEvents.h
@@ -6,7 +6,7 @@
 
 #include "PlayAnim/PlayAnimRequest.h"
 #include "PlayAnim/PlayAnimStatus.h"
-#include "Scheduler/ScheduleEvents.h"
+#include "Module/ModuleEvents.h"
 #include "TraitCore/TraitEvent.h"
 #include "TraitInterfaces/ITimeline.h"
 
@@ -49,9 +49,9 @@ namespace UE::AnimNext
 	 *
 	 * Event raised when the status of a request changes.
 	 */
-	struct FPlayAnim_StatusUpdateEvent : public FAnimNextSchedule_ActionEvent
+	struct FPlayAnim_StatusUpdateEvent : public FAnimNextModule_ActionEvent
 	{
-		DECLARE_ANIM_TRAIT_EVENT(FPlayAnim_StatusUpdateEvent, FAnimNextSchedule_ActionEvent)
+		DECLARE_ANIM_TRAIT_EVENT(FPlayAnim_StatusUpdateEvent, FAnimNextModule_ActionEvent)
 
 		// FAnimNextSchedule_ActionEvent impl
 		virtual void Execute() const override;
@@ -68,9 +68,9 @@ namespace UE::AnimNext
 	 *
 	 * Event raised when a request is playing with its updated timeline progress.
 	 */
-	struct FPlayAnim_TimelineUpdateEvent : public FAnimNextSchedule_ActionEvent
+	struct FPlayAnim_TimelineUpdateEvent : public FAnimNextModule_ActionEvent
 	{
-		DECLARE_ANIM_TRAIT_EVENT(FPlayAnim_TimelineUpdateEvent, FAnimNextSchedule_ActionEvent)
+		DECLARE_ANIM_TRAIT_EVENT(FPlayAnim_TimelineUpdateEvent, FAnimNextModule_ActionEvent)
 
 		// FAnimNextSchedule_ActionEvent impl
 		virtual void Execute() const override;
diff --git a/Engine/Plugins/Experimental/AnimNextAnimGraph/Source/AnimNextAnimGraph/Private/PlayAnim/PlayAnimSlotTrait.cpp b/Engine/Plugins/Experimental/AnimNextAnimGraph/Source/AnimNextAnimGraph/Private/PlayAnim/PlayAnimSlotTrait.cpp
--- a/Engine/Plugins/Experimental/AnimNextAnimGraph/Source/AnimNextAnimGraph/Private/PlayAnim/PlayAnimSlotTrait.cpp
+++ b/Engine/Plugins/Experimental/AnimNextAnimGraph/Source/AnimNextAnimGraph/Private/PlayAnim/PlayAnimSlotTrait.cpp
@@ -7,7 +7,7 @@
 #include "TraitCore/NodeInstance.h"
 #include "TraitInterfaces/ITimeline.h"
 #include "Graph/AnimNextGraphInstance.h"
-#include "Scheduler/ScheduleEvents.h"
+#include "Module/ModuleEvents.h"
 
 namespace UE::AnimNext
 {
@@ -31,14 +31,14 @@ namespace UE::AnimNext
 	#undef TRAIT_INTERFACE_ENUMERATOR
 	#undef TRAIT_EVENT_ENUMERATOR
 
-	void FPlayAnimSlotTrait::FPlayAnimSlotRequest::Initialize(FPlayAnimRequestPtr InRequest, const FPlayAnimBlendSettings& InBlendSettings, const UAnimNextModule* InModule)
+	void FPlayAnimSlotTrait::FPlayAnimSlotRequest::Initialize(FPlayAnimRequestPtr InRequest, const FPlayAnimBlendSettings& InBlendSettings, const UAnimNextAnimationGraph* InAnimationGraph)
 	{
 		Request = InRequest;
 		BlendSettings = InBlendSettings;
-		Module = InModule;
+		AnimationGraph = InAnimationGraph;
 
 		// If no input is provided, we'll use the source
-		State = InModule != nullptr ? EPlayAnimRequestState::Active : EPlayAnimRequestState::ActiveSource;
+		State = InAnimationGraph != nullptr ? EPlayAnimRequestState::Active : EPlayAnimRequestState::ActiveSource;
 		bWasRelevant = false;
 	}
 
@@ -155,7 +155,7 @@ namespace UE::AnimNext
 			InstanceData->PendingRequest.Reset();
 
 			FPlayAnimBlendSettings BlendSettings;
-			const UAnimNextModule* Module = nullptr;
+			const UAnimNextAnimationGraph* AnimationGraph = nullptr;
 			if (Request)
 			{
 				// This is a new pending request, lookup the sub-graph to use with our chooser and the desired animation object
@@ -170,21 +170,21 @@ namespace UE::AnimNext
 						FChooserEvaluationContext ChooserContext;
 						ChooserContext.AddStructParam(ChooserParameters);
 
-						UChooserTable::EvaluateChooser(ChooserContext, Chooser, FObjectChooserBase::FObjectChooserIteratorCallback::CreateLambda([&Module](UObject* InResult)
+						UChooserTable::EvaluateChooser(ChooserContext, Chooser, FObjectChooserBase::FObjectChooserIteratorCallback::CreateLambda([&AnimationGraph](UObject* InResult)
 							{
-								Module = Cast<UAnimNextModule>(InResult);
+								AnimationGraph = Cast<UAnimNextAnimationGraph>(InResult);
 								return FObjectChooserBase::EIteratorStatus::Stop;
 							}));
 					}
 
-					if (Module != nullptr)
+					if (AnimationGraph != nullptr)
 					{
 						// Check for re-entrancy and early-out if we are linking back to the current instance or one of its parents
-						const FName EntryPoint = Module->GetDefaultEntryPoint();
+						const FName EntryPoint = AnimationGraph->DefaultEntryPoint;
 						const FAnimNextGraphInstance* OwnerGraphInstance = &Binding.GetTraitPtr().GetNodeInstance()->GetOwner();
 						while (OwnerGraphInstance != nullptr)
 						{
-							if (OwnerGraphInstance->UsesModule(Module) && OwnerGraphInstance->UsesEntryPoint(EntryPoint))
+							if (OwnerGraphInstance->UsesAnimationGraph(AnimationGraph) && OwnerGraphInstance->UsesEntryPoint(EntryPoint))
 							{
 								return;
 							}
@@ -215,7 +215,7 @@ namespace UE::AnimNext
 			const int32 FreeRequestIndex = FindFreeRequestIndexOrAdd(*InstanceData);
 
 			FPlayAnimSlotRequest& SlotRequest = InstanceData->SlotRequests[FreeRequestIndex];
-			SlotRequest.Initialize(Request, BlendSettings, Module);
+			SlotRequest.Initialize(Request, BlendSettings, AnimationGraph);
 
 			const int32 OldChildIndex = InstanceData->CurrentlyActiveRequestIndex;
 			const int32 NewChildIndex = FreeRequestIndex;
@@ -297,31 +297,6 @@ namespace UE::AnimNext
 				}
 			}
 		}
-
-		if (InstanceData->CurrentlyActiveRequestIndex != INDEX_NONE)
-		{
-			const FPlayAnimSlotRequest& ActiveSlotRequest = InstanceData->SlotRequests[InstanceData->CurrentlyActiveRequestIndex];
-			if (ActiveSlotRequest.State == EPlayAnimRequestState::Active)
-			{
-				InstanceData->SubGraphParameters->Update(TraitState.GetDeltaTime());
-
-				FParamStack& ParamStack = FParamStack::Get();
-				InstanceData->SubGraphParametersLayerHandle = ParamStack.PushLayer(InstanceData->SubGraphParameters->GetLayerHandle());
-			}
-		}
-	}
-
-	void FPlayAnimSlotTrait::PostUpdate(FUpdateTraversalContext& Context, const TTraitBinding<IUpdate>& Binding, const FTraitUpdateState& TraitState) const
-	{
-		FInstanceData* InstanceData = Binding.GetInstanceData<FInstanceData>();
-
-		if (InstanceData->SubGraphParametersLayerHandle.IsValid())
-		{
-			FParamStack& ParamStack = FParamStack::Get();
-			ParamStack.PopLayer(InstanceData->SubGraphParametersLayerHandle);
-
-			InstanceData->SubGraphParametersLayerHandle.Invalidate();
-		}
 	}
 
 	void FPlayAnimSlotTrait::QueueChildrenForTraversal(FUpdateTraversalContext& Context, const TTraitBinding<IUpdateTraversal>& Binding, const FTraitUpdateState& TraitState, FUpdateTraversalQueue& TraversalQueue) const
@@ -397,14 +372,14 @@ namespace UE::AnimNext
 
 			if (SlotRequest.State == EPlayAnimRequestState::Active)
 			{
-				const FName EntryPoint = SlotRequest.Module->GetDefaultEntryPoint();
-				SlotRequest.Module->AllocateInstance(Binding.GetTraitPtr().GetNodeInstance()->GetOwner(), SlotRequest.GraphInstance, EntryPoint);
+				const FName EntryPoint = SlotRequest.AnimationGraph->DefaultEntryPoint;
+				SlotRequest.AnimationGraph->AllocateInstance(Binding.GetTraitPtr().GetNodeInstance()->GetOwner(), SlotRequest.GraphInstance, EntryPoint);
 				SlotRequest.ChildPtr = SlotRequest.GraphInstance.GetGraphRootPtr();
 
 				// Setup our graph parameters
 				const FAnimNextPlayAnimRequestArgs& RequestArgs = SlotRequest.Request->GetArgs();
 
-				const FString SubGraphPath = SlotRequest.Module->GetPathName();
+				const FString SubGraphPath = SlotRequest.AnimationGraph->GetPathName();
 
 				const FName AnimationObjectName(FString::Printf(TEXT("%s:AnimationObject"), *SubGraphPath));
 				const FName StartPositionName(FString::Printf(TEXT("%s:StartPosition"), *SubGraphPath));
@@ -421,7 +396,7 @@ namespace UE::AnimNext
 				SubGraphParameters.AddProperty(IsLoopingName, EPropertyBagPropertyType::Bool);
 				SubGraphParameters.SetValueBool(IsLoopingName, false);
 
-				InstanceData->SubGraphParameters = MakeUnique<FPropertyBagProxy>(NAME_None, MoveTemp(SubGraphParameters));
+				// TODO: forwarding variables to playanim 
 
 				// TODO: Validate that our child implements the ITimeline interface
 
@@ -559,7 +534,7 @@ namespace UE::AnimNext
 				SlotRequest.Request->AddReferencedObjects(Collector);
 			}
 
-			Collector.AddReferencedObject(SlotRequest.Module);
+			Collector.AddReferencedObject(SlotRequest.AnimationGraph);
 		}
 	}
 }
diff --git a/Engine/Plugins/Experimental/AnimNextAnimGraph/Source/AnimNextAnimGraph/Private/PlayAnim/PlayAnimSlotTrait.h b/Engine/Plugins/Experimental/AnimNextAnimGraph/Source/AnimNextAnimGraph/Private/PlayAnim/PlayAnimSlotTrait.h
--- a/Engine/Plugins/Experimental/AnimNextAnimGraph/Source/AnimNextAnimGraph/Private/PlayAnim/PlayAnimSlotTrait.h
+++ b/Engine/Plugins/Experimental/AnimNextAnimGraph/Source/AnimNextAnimGraph/Private/PlayAnim/PlayAnimSlotTrait.h
@@ -11,10 +11,9 @@
 #include "TraitInterfaces/IInertializerBlend.h"
 #include "TraitInterfaces/ISmoothBlend.h"
 #include "TraitInterfaces/IUpdate.h"
-#include "Module/AnimNextModule.h"
+#include "Graph/AnimNextAnimationGraph.h"
 #include "Graph/AnimNextGraphInstancePtr.h"
 #include "PlayAnim/PlayAnimEvents.h"
-#include "Param/PropertyBagProxy.h"
 #include "Chooser.h"
 
 #include "PlayAnimSlotTrait.generated.h"
@@ -109,7 +108,7 @@ namespace UE::AnimNext
 			FPlayAnimBlendSettings BlendSettings;
 
 			// The module used by the graph instance, as selected by the chooser
-			TObjectPtr<const UAnimNextModule> Module;
+			TObjectPtr<const UAnimNextAnimationGraph> AnimationGraph;
 
 			// The graph instance
 			FAnimNextGraphInstancePtr GraphInstance;
@@ -125,7 +124,7 @@ namespace UE::AnimNext
 			bool bWasRelevant = false;
 
 			// Initializes a request to begin playing
-			void Initialize(FPlayAnimRequestPtr InRequest, const FPlayAnimBlendSettings& InBlendSettings, const UAnimNextModule* InModule);
+			void Initialize(FPlayAnimRequestPtr InRequest, const FPlayAnimBlendSettings& InBlendSettings, const UAnimNextAnimationGraph* InAnimationGraph);
 		};
 
 		using FSharedData = FAnimNextPlayAnimSlotTraitSharedData;
@@ -142,12 +141,6 @@ namespace UE::AnimNext
 			// All other requests are blending out
 			int32 CurrentlyActiveRequestIndex = INDEX_NONE;
 
-			// The parameters that control our sub-graphs
-			TUniquePtr<FPropertyBagProxy> SubGraphParameters;
-
-			// The param stack layer handle for our sub-graph parameters
-			FParamStack::FPushedLayerHandle SubGraphParametersLayerHandle;
-
 			void Construct(const FExecutionContext& Context, const FTraitBinding& Binding);
 			void Destruct(const FExecutionContext& Context, const FTraitBinding& Binding);
 		};
@@ -161,7 +154,6 @@ namespace UE::AnimNext
 
 		// IUpdate impl
 		virtual void PreUpdate(FUpdateTraversalContext& Context, const TTraitBinding<IUpdate>& Binding, const FTraitUpdateState& TraitState) const override;
-		virtual void PostUpdate(FUpdateTraversalContext& Context, const TTraitBinding<IUpdate>& Binding, const FTraitUpdateState& TraitState) const override;
 
 		// IUpdateTraversal impl
 		virtual void QueueChildrenForTraversal(FUpdateTraversalContext& Context, const TTraitBinding<IUpdateTraversal>& Binding, const FTraitUpdateState& TraitState, FUpdateTraversalQueue& TraversalQueue) const override;
diff --git a/Engine/Plugins/Experimental/AnimNextAnimGraph/Source/AnimNextAnimGraph/Private/TraitInterfaces/IEvaluate.cpp b/Engine/Plugins/Experimental/AnimNextAnimGraph/Source/AnimNextAnimGraph/Private/TraitInterfaces/IEvaluate.cpp
--- a/Engine/Plugins/Experimental/AnimNextAnimGraph/Source/AnimNextAnimGraph/Private/TraitInterfaces/IEvaluate.cpp
+++ b/Engine/Plugins/Experimental/AnimNextAnimGraph/Source/AnimNextAnimGraph/Private/TraitInterfaces/IEvaluate.cpp
@@ -88,7 +88,7 @@ namespace UE::AnimNext
 		}
 	};
 
-	FEvaluationProgram EvaluateGraph(FAnimNextGraphInstancePtr& GraphInstance)
+	FEvaluationProgram EvaluateGraph(const FAnimNextGraphInstancePtr& GraphInstance)
 	{
 		return EvaluateGraph(GraphInstance.GetGraphRootPtr());
 	}
diff --git a/Engine/Plugins/Experimental/AnimNextAnimGraph/Source/AnimNextAnimGraph/Private/TraitInterfaces/IUpdate.cpp b/Engine/Plugins/Experimental/AnimNextAnimGraph/Source/AnimNextAnimGraph/Private/TraitInterfaces/IUpdate.cpp
--- a/Engine/Plugins/Experimental/AnimNextAnimGraph/Source/AnimNextAnimGraph/Private/TraitInterfaces/IUpdate.cpp
+++ b/Engine/Plugins/Experimental/AnimNextAnimGraph/Source/AnimNextAnimGraph/Private/TraitInterfaces/IUpdate.cpp
@@ -5,7 +5,6 @@
 #include "TraitInterfaces/IHierarchy.h"
 #include "Graph/GraphInstanceComponent.h"
 #include "AnimNextAnimGraphStats.h"
-#include "Param/ParamStack.h"
 
 DEFINE_STAT(STAT_AnimNext_UpdateGraph);
 
@@ -473,7 +472,7 @@ namespace UE::AnimNext
 	// perform as much useful work as possible while waiting for memory, hiding its slow latency by fully
 	// leveraging out-of-order CPU execution.
 
-	void UpdateGraph(FAnimNextGraphInstancePtr& GraphInstance, float DeltaTime, FTraitEventList& InputEventList, FTraitEventList& OutputEventList)
+	void UpdateGraph(const FAnimNextGraphInstancePtr& GraphInstance, float DeltaTime, FTraitEventList& InputEventList, FTraitEventList& OutputEventList)
 	{
 		SCOPE_CYCLE_COUNTER(STAT_AnimNext_UpdateGraph);
 		
@@ -519,10 +518,7 @@ namespace UE::AnimNext
 
 			It.Value()->PreUpdate(TraversalContext);
 		}
-
-		// Push any parameters this graph owns
-		FParamStack::FPushedLayerHandle PushedLayerHandle = GraphInstance.UpdateAndPushGraphState(DeltaTime);
-
+		
 		// Add the graph root to start the update process
 		Private::FUpdateEntry RootEntry(GraphInstance.GetGraphRootPtr(), RootState);
 		TraversalContext.PushUpdateEntry(&RootEntry);
@@ -638,9 +634,6 @@ namespace UE::AnimNext
 			}
 		}
 
-		// Pop any parameters we pushed above
-		GraphInstance.PopGraphState(PushedLayerHandle);
-
 		// Clear our executing entry
 		TraversalContext.ExecutingEntry = nullptr;
 
diff --git a/Engine/Plugins/Experimental/AnimNextAnimGraph/Source/AnimNextAnimGraph/Private/Traits/SubGraphHost.cpp b/Engine/Plugins/Experimental/AnimNextAnimGraph/Source/AnimNextAnimGraph/Private/Traits/SubGraphHost.cpp
--- a/Engine/Plugins/Experimental/AnimNextAnimGraph/Source/AnimNextAnimGraph/Private/Traits/SubGraphHost.cpp
+++ b/Engine/Plugins/Experimental/AnimNextAnimGraph/Source/AnimNextAnimGraph/Private/Traits/SubGraphHost.cpp
@@ -69,25 +69,25 @@ namespace UE::AnimNext
 
 		const bool bHasActiveSubGraph = InstanceData->CurrentlyActiveSubGraphIndex != INDEX_NONE;
 
-		TObjectPtr<const UAnimNextModule> CurrentActiveModule;
+		TObjectPtr<const UAnimNextAnimationGraph> CurrentActiveAnimationGraph;
 		FName CurrentActiveEntryPoint = NAME_None;
 		if (bHasActiveSubGraph)
 		{
 			FSubGraphSlot& SubGraphSlot = InstanceData->SubGraphSlots[InstanceData->CurrentlyActiveSubGraphIndex];
-			CurrentActiveModule = SubGraphSlot.Module;
+			CurrentActiveAnimationGraph = SubGraphSlot.AnimationGraph;
 			CurrentActiveEntryPoint = SubGraphSlot.EntryPoint;
 
 			SubGraphSlot.bWasRelevant = true;
 		}
 
-		const TObjectPtr<const UAnimNextModule> DesiredModule = SharedData->GetModule(Binding);
+		const TObjectPtr<const UAnimNextAnimationGraph> DesiredAnimationGraph = SharedData->GetAnimationGraph(Binding);
 		const FName EntryPoint = SharedData->GetEntryPoint(Binding);
 
 		// Check for re-entrancy and early-out if we are linking back to the current instance or one of its parents
 		const FAnimNextGraphInstance* OwnerGraphInstance = &Binding.GetTraitPtr().GetNodeInstance()->GetOwner();
 		while (OwnerGraphInstance != nullptr)
 		{
-			if (OwnerGraphInstance->UsesModule(DesiredModule) && OwnerGraphInstance->UsesEntryPoint(EntryPoint))
+			if (OwnerGraphInstance->UsesAnimationGraph(DesiredAnimationGraph) && OwnerGraphInstance->UsesEntryPoint(EntryPoint))
 			{
 				return;
 			}
@@ -95,7 +95,7 @@ namespace UE::AnimNext
 			OwnerGraphInstance = OwnerGraphInstance->GetParentGraphInstance();
 		}
 
-		if (!bHasActiveSubGraph || CurrentActiveModule != DesiredModule || CurrentActiveEntryPoint != EntryPoint)
+		if (!bHasActiveSubGraph || CurrentActiveAnimationGraph != DesiredAnimationGraph || CurrentActiveEntryPoint != EntryPoint)
 		{
 			// Find an empty slot we can use
 			int32 FreeSlotIndex = INDEX_NONE;
@@ -118,8 +118,8 @@ namespace UE::AnimNext
 			}
 
 			FSubGraphSlot& SubGraphSlot = InstanceData->SubGraphSlots[FreeSlotIndex];
-			SubGraphSlot.Module = DesiredModule;
-			SubGraphSlot.State = DesiredModule ? ESlotState::ActiveWithGraph : ESlotState::ActiveWithReferencePose;
+			SubGraphSlot.AnimationGraph = DesiredAnimationGraph;
+			SubGraphSlot.State = DesiredAnimationGraph ? ESlotState::ActiveWithGraph : ESlotState::ActiveWithReferencePose;
 			SubGraphSlot.EntryPoint = EntryPoint;
 
 			const int32 OldChildIndex = InstanceData->CurrentlyActiveSubGraphIndex;
@@ -210,7 +210,7 @@ namespace UE::AnimNext
 
 			if (SubGraphEntry.State == ESlotState::ActiveWithGraph)
 			{
-				SubGraphEntry.Module->AllocateInstance(Binding.GetTraitPtr().GetNodeInstance()->GetOwner(), SubGraphEntry.GraphInstance, SubGraphEntry.EntryPoint);
+				SubGraphEntry.AnimationGraph->AllocateInstance(Binding.GetTraitPtr().GetNodeInstance()->GetOwner(), SubGraphEntry.GraphInstance, SubGraphEntry.EntryPoint);
 			}
 		}
 	}
diff --git a/Engine/Plugins/Experimental/AnimNextAnimGraph/Source/AnimNextAnimGraph/Private/Traits/SubGraphHost.h b/Engine/Plugins/Experimental/AnimNextAnimGraph/Source/AnimNextAnimGraph/Private/Traits/SubGraphHost.h
--- a/Engine/Plugins/Experimental/AnimNextAnimGraph/Source/AnimNextAnimGraph/Private/Traits/SubGraphHost.h
+++ b/Engine/Plugins/Experimental/AnimNextAnimGraph/Source/AnimNextAnimGraph/Private/Traits/SubGraphHost.h
@@ -9,7 +9,7 @@
 #include "TraitInterfaces/IGarbageCollection.h"
 #include "TraitInterfaces/IHierarchy.h"
 #include "TraitInterfaces/IUpdate.h"
-#include "Module/AnimNextModule.h"
+#include "Graph/AnimNextAnimationGraph.h"
 #include "Graph/AnimNextGraphInstancePtr.h"
 
 #include "SubGraphHost.generated.h"
@@ -19,9 +19,9 @@ struct FAnimNextSubGraphHostTraitSharedData : public FAnimNextTraitSharedData
 {
 	GENERATED_BODY()
 
-	/** Module of the sub-graph that we host and manage. */
+	/** Asset to use as a sub-graph */
 	UPROPERTY(EditAnywhere, Category = "Default")
-	TObjectPtr<const UAnimNextModule> Module;
+	TObjectPtr<const UAnimNextAnimationGraph> AnimationGraph;
 
 	/** A dummy child that we can use to output the bind pose. This property is hidden and automatically populated during compilation. */
 	UPROPERTY(meta = (Hidden))
@@ -33,7 +33,7 @@ struct FAnimNextSubGraphHostTraitSharedData : public FAnimNextTraitSharedData
 
 	// Latent pin support boilerplate
 	#define TRAIT_LATENT_PROPERTIES_ENUMERATOR(GeneratorMacro) \
-		GeneratorMacro(Module) \
+		GeneratorMacro(AnimationGraph) \
 		GeneratorMacro(EntryPoint) \
 
 	GENERATE_TRAIT_LATENT_PROPERTIES(FAnimNextSubGraphHostTraitSharedData, TRAIT_LATENT_PROPERTIES_ENUMERATOR)
@@ -60,8 +60,8 @@ namespace UE::AnimNext
 
 		struct FSubGraphSlot
 		{
-			// The module to use
-			TObjectPtr<const UAnimNextModule> Module;
+			// The animation graph to use
+			TObjectPtr<const UAnimNextAnimationGraph> AnimationGraph;
 
 			// The graph instance
 			FAnimNextGraphInstancePtr GraphInstance;
diff --git a/Engine/Plugins/Experimental/AnimNextAnimGraph/Source/AnimNextAnimGraph/Public/TraitInterfaces/IEvaluate.h b/Engine/Plugins/Experimental/AnimNextAnimGraph/Source/AnimNextAnimGraph/Public/TraitInterfaces/IEvaluate.h
--- a/Engine/Plugins/Experimental/AnimNextAnimGraph/Source/AnimNextAnimGraph/Public/TraitInterfaces/IEvaluate.h
+++ b/Engine/Plugins/Experimental/AnimNextAnimGraph/Source/AnimNextAnimGraph/Public/TraitInterfaces/IEvaluate.h
@@ -97,7 +97,7 @@ namespace UE::AnimNext
 	 *
 	 * @see IEvaluate::PreEvaluate, IEvaluate::PostEvaluate, IHierarchy::GetChildren
 	 */
-	[[nodiscard]] ANIMNEXTANIMGRAPH_API FEvaluationProgram EvaluateGraph(FAnimNextGraphInstancePtr& GraphInstance);
+	[[nodiscard]] ANIMNEXTANIMGRAPH_API FEvaluationProgram EvaluateGraph(const FAnimNextGraphInstancePtr& GraphInstance);
 
 	/**
 	 * Evaluates a sub-graph starting at its root and produces an evaluation program.
diff --git a/Engine/Plugins/Experimental/AnimNextAnimGraph/Source/AnimNextAnimGraph/Public/TraitInterfaces/IUpdate.h b/Engine/Plugins/Experimental/AnimNextAnimGraph/Source/AnimNextAnimGraph/Public/TraitInterfaces/IUpdate.h
--- a/Engine/Plugins/Experimental/AnimNextAnimGraph/Source/AnimNextAnimGraph/Public/TraitInterfaces/IUpdate.h
+++ b/Engine/Plugins/Experimental/AnimNextAnimGraph/Source/AnimNextAnimGraph/Public/TraitInterfaces/IUpdate.h
@@ -185,7 +185,7 @@ namespace UE::AnimNext
 		// The root node doesn't have a parent but we need a bookkeeping list regardless
 		Private::FUpdateEventBookkeepingList* RootParentBookkeepingEntryList = nullptr;
 
-		friend ANIMNEXTANIMGRAPH_API void UpdateGraph(FAnimNextGraphInstancePtr& GraphInstance, float DeltaTime, FTraitEventList& InputEventList, FTraitEventList& OutputEventList);
+		friend ANIMNEXTANIMGRAPH_API void UpdateGraph(const FAnimNextGraphInstancePtr& GraphInstance, float DeltaTime, FTraitEventList& InputEventList, FTraitEventList& OutputEventList);
 		friend FUpdateTraversalQueue;
 	};
 
@@ -213,7 +213,7 @@ namespace UE::AnimNext
 		// emptied and pushed onto the update stack
 		Private::FUpdateEntry* QueuedUpdateStackHead = nullptr;
 
-		friend ANIMNEXTANIMGRAPH_API void UpdateGraph(FAnimNextGraphInstancePtr& GraphInstance, float DeltaTime, FTraitEventList& InputEventList, FTraitEventList& OutputEventList);
+		friend ANIMNEXTANIMGRAPH_API void UpdateGraph(const FAnimNextGraphInstancePtr& GraphInstance, float DeltaTime, FTraitEventList& InputEventList, FTraitEventList& OutputEventList);
 		friend FUpdateTraversalContext;
 	};
 
@@ -327,5 +327,5 @@ namespace UE::AnimNext
 	 *
 	 * @see IUpdate::PreUpdate, IUpdate::PostUpdate, IHierarchy::GetChildren
 	 */
-	ANIMNEXTANIMGRAPH_API void UpdateGraph(FAnimNextGraphInstancePtr& GraphInstance, float DeltaTime, FTraitEventList& InputEventList, FTraitEventList& OutputEventList);
+	ANIMNEXTANIMGRAPH_API void UpdateGraph(const FAnimNextGraphInstancePtr& GraphInstance, float DeltaTime, FTraitEventList& InputEventList, FTraitEventList& OutputEventList);
 }
diff --git a/Engine/Plugins/Experimental/AnimNextAnimGraph/Source/AnimNextAnimGraphTestSuite/Private/AnimNextAnimGraphTraitEventTest.cpp b/Engine/Plugins/Experimental/AnimNextAnimGraph/Source/AnimNextAnimGraphTestSuite/Private/AnimNextAnimGraphTraitEventTest.cpp
--- a/Engine/Plugins/Experimental/AnimNextAnimGraph/Source/AnimNextAnimGraphTestSuite/Private/AnimNextAnimGraphTraitEventTest.cpp
+++ b/Engine/Plugins/Experimental/AnimNextAnimGraph/Source/AnimNextAnimGraphTestSuite/Private/AnimNextAnimGraphTraitEventTest.cpp
@@ -14,8 +14,8 @@
 #include "TraitCore/NodeTemplateRegistry.h"
 #include "TraitInterfaces/IHierarchy.h"
 #include "TraitInterfaces/IUpdate.h"
-#include "Module/AnimNextModule.h"
-#include "Module/ModuleFactory.h"
+#include "Graph/AnimNextAnimationGraph.h"
+#include "Graph/AnimNextAnimationGraphFactory.h"
 #include "Misc/AutomationTest.h"
 
 #if WITH_DEV_AUTOMATION_TESTS
@@ -194,9 +194,9 @@ bool FAnimationAnimNextRuntimeTest_GraphTraitEvent::RunTest(const FString& InPar
 		AUTO_REGISTER_ANIM_TRAIT(FTraitGraphTest_EventAB_OneChild)
 		AUTO_REGISTER_ANIM_TRAIT(FTraitGraphTest_EventAB_TwoChildren)
 
-		UFactory* GraphFactory = NewObject<UAnimNextModuleFactory>();
-		UAnimNextModule* Module = CastChecked<UAnimNextModule>(GraphFactory->FactoryCreateNew(UAnimNextModule::StaticClass(), GetTransientPackage(), TEXT("TestAnimNextGraph"), RF_Transient, nullptr, nullptr, NAME_None));
-		UE_RETURN_ON_ERROR(Module != nullptr, "FAnimationAnimNextRuntimeTest_GraphTraitEvent -> Failed to create module");
+		UFactory* GraphFactory = NewObject<UAnimNextAnimationGraphFactory>();
+		UAnimNextAnimationGraph* AnimationGraph = CastChecked<UAnimNextAnimationGraph>(GraphFactory->FactoryCreateNew(UAnimNextAnimationGraph::StaticClass(), GetTransientPackage(), TEXT("TestAnimNextGraph"), RF_Transient, nullptr, nullptr, NAME_None));
+		UE_RETURN_ON_ERROR(AnimationGraph != nullptr, "FAnimationAnimNextRuntimeTest_GraphTraitEvent -> Failed to create animation graph");
 
 		FScopedClearNodeTemplateRegistry ScopedClearNodeTemplateRegistry;
 		FNodeTemplateRegistry& Registry = FNodeTemplateRegistry::Get();
@@ -296,10 +296,10 @@ bool FAnimationAnimNextRuntimeTest_GraphTraitEvent::RunTest(const FString& InPar
 		}
 
 		// Read our graph
-		FTestUtils::LoadFromArchiveBuffer(*Module, NodeHandles, GraphSharedDataArchiveBuffer);
+		FTestUtils::LoadFromArchiveBuffer(*AnimationGraph, NodeHandles, GraphSharedDataArchiveBuffer);
 
 		FAnimNextGraphInstancePtr GraphInstance;
-		Module->AllocateInstance(GraphInstance);
+		AnimationGraph->AllocateInstance(GraphInstance);
 
 		{
 			// Raise EventA and EventB on graph, every node sees them
diff --git a/Engine/Plugins/Experimental/AnimNextAnimGraph/Source/AnimNextAnimGraphTestSuite/Private/AnimNextAnimGraphTraitGraphTest.cpp b/Engine/Plugins/Experimental/AnimNextAnimGraph/Source/AnimNextAnimGraphTestSuite/Private/AnimNextAnimGraphTraitGraphTest.cpp
--- a/Engine/Plugins/Experimental/AnimNextAnimGraph/Source/AnimNextAnimGraphTestSuite/Private/AnimNextAnimGraphTraitGraphTest.cpp
+++ b/Engine/Plugins/Experimental/AnimNextAnimGraph/Source/AnimNextAnimGraphTestSuite/Private/AnimNextAnimGraphTraitGraphTest.cpp
@@ -5,22 +5,19 @@
 #include "AnimNextRuntimeTest.h"
 #include "AnimNextTest.h"
 #include "AssetToolsModule.h"
-#include "Context.h"
 #include "UncookedOnlyUtils.h"
 #include "TraitCore/TraitRegistry.h"
 #include "TraitInterfaces/IEvaluate.h"
 #include "TraitInterfaces/IUpdate.h"
 #include "Editor/Transactor.h"
-#include "AnimNextExecuteContext.h"
-#include "Module/AnimNextModule.h"
-#include "Module/AnimNextModule_EditorData.h"
-#include "Module/ModuleFactory.h"
+#include "Graph/AnimNextAnimationGraph_EditorData.h"
+#include "Graph/AnimNextAnimationGraph.h"
+#include "Graph/AnimNextAnimationGraphFactory.h"
+#include "Graph/AnimNextGraphInstance.h"
 #include "Graph/RigDecorator_AnimNextCppTrait.h"
 #include "Graph/RigUnit_AnimNextGraphRoot.h"
 #include "Graph/RigUnit_AnimNextTraitStack.h"
 #include "Misc/AutomationTest.h"
-#include "Param/ParamStack.h"
-#include "Param/RigVMDispatch_GetParameter.h"
 #include "RigVMFunctions/Math/RigVMFunction_MathInt.h"
 #include "RigVMFunctions/Math/RigVMFunction_MathFloat.h"
 
@@ -38,23 +35,27 @@ namespace UE::AnimNext
 
 		using FSharedData = FTestTraitSharedData;
 
+		struct FInstanceData : FTrait::FInstanceData
+		{
+			int32 UpdateCount = 0;
+			int32 EvaluateCount = 0;
+		};
+		
 		// IUpdate impl
 		virtual void PostUpdate(FUpdateTraversalContext& Context, const TTraitBinding<IUpdate>& Binding, const FTraitUpdateState& TraitState) const override
 		{
 			IUpdate::PostUpdate(Context, Binding, TraitState);
 
 			const FSharedData* SharedData = Binding.GetSharedData<FSharedData>();
-			UE::AnimNext::FParamStack& ParamStack = UE::AnimNext::FParamStack::Get();
-
-			int32& UpdateCount = ParamStack.GetMutableParam<int32>("/Engine/Transient.TestAnimNextGraph:UpdateCount");
-			UpdateCount++;
-
-			ParamStack.GetMutableParam<int32>("/Engine/Transient.TestAnimNextGraph:SomeInt32") = SharedData->SomeInt32;
-			ParamStack.GetMutableParam<float>("/Engine/Transient.TestAnimNextGraph:SomeFloat") = SharedData->SomeFloat;
-
-			ParamStack.GetMutableParam<int32>("/Engine/Transient.TestAnimNextGraph:SomeLatentInt32") = SharedData->GetSomeLatentInt32(Binding);				// MathAdd with constants, latent
-			ParamStack.GetMutableParam<int32>("/Engine/Transient.TestAnimNextGraph:SomeOtherLatentInt32") = SharedData->GetSomeOtherLatentInt32(Binding);	// GetParameter, latent
-			ParamStack.GetMutableParam<float>("/Engine/Transient.TestAnimNextGraph:SomeLatentFloat") = SharedData->GetSomeLatentFloat(Binding);				// Inline value, not latent
+			FInstanceData* InstanceData = Binding.GetInstanceData<FInstanceData>();
+
+			FRigVMExecuteContext& ExecuteContext = Context.GetRootGraphInstance().GetExtendedExecuteContext().GetPublicData();
+			ExecuteContext.Logf(EMessageSeverity::Info, TEXT("UpdateCount == %d"), ++InstanceData->UpdateCount);
+			ExecuteContext.Logf(EMessageSeverity::Info, TEXT("SomeInt32 == %d"), SharedData->SomeInt32);
+			ExecuteContext.Logf(EMessageSeverity::Info, TEXT("SomeFloat == %.02f"), SharedData->SomeFloat);
+			ExecuteContext.Logf(EMessageSeverity::Info, TEXT("SomeLatentInt32 == %d"), SharedData->GetSomeLatentInt32(Binding));
+			ExecuteContext.Logf(EMessageSeverity::Info, TEXT("SomeOtherLatentInt32 == %d"), SharedData->GetSomeOtherLatentInt32(Binding));
+			ExecuteContext.Logf(EMessageSeverity::Info, TEXT("SomeLatentFloat == %.02f"), SharedData->GetSomeLatentFloat(Binding));
 		}
 
 		// IEvaluate impl
@@ -62,11 +63,9 @@ namespace UE::AnimNext
 		{
 			IEvaluate::PostEvaluate(Context, Binding);
 
-			const FSharedData* SharedData = Binding.GetSharedData<FSharedData>();
-			UE::AnimNext::FParamStack& ParamStack = UE::AnimNext::FParamStack::Get();
-			
-			int32& EvaluateCount = ParamStack.GetMutableParam<int32>("/Engine/Transient.TestAnimNextGraph:EvaluateCount");
-			EvaluateCount++;
+			FInstanceData* InstanceData = Binding.GetInstanceData<FInstanceData>();
+			FRigVMExecuteContext& ExecuteContext = Context.GetRootGraphInstance().GetExtendedExecuteContext().GetPublicData();
+			ExecuteContext.Logf(EMessageSeverity::Info, TEXT("EvaluateCount == %d"), ++InstanceData->EvaluateCount);
 		}
 	};
 
@@ -102,15 +101,13 @@ bool FAnimationAnimNextEditorTest_GraphAddTrait::RunTest(const FString& InParame
 
 		FScopedClearNodeTemplateRegistry ScopedClearNodeTemplateRegistry;
 
-		UFactory* GraphFactory = NewObject<UAnimNextModuleFactory>();
-		UAnimNextModule* Module = CastChecked<UAnimNextModule>(GraphFactory->FactoryCreateNew(UAnimNextModule::StaticClass(), GetTransientPackage(), TEXT("TestAnimNextGraph"), RF_Transient, nullptr, nullptr, NAME_None));
-		UE_RETURN_ON_ERROR(Module != nullptr, "FAnimationAnimNextEditorTest_GraphAddTrait -> Failed to create module");
+		UFactory* GraphFactory = NewObject<UAnimNextAnimationGraphFactory>();
+		UAnimNextAnimationGraph* AnimationGraph = CastChecked<UAnimNextAnimationGraph>(GraphFactory->FactoryCreateNew(UAnimNextAnimationGraph::StaticClass(), GetTransientPackage(), TEXT("TestAnimNextGraph"), RF_Transient, nullptr, nullptr, NAME_None));
+		UE_RETURN_ON_ERROR(AnimationGraph != nullptr, "FAnimationAnimNextEditorTest_GraphAddTrait -> Failed to create animation graph");
 
-		UAnimNextModule_EditorData* EditorData = UncookedOnly::FUtils::GetEditorData(Module);
+		UAnimNextAnimationGraph_EditorData* EditorData = UncookedOnly::FUtils::GetEditorData<UAnimNextAnimationGraph_EditorData>(AnimationGraph);
 		UE_RETURN_ON_ERROR(EditorData != nullptr, "FAnimationAnimNextEditorTest_GraphAddTrait -> Failed to find module editor data");
 
-		UE_RETURN_ON_ERROR(EditorData->AddAnimationGraph(FRigUnit_AnimNextGraphRoot::DefaultEntryPoint, false) != nullptr, "FAnimationAnimNextEditorTest_GraphAddTrait -> Failed to add animation graph");
-
 		URigVMController* Controller = EditorData->GetRigVMClient()->GetController(EditorData->GetRigVMClient()->GetDefaultModel());
 		UE_RETURN_ON_ERROR(Controller != nullptr, "FAnimationAnimNextEditorTest_GraphAddTrait -> Failed to get RigVM controller");
 
@@ -163,32 +160,42 @@ bool FAnimationAnimNextEditorTest_GraphAddTrait::RunTest(const FString& InParame
 		UE_RETURN_ON_ERROR(TraitPin->GetCPPTypeObject() == FRigDecorator_AnimNextCppDecorator::StaticStruct(), TEXT("FAnimationAnimNextEditorTest_GraphAddTrait -> Unexpected pin type"));
 
 		// Our first sub-pin is the hard coded script struct member that parametrizes the trait, dynamic trait sub-pins follow
-		UE_RETURN_ON_ERROR(TraitPin->GetSubPins().Num() == 6, TEXT("FAnimationAnimNextEditorTest_GraphAddTrait -> Unexpected trait sub pins"));
+		UE_RETURN_ON_ERROR(TraitPin->GetSubPins().Num() == 8, TEXT("FAnimationAnimNextEditorTest_GraphAddTrait -> Unexpected trait sub pins"));
 
-		// SomeInt32
+		// UpdateCount
 		UE_RETURN_ON_ERROR(TraitPin->GetSubPins()[1]->GetCPPType() == TEXT("int32"), TEXT("FAnimationAnimNextEditorTest_GraphAddTrait -> Unexpected trait pin type"));
-		UE_RETURN_ON_ERROR(TraitPin->GetSubPins()[1]->GetDefaultValue() == TEXT("3"), TEXT("FAnimationAnimNextEditorTest_GraphAddTrait -> Unexpected trait pin value"));
+		UE_RETURN_ON_ERROR(TraitPin->GetSubPins()[1]->GetDefaultValue() == TEXT("0"), TEXT("FAnimationAnimNextEditorTest_GraphAddTrait -> Unexpected trait pin value"));
 		UE_RETURN_ON_ERROR(!TraitPin->GetSubPins()[1]->IsLazy(), TEXT("FAnimationAnimNextEditorTest_GraphAddTrait -> Expected non-lazy trait pin"));
 
-		// SomeFloat
-		UE_RETURN_ON_ERROR(TraitPin->GetSubPins()[2]->GetCPPType() == TEXT("float"), TEXT("FAnimationAnimNextEditorTest_GraphAddTrait -> Unexpected trait pin type"));
-		UE_RETURN_ON_ERROR(TraitPin->GetSubPins()[2]->GetDefaultValue() == TEXT("34.000000"), TEXT("FAnimationAnimNextEditorTest_GraphAddTrait -> Unexpected trait pin value"));
+		// EvaluateCount
+		UE_RETURN_ON_ERROR(TraitPin->GetSubPins()[2]->GetCPPType() == TEXT("int32"), TEXT("FAnimationAnimNextEditorTest_GraphAddTrait -> Unexpected trait pin type"));
+		UE_RETURN_ON_ERROR(TraitPin->GetSubPins()[2]->GetDefaultValue() == TEXT("0"), TEXT("FAnimationAnimNextEditorTest_GraphAddTrait -> Unexpected trait pin value"));
 		UE_RETURN_ON_ERROR(!TraitPin->GetSubPins()[2]->IsLazy(), TEXT("FAnimationAnimNextEditorTest_GraphAddTrait -> Expected non-lazy trait pin"));
 
-		// SomeLatentInt32
+		// SomeInt32
 		UE_RETURN_ON_ERROR(TraitPin->GetSubPins()[3]->GetCPPType() == TEXT("int32"), TEXT("FAnimationAnimNextEditorTest_GraphAddTrait -> Unexpected trait pin type"));
 		UE_RETURN_ON_ERROR(TraitPin->GetSubPins()[3]->GetDefaultValue() == TEXT("3"), TEXT("FAnimationAnimNextEditorTest_GraphAddTrait -> Unexpected trait pin value"));
-		UE_RETURN_ON_ERROR(TraitPin->GetSubPins()[3]->IsLazy(), TEXT("FAnimationAnimNextEditorTest_GraphAddTrait -> Expected lazy trait pin"));
+		UE_RETURN_ON_ERROR(!TraitPin->GetSubPins()[3]->IsLazy(), TEXT("FAnimationAnimNextEditorTest_GraphAddTrait -> Expected non-lazy trait pin"));
+
+		// SomeFloat
+		UE_RETURN_ON_ERROR(TraitPin->GetSubPins()[4]->GetCPPType() == TEXT("float"), TEXT("FAnimationAnimNextEditorTest_GraphAddTrait -> Unexpected trait pin type"));
+		UE_RETURN_ON_ERROR(TraitPin->GetSubPins()[4]->GetDefaultValue() == TEXT("34.000000"), TEXT("FAnimationAnimNextEditorTest_GraphAddTrait -> Unexpected trait pin value"));
+		UE_RETURN_ON_ERROR(!TraitPin->GetSubPins()[4]->IsLazy(), TEXT("FAnimationAnimNextEditorTest_GraphAddTrait -> Expected non-lazy trait pin"));
+
+		// SomeLatentInt32
+		UE_RETURN_ON_ERROR(TraitPin->GetSubPins()[5]->GetCPPType() == TEXT("int32"), TEXT("FAnimationAnimNextEditorTest_GraphAddTrait -> Unexpected trait pin type"));
+		UE_RETURN_ON_ERROR(TraitPin->GetSubPins()[5]->GetDefaultValue() == TEXT("3"), TEXT("FAnimationAnimNextEditorTest_GraphAddTrait -> Unexpected trait pin value"));
+		UE_RETURN_ON_ERROR(TraitPin->GetSubPins()[5]->IsLazy(), TEXT("FAnimationAnimNextEditorTest_GraphAddTrait -> Expected lazy trait pin"));
 
 		// SomeOtherLatentInt32
-		UE_RETURN_ON_ERROR(TraitPin->GetSubPins()[4]->GetCPPType() == TEXT("int32"), TEXT("FAnimationAnimNextEditorTest_GraphAddTrait -> Unexpected trait pin type"));
-		UE_RETURN_ON_ERROR(TraitPin->GetSubPins()[4]->GetDefaultValue() == TEXT("3"), TEXT("FAnimationAnimNextEditorTest_GraphAddTrait -> Unexpected trait pin value"));
-		UE_RETURN_ON_ERROR(TraitPin->GetSubPins()[4]->IsLazy(), TEXT("FAnimationAnimNextEditorTest_GraphAddTrait -> Expected lazy trait pin"));
+		UE_RETURN_ON_ERROR(TraitPin->GetSubPins()[6]->GetCPPType() == TEXT("int32"), TEXT("FAnimationAnimNextEditorTest_GraphAddTrait -> Unexpected trait pin type"));
+		UE_RETURN_ON_ERROR(TraitPin->GetSubPins()[6]->GetDefaultValue() == TEXT("3"), TEXT("FAnimationAnimNextEditorTest_GraphAddTrait -> Unexpected trait pin value"));
+		UE_RETURN_ON_ERROR(TraitPin->GetSubPins()[6]->IsLazy(), TEXT("FAnimationAnimNextEditorTest_GraphAddTrait -> Expected lazy trait pin"));
 
 		// SomeLatentFloat
-		UE_RETURN_ON_ERROR(TraitPin->GetSubPins()[5]->GetCPPType() == TEXT("float"), TEXT("FAnimationAnimNextEditorTest_GraphAddTrait -> Unexpected trait pin type"));
-		UE_RETURN_ON_ERROR(TraitPin->GetSubPins()[5]->GetDefaultValue() == TEXT("34.000000"), TEXT("FAnimationAnimNextEditorTest_GraphAddTrait -> Unexpected trait pin value"));
-		UE_RETURN_ON_ERROR(TraitPin->GetSubPins()[5]->IsLazy(), TEXT("FAnimationAnimNextEditorTest_GraphAddTrait -> Expected lazy trait pin"));
+		UE_RETURN_ON_ERROR(TraitPin->GetSubPins()[7]->GetCPPType() == TEXT("float"), TEXT("FAnimationAnimNextEditorTest_GraphAddTrait -> Unexpected trait pin type"));
+		UE_RETURN_ON_ERROR(TraitPin->GetSubPins()[7]->GetDefaultValue() == TEXT("34.000000"), TEXT("FAnimationAnimNextEditorTest_GraphAddTrait -> Unexpected trait pin value"));
+		UE_RETURN_ON_ERROR(TraitPin->GetSubPins()[7]->IsLazy(), TEXT("FAnimationAnimNextEditorTest_GraphAddTrait -> Expected lazy trait pin"));
 	}
 
 	Tests::FUtils::CleanupAfterTests();
@@ -208,16 +215,14 @@ bool FAnimationAnimNextEditorTest_GraphTraitOperations::RunTest(const FString& I
 
 		FScopedClearNodeTemplateRegistry ScopedClearNodeTemplateRegistry;
 
-		UFactory* GraphFactory = NewObject<UAnimNextModuleFactory>();
-		UAnimNextModule* Module = CastChecked<UAnimNextModule>(GraphFactory->FactoryCreateNew(UAnimNextModule::StaticClass(), GetTransientPackage(), TEXT("TestAnimNextGraph"), RF_Transient, nullptr, nullptr, NAME_None));
-		UE_RETURN_ON_ERROR(Module != nullptr, "FAnimationAnimNextEditorTest_GraphTraitOperations -> Failed to create module");
+		UFactory* GraphFactory = NewObject<UAnimNextAnimationGraphFactory>();
+		UAnimNextAnimationGraph* AnimationGraph = CastChecked<UAnimNextAnimationGraph>(GraphFactory->FactoryCreateNew(UAnimNextAnimationGraph::StaticClass(), GetTransientPackage(), TEXT("TestAnimNextGraph"), RF_Transient, nullptr, nullptr, NAME_None));
+		UE_RETURN_ON_ERROR(AnimationGraph != nullptr, "FAnimationAnimNextEditorTest_GraphTraitOperations -> Failed to create animation graph");
 
-		UAnimNextModule_EditorData* EditorData = UncookedOnly::FUtils::GetEditorData(Module);
+		UAnimNextAnimationGraph_EditorData* EditorData = UncookedOnly::FUtils::GetEditorData<UAnimNextAnimationGraph_EditorData>(AnimationGraph);
 		UE_RETURN_ON_ERROR(EditorData != nullptr, "FAnimationAnimNextEditorTest_GraphTraitOperations -> Failed to find module editor data");
 
-		UE_RETURN_ON_ERROR(EditorData->AddAnimationGraph(FRigUnit_AnimNextGraphRoot::DefaultEntryPoint, false) != nullptr, "FAnimationAnimNextEditorTest_GraphTraitOperations -> Failed to add animation graph");
-
-		UAnimNextModule_Controller* Controller = Cast<UAnimNextModule_Controller>(EditorData->GetRigVMClient()->GetController(EditorData->GetRigVMClient()->GetDefaultModel()));
+		UAnimNextController* Controller = Cast<UAnimNextController>(EditorData->GetRigVMClient()->GetController(EditorData->GetRigVMClient()->GetDefaultModel()));
 		UE_RETURN_ON_ERROR(Controller != nullptr, "FAnimationAnimNextEditorTest_GraphTraitOperations -> Failed to get RigVM controller");
 
 		// Create an empty trait stack node
@@ -248,32 +253,42 @@ bool FAnimationAnimNextEditorTest_GraphTraitOperations::RunTest(const FString& I
 			UE_RETURN_ON_ERROR(TraitPin->GetCPPTypeObject() == FRigDecorator_AnimNextCppDecorator::StaticStruct(), TEXT("FAnimationAnimNextEditorTest_GraphTraitOperations -> Unexpected pin type"));
 
 			// Our first sub-pin is the hard coded script struct member that parametrizes the trait, dynamic trait sub-pins follow
-			UE_RETURN_ON_ERROR(TraitPin->GetSubPins().Num() == 6, TEXT("FAnimationAnimNextEditorTest_GraphTraitOperations -> Unexpected trait sub pins"));
+			UE_RETURN_ON_ERROR(TraitPin->GetSubPins().Num() == 8, TEXT("FAnimationAnimNextEditorTest_GraphAddTrait -> Unexpected trait sub pins"));
+
+			// UpdateCount
+			UE_RETURN_ON_ERROR(TraitPin->GetSubPins()[1]->GetCPPType() == TEXT("int32"), TEXT("FAnimationAnimNextEditorTest_GraphAddTrait -> Unexpected trait pin type"));
+			UE_RETURN_ON_ERROR(TraitPin->GetSubPins()[1]->GetDefaultValue() == TEXT("0"), TEXT("FAnimationAnimNextEditorTest_GraphAddTrait -> Unexpected trait pin value"));
+			UE_RETURN_ON_ERROR(!TraitPin->GetSubPins()[1]->IsLazy(), TEXT("FAnimationAnimNextEditorTest_GraphAddTrait -> Expected non-lazy trait pin"));
+
+			// EvaluateCount
+			UE_RETURN_ON_ERROR(TraitPin->GetSubPins()[2]->GetCPPType() == TEXT("int32"), TEXT("FAnimationAnimNextEditorTest_GraphAddTrait -> Unexpected trait pin type"));
+			UE_RETURN_ON_ERROR(TraitPin->GetSubPins()[2]->GetDefaultValue() == TEXT("0"), TEXT("FAnimationAnimNextEditorTest_GraphAddTrait -> Unexpected trait pin value"));
+			UE_RETURN_ON_ERROR(!TraitPin->GetSubPins()[2]->IsLazy(), TEXT("FAnimationAnimNextEditorTest_GraphAddTrait -> Expected non-lazy trait pin"));
 
 			// SomeInt32
-			UE_RETURN_ON_ERROR(TraitPin->GetSubPins()[1]->GetCPPType() == TEXT("int32"), TEXT("FAnimationAnimNextEditorTest_GraphTraitOperations -> Unexpected trait pin type"));
-			UE_RETURN_ON_ERROR(TraitPin->GetSubPins()[1]->GetDefaultValue() == TEXT("3"), TEXT("FAnimationAnimNextEditorTest_GraphTraitOperations -> Unexpected trait pin value"));
-			UE_RETURN_ON_ERROR(!TraitPin->GetSubPins()[1]->IsLazy(), TEXT("FAnimationAnimNextEditorTest_GraphTraitOperations -> Expected non-lazy trait pin"));
+			UE_RETURN_ON_ERROR(TraitPin->GetSubPins()[3]->GetCPPType() == TEXT("int32"), TEXT("FAnimationAnimNextEditorTest_GraphAddTrait -> Unexpected trait pin type"));
+			UE_RETURN_ON_ERROR(TraitPin->GetSubPins()[3]->GetDefaultValue() == TEXT("3"), TEXT("FAnimationAnimNextEditorTest_GraphAddTrait -> Unexpected trait pin value"));
+			UE_RETURN_ON_ERROR(!TraitPin->GetSubPins()[3]->IsLazy(), TEXT("FAnimationAnimNextEditorTest_GraphAddTrait -> Expected non-lazy trait pin"));
 
 			// SomeFloat
-			UE_RETURN_ON_ERROR(TraitPin->GetSubPins()[2]->GetCPPType() == TEXT("float"), TEXT("FAnimationAnimNextEditorTest_GraphTraitOperations -> Unexpected trait pin type"));
-			UE_RETURN_ON_ERROR(TraitPin->GetSubPins()[2]->GetDefaultValue() == TEXT("34.000000"), TEXT("FAnimationAnimNextEditorTest_GraphTraitOperations -> Unexpected trait pin value"));
-			UE_RETURN_ON_ERROR(!TraitPin->GetSubPins()[2]->IsLazy(), TEXT("FAnimationAnimNextEditorTest_GraphTraitOperations -> Expected non-lazy trait pin"));
+			UE_RETURN_ON_ERROR(TraitPin->GetSubPins()[4]->GetCPPType() == TEXT("float"), TEXT("FAnimationAnimNextEditorTest_GraphAddTrait -> Unexpected trait pin type"));
+			UE_RETURN_ON_ERROR(TraitPin->GetSubPins()[4]->GetDefaultValue() == TEXT("34.000000"), TEXT("FAnimationAnimNextEditorTest_GraphAddTrait -> Unexpected trait pin value"));
+			UE_RETURN_ON_ERROR(!TraitPin->GetSubPins()[4]->IsLazy(), TEXT("FAnimationAnimNextEditorTest_GraphAddTrait -> Expected non-lazy trait pin"));
 
 			// SomeLatentInt32
-			UE_RETURN_ON_ERROR(TraitPin->GetSubPins()[3]->GetCPPType() == TEXT("int32"), TEXT("FAnimationAnimNextEditorTest_GraphTraitOperations -> Unexpected trait pin type"));
-			UE_RETURN_ON_ERROR(TraitPin->GetSubPins()[3]->GetDefaultValue() == TEXT("3"), TEXT("FAnimationAnimNextEditorTest_GraphTraitOperations -> Unexpected trait pin value"));
-			UE_RETURN_ON_ERROR(TraitPin->GetSubPins()[3]->IsLazy(), TEXT("FAnimationAnimNextEditorTest_GraphTraitOperations -> Expected lazy trait pin"));
+			UE_RETURN_ON_ERROR(TraitPin->GetSubPins()[5]->GetCPPType() == TEXT("int32"), TEXT("FAnimationAnimNextEditorTest_GraphAddTrait -> Unexpected trait pin type"));
+			UE_RETURN_ON_ERROR(TraitPin->GetSubPins()[5]->GetDefaultValue() == TEXT("3"), TEXT("FAnimationAnimNextEditorTest_GraphAddTrait -> Unexpected trait pin value"));
+			UE_RETURN_ON_ERROR(TraitPin->GetSubPins()[5]->IsLazy(), TEXT("FAnimationAnimNextEditorTest_GraphAddTrait -> Expected lazy trait pin"));
 
 			// SomeOtherLatentInt32
-			UE_RETURN_ON_ERROR(TraitPin->GetSubPins()[4]->GetCPPType() == TEXT("int32"), TEXT("FAnimationAnimNextEditorTest_GraphTraitOperations -> Unexpected trait pin type"));
-			UE_RETURN_ON_ERROR(TraitPin->GetSubPins()[4]->GetDefaultValue() == TEXT("3"), TEXT("FAnimationAnimNextEditorTest_GraphTraitOperations -> Unexpected trait pin value"));
-			UE_RETURN_ON_ERROR(TraitPin->GetSubPins()[4]->IsLazy(), TEXT("FAnimationAnimNextEditorTest_GraphTraitOperations -> Expected lazy trait pin"));
+			UE_RETURN_ON_ERROR(TraitPin->GetSubPins()[6]->GetCPPType() == TEXT("int32"), TEXT("FAnimationAnimNextEditorTest_GraphAddTrait -> Unexpected trait pin type"));
+			UE_RETURN_ON_ERROR(TraitPin->GetSubPins()[6]->GetDefaultValue() == TEXT("3"), TEXT("FAnimationAnimNextEditorTest_GraphAddTrait -> Unexpected trait pin value"));
+			UE_RETURN_ON_ERROR(TraitPin->GetSubPins()[6]->IsLazy(), TEXT("FAnimationAnimNextEditorTest_GraphAddTrait -> Expected lazy trait pin"));
 
 			// SomeLatentFloat
-			UE_RETURN_ON_ERROR(TraitPin->GetSubPins()[5]->GetCPPType() == TEXT("float"), TEXT("FAnimationAnimNextEditorTest_GraphTraitOperations -> Unexpected trait pin type"));
-			UE_RETURN_ON_ERROR(TraitPin->GetSubPins()[5]->GetDefaultValue() == TEXT("34.000000"), TEXT("FAnimationAnimNextEditorTest_GraphTraitOperations -> Unexpected trait pin value"));
-			UE_RETURN_ON_ERROR(TraitPin->GetSubPins()[5]->IsLazy(), TEXT("FAnimationAnimNextEditorTest_GraphTraitOperations -> Expected lazy trait pin"));
+			UE_RETURN_ON_ERROR(TraitPin->GetSubPins()[7]->GetCPPType() == TEXT("float"), TEXT("FAnimationAnimNextEditorTest_GraphAddTrait -> Unexpected trait pin type"));
+			UE_RETURN_ON_ERROR(TraitPin->GetSubPins()[7]->GetDefaultValue() == TEXT("34.000000"), TEXT("FAnimationAnimNextEditorTest_GraphAddTrait -> Unexpected trait pin value"));
+			UE_RETURN_ON_ERROR(TraitPin->GetSubPins()[7]->IsLazy(), TEXT("FAnimationAnimNextEditorTest_GraphAddTrait -> Expected lazy trait pin"));
 		}
 
 		// --- Undo Add Trait ---
@@ -370,15 +385,13 @@ bool FAnimationAnimNextRuntimeTest_GraphExecute::RunTest(const FString& InParame
 
 		FScopedClearNodeTemplateRegistry ScopedClearNodeTemplateRegistry;
 
-		UFactory* GraphFactory = NewObject<UAnimNextModuleFactory>();
-		UAnimNextModule* Module = CastChecked<UAnimNextModule>(GraphFactory->FactoryCreateNew(UAnimNextModule::StaticClass(), GetTransientPackage(), TEXT("TestAnimNextGraph"), RF_Transient, nullptr, nullptr, NAME_None));
-		UE_RETURN_ON_ERROR(Module != nullptr, "FAnimationAnimNextRuntimeTest_GraphExecute -> Failed to create module");
+		UFactory* GraphFactory = NewObject<UAnimNextAnimationGraphFactory>();
+		UAnimNextAnimationGraph* AnimationGraph = CastChecked<UAnimNextAnimationGraph>(GraphFactory->FactoryCreateNew(UAnimNextAnimationGraph::StaticClass(), GetTransientPackage(), TEXT("TestAnimNextGraph"), RF_Transient, nullptr, nullptr, NAME_None));
+		UE_RETURN_ON_ERROR(AnimationGraph != nullptr, "FAnimationAnimNextEditorTest_GraphTraitOperations -> Failed to create animation graph");
 
-		UAnimNextModule_EditorData* EditorData = UncookedOnly::FUtils::GetEditorData(Module);
+		UAnimNextAnimationGraph_EditorData* EditorData = UncookedOnly::FUtils::GetEditorData<UAnimNextAnimationGraph_EditorData>(AnimationGraph);
 		UE_RETURN_ON_ERROR(EditorData != nullptr, "FAnimationAnimNextRuntimeTest_GraphExecute -> Failed to find module editor data");
 
-		UE_RETURN_ON_ERROR(EditorData->AddAnimationGraph(FRigUnit_AnimNextGraphRoot::DefaultEntryPoint, false) != nullptr, "FAnimationAnimNextRuntimeTest_GraphExecute -> Failed to add animation graph");
-
 		URigVMController* Controller = EditorData->GetRigVMClient()->GetController(EditorData->GetRigVMClient()->GetDefaultModel());
 		UE_RETURN_ON_ERROR(Controller != nullptr, "FAnimationAnimNextRuntimeTest_GraphExecute -> Failed to get RigVM controller");
 
@@ -435,24 +448,19 @@ bool FAnimationAnimNextRuntimeTest_GraphExecute::RunTest(const FString& InParame
 		UE_RETURN_ON_ERROR(DecoratorPin != nullptr, TEXT("FAnimationAnimNextRuntimeTest_GraphExecute -> Failed to find trait pin"));
 
 		// Set some values on our trait
-		Controller->SetPinDefaultValue(DecoratorPin->GetSubPins()[1]->GetPinPath(), TEXT("78"));
-		Controller->SetPinDefaultValue(DecoratorPin->GetSubPins()[2]->GetPinPath(), TEXT("142.33"));
-
-		TSharedRef<FParamStack> ParamStack = MakeShared<FParamStack>();
-		FParamStack::AttachToCurrentThread(ParamStack);
+		Controller->SetPinDefaultValue(DecoratorPin->GetSubPins()[3]->GetPinPath(), TEXT("78"));
+		Controller->SetPinDefaultValue(DecoratorPin->GetSubPins()[4]->GetPinPath(), TEXT("142.33"));
 
 		FAnimNextGraphInstancePtr GraphInstance;
-		Module->AllocateInstance(GraphInstance);
-
-		FParamStack::FPushedLayerHandle LayerHandle = ParamStack->PushValues(
-			"/Engine/Transient.TestAnimNextGraph:UpdateCount", (int32)0,
-			"/Engine/Transient.TestAnimNextGraph:EvaluateCount", (int32)0,
-			"/Engine/Transient.TestAnimNextGraph:SomeInt32", (int32)0,
-			"/Engine/Transient.TestAnimNextGraph:SomeFloat", 0.0f,
-			"/Engine/Transient.TestAnimNextGraph:SomeLatentInt32", (int32)0,
-			"/Engine/Transient.TestAnimNextGraph:SomeOtherLatentInt32", (int32)0,
-			"/Engine/Transient.TestAnimNextGraph:SomeLatentFloat", 0.0f
-		);
+		AnimationGraph->AllocateInstance(GraphInstance);
+
+		TArray<FString> Messages;
+		FRigVMRuntimeSettings RuntimeSettings;
+		RuntimeSettings.SetLogFunction([&Messages](const FRigVMLogSettings& InLogSettings, const FRigVMExecuteContext* InContext, const FString& Message)
+		{
+			Messages.Add(Message);
+		});
+		GraphInstance.GetExtendedExecuteContext().SetRuntimeSettings(RuntimeSettings);
 
 		{
 			UE::AnimNext::FTraitEventList InputEventList;
@@ -461,18 +469,14 @@ bool FAnimationAnimNextRuntimeTest_GraphExecute::RunTest(const FString& InParame
 			(void)UE::AnimNext::EvaluateGraph(GraphInstance);
 		}
 
-		AddErrorIfFalse(ParamStack->GetParam<int32>("/Engine/Transient.TestAnimNextGraph:UpdateCount") == 1, "FAnimationAnimNextRuntimeTest_GraphExecute -> Unexpected update count");
-		AddErrorIfFalse(ParamStack->GetParam<int32>("/Engine/Transient.TestAnimNextGraph:EvaluateCount") == 1, "FAnimationAnimNextRuntimeTest_GraphExecute -> Unexpected evaluate count");
-		AddErrorIfFalse(ParamStack->GetParam<int32>("/Engine/Transient.TestAnimNextGraph:SomeInt32") == 78, "FAnimationAnimNextRuntimeTest_GraphExecute -> Unexpected SomeInt32 value");
-		AddErrorIfFalse(ParamStack->GetParam<float>("/Engine/Transient.TestAnimNextGraph:SomeFloat") == 142.33f, "FAnimationAnimNextRuntimeTest_GraphExecute -> Unexpected SomeFloat value");
-		AddErrorIfFalse(ParamStack->GetParam<int32>("/Engine/Transient.TestAnimNextGraph:SomeLatentInt32") == 3, "FAnimationAnimNextRuntimeTest_GraphExecute -> Unexpected SomeLatentInt32 value");
-		AddErrorIfFalse(ParamStack->GetParam<int32>("/Engine/Transient.TestAnimNextGraph:SomeOtherLatentInt32") == 3, "FAnimationAnimNextRuntimeTest_GraphExecute -> Unexpected SomeOtherLatentInt32 value");
-		AddErrorIfFalse(ParamStack->GetParam<float>("/Engine/Transient.TestAnimNextGraph:SomeLatentFloat") == 34.0f, "FAnimationAnimNextRuntimeTest_GraphExecute -> Unexpected SomeLatentFloat value");
-
-		ParamStack->PopLayer(LayerHandle);
-		GraphInstance.Release();
-
-		FParamStack::DetachFromCurrentThread();
+		AddErrorIfFalse(Messages.Num() == 7, "FAnimationAnimNextRuntimeTest_GraphExecute - > Unexpected message count");
+		AddErrorIfFalse(Messages[0] == "UpdateCount == 1", "FAnimationAnimNextRuntimeTest_GraphExecute -> Unexpected update count");
+		AddErrorIfFalse(Messages[1] == "SomeInt32 == 78", "FAnimationAnimNextRuntimeTest_GraphExecute -> Unexpected SomeInt32 value");
+		AddErrorIfFalse(Messages[2] == "SomeFloat == 142.33", "FAnimationAnimNextRuntimeTest_GraphExecute -> Unexpected SomeFloat value");
+		AddErrorIfFalse(Messages[3] == "SomeLatentInt32 == 3", "FAnimationAnimNextRuntimeTest_GraphExecute -> Unexpected SomeLatentInt32 value");
+		AddErrorIfFalse(Messages[4] == "SomeOtherLatentInt32 == 3", "FAnimationAnimNextRuntimeTest_GraphExecute -> Unexpected SomeOtherLatentInt32 value");
+		AddErrorIfFalse(Messages[5] == "SomeLatentFloat == 34.00", "FAnimationAnimNextRuntimeTest_GraphExecute -> Unexpected SomeLatentFloat value");
+		AddErrorIfFalse(Messages[6] == "EvaluateCount == 1", "FAnimationAnimNextRuntimeTest_GraphExecute -> Unexpected evaluate count");
 	}
 
 	Tests::FUtils::CleanupAfterTests();
@@ -490,16 +494,16 @@ bool FAnimationAnimNextRuntimeTest_GraphExecuteLatent::RunTest(const FString& In
 		AUTO_REGISTER_ANIM_TRAIT(FTestTrait)
 		FScopedClearNodeTemplateRegistry ScopedClearNodeTemplateRegistry;
 
-		UFactory* GraphFactory = NewObject<UAnimNextModuleFactory>();
-		UAnimNextModule* Module = CastChecked<UAnimNextModule>(GraphFactory->FactoryCreateNew(UAnimNextModule::StaticClass(), GetTransientPackage(), TEXT("TestAnimNextGraph"), RF_Transient, nullptr, nullptr, NAME_None));
-		UE_RETURN_ON_ERROR(Module != nullptr, "FAnimationAnimNextRuntimeTest_GraphExecuteLatent -> Failed to create module");
+		UFactory* GraphFactory = NewObject<UAnimNextAnimationGraphFactory>();
+		UAnimNextAnimationGraph* AnimationGraph = CastChecked<UAnimNextAnimationGraph>(GraphFactory->FactoryCreateNew(UAnimNextAnimationGraph::StaticClass(), GetTransientPackage(), TEXT("TestAnimNextGraph"), RF_Transient, nullptr, nullptr, NAME_None));
+		UE_RETURN_ON_ERROR(AnimationGraph != nullptr, "FAnimationAnimNextEditorTest_GraphTraitOperations -> Failed to create animation graph");
 
-		UAnimNextModule_EditorData* EditorData = UncookedOnly::FUtils::GetEditorData(Module);
+		UAnimNextAnimationGraph_EditorData* EditorData = UncookedOnly::FUtils::GetEditorData<UAnimNextAnimationGraph_EditorData>(AnimationGraph);
 		UE_RETURN_ON_ERROR(EditorData != nullptr, "FAnimationAnimNextRuntimeTest_GraphExecuteLatent -> Failed to find module editor data");
 
-		UE_RETURN_ON_ERROR(EditorData->AddAnimationGraph(FRigUnit_AnimNextGraphRoot::DefaultEntryPoint, false) != nullptr, "FAnimationAnimNextRuntimeTest_GraphExecuteLatent -> Failed to add animation graph");
-
-		UAnimNextModule_Controller* Controller = Cast<UAnimNextModule_Controller>(EditorData->GetRigVMClient()->GetController(EditorData->GetRigVMClient()->GetDefaultModel()));
+		EditorData->AddVariable(TEXT("TestIntVar"), FAnimNextParamType::GetType<int32>(), TEXT("34"));
+		
+		UAnimNextController* Controller = Cast<UAnimNextController>(EditorData->GetRigVMClient()->GetController(EditorData->GetRigVMClient()->GetDefaultModel()));
 		UE_RETURN_ON_ERROR(Controller != nullptr, "FAnimationAnimNextRuntimeTest_GraphExecuteLatent -> Failed to get RigVM controller");
 
 		// Find graph entry point
@@ -555,9 +559,9 @@ bool FAnimationAnimNextRuntimeTest_GraphExecuteLatent::RunTest(const FString& In
 		URigVMPin* DecoratorPin = DecoratorStackNode->FindPin(*DisplayName);
 		UE_RETURN_ON_ERROR(DecoratorPin != nullptr, TEXT("FAnimationAnimNextRuntimeTest_GraphExecuteLatent -> Failed to find trait pin"));
 
-		Controller->SetPinDefaultValue(DecoratorPin->GetSubPins()[1]->GetPinPath(), TEXT("78"));		// SomeInt32
-		Controller->SetPinDefaultValue(DecoratorPin->GetSubPins()[2]->GetPinPath(), TEXT("142.33"));	// SomeFloat
-		Controller->SetPinDefaultValue(DecoratorPin->GetSubPins()[5]->GetPinPath(), TEXT("1123.31"));	// SomeLatentFloat, inline value on latent pin
+		Controller->SetPinDefaultValue(DecoratorPin->GetSubPins()[3]->GetPinPath(), TEXT("78"));		// SomeInt32
+		Controller->SetPinDefaultValue(DecoratorPin->GetSubPins()[4]->GetPinPath(), TEXT("142.33"));	// SomeFloat
+		Controller->SetPinDefaultValue(DecoratorPin->GetSubPins()[7]->GetPinPath(), TEXT("1123.31"));	// SomeLatentFloat, inline value on latent pin
 
 		// Set some latent values on our trait
 		{
@@ -570,35 +574,29 @@ bool FAnimationAnimNextRuntimeTest_GraphExecuteLatent::RunTest(const FString& In
 
 			Controller->AddLink(
 				IntAddNode->FindPin(GET_MEMBER_NAME_STRING_CHECKED(FRigVMFunction_MathIntAdd, Result)),
-				DecoratorPin->GetSubPins()[3]);	// SomeLatentInt32
+				DecoratorPin->GetSubPins()[5]);	// SomeLatentInt32
 		}
 
 		{
-			URigVMNode* GetParameterNode = Controller->AddGetAnimNextParameterNode(FVector2D::ZeroVector, UncookedOnly::FUtils::GetQualifiedName(Module, "SomeSourceInt"), FAnimNextParamType::GetType<int32>());
-			UE_RETURN_ON_ERROR(GetParameterNode != nullptr, TEXT("FAnimationAnimNextRuntimeTest_GraphExecuteLatent -> Failed to create GetParameter node"));
+			URigVMVariableNode* GetVariableNode = Controller->AddVariableNode(TEXT("TestIntVar"), RigVMTypeUtils::Int32Type, nullptr, true, TEXT(""));
+			UE_RETURN_ON_ERROR(GetVariableNode != nullptr, TEXT("FAnimationAnimNextRuntimeTest_GraphExecuteLatent -> Failed to create variable node"));
 
 			Controller->AddLink(
-				GetParameterNode->FindPin(TEXT("Value")),
-				DecoratorPin->GetSubPins()[4]);	// SomeOtherLatentInt32
+				GetVariableNode->FindPin(TEXT("Value")),
+				DecoratorPin->GetSubPins()[6]);	// SomeOtherLatentInt32
 		}
 
-		TSharedRef<FParamStack> ParamStack = MakeShared<FParamStack>();
-		FParamStack::AttachToCurrentThread(ParamStack);
-
 		FAnimNextGraphInstancePtr GraphInstance;
-		Module->AllocateInstance(GraphInstance);
-
-		FParamStack::FPushedLayerHandle LayerHandle = ParamStack->PushValues(
-			UncookedOnly::FUtils::GetQualifiedName(Module, "UpdateCount"), (int32)0,
-			UncookedOnly::FUtils::GetQualifiedName(Module, "EvaluateCount"), (int32)0,
-			UncookedOnly::FUtils::GetQualifiedName(Module, "SomeSourceInt"), (int32)1223,
-			UncookedOnly::FUtils::GetQualifiedName(Module, "SomeInt32"), (int32)0,
-			UncookedOnly::FUtils::GetQualifiedName(Module, "SomeFloat"), 0.0f,
-			UncookedOnly::FUtils::GetQualifiedName(Module, "SomeLatentInt32"), (int32)0,
-			UncookedOnly::FUtils::GetQualifiedName(Module, "SomeOtherLatentInt32"), (int32)0,
-			UncookedOnly::FUtils::GetQualifiedName(Module, "SomeLatentFloat"), 0.0f
-		);
+		AnimationGraph->AllocateInstance(GraphInstance);
 
+		TArray<FString> Messages;
+		FRigVMRuntimeSettings RuntimeSettings;
+		RuntimeSettings.SetLogFunction([&Messages](const FRigVMLogSettings& InLogSettings, const FRigVMExecuteContext* InContext, const FString& Message)
+		{
+			Messages.Add(Message);
+		});
+		GraphInstance.GetExtendedExecuteContext().SetRuntimeSettings(RuntimeSettings);
+		
 		{
 			UE::AnimNext::FTraitEventList InputEventList;
 			UE::AnimNext::FTraitEventList OutputEventList;
@@ -606,18 +604,14 @@ bool FAnimationAnimNextRuntimeTest_GraphExecuteLatent::RunTest(const FString& In
 			(void)UE::AnimNext::EvaluateGraph(GraphInstance);
 		}
 
-		AddErrorIfFalse(ParamStack->GetParam<int32>(UncookedOnly::FUtils::GetQualifiedName(Module, "UpdateCount")) == 1, "FAnimationAnimNextRuntimeTest_GraphExecuteLatent -> Unexpected update count");
-		AddErrorIfFalse(ParamStack->GetParam<int32>(UncookedOnly::FUtils::GetQualifiedName(Module, "EvaluateCount")) == 1, "FAnimationAnimNextRuntimeTest_GraphExecuteLatent -> Unexpected evaluate count");
-		AddErrorIfFalse(ParamStack->GetParam<int32>(UncookedOnly::FUtils::GetQualifiedName(Module, "SomeInt32")) == 78, "FAnimationAnimNextRuntimeTest_GraphExecuteLatent -> Unexpected SomeInt32 value");
-		AddErrorIfFalse(ParamStack->GetParam<float>(UncookedOnly::FUtils::GetQualifiedName(Module, "SomeFloat")) == 142.33f, "FAnimationAnimNextRuntimeTest_GraphExecuteLatent -> Unexpected SomeFloat value");
-		AddErrorIfFalse(ParamStack->GetParam<int32>(UncookedOnly::FUtils::GetQualifiedName(Module, "SomeLatentInt32")) == 33, "FAnimationAnimNextRuntimeTest_GraphExecuteLatent -> Unexpected SomeLatentInt32 value");
-		AddErrorIfFalse(ParamStack->GetParam<int32>(UncookedOnly::FUtils::GetQualifiedName(Module, "SomeOtherLatentInt32")) == 1223, "FAnimationAnimNextRuntimeTest_GraphExecuteLatent -> Unexpected SomeOtherLatentInt32 value");
-		AddErrorIfFalse(ParamStack->GetParam<float>(UncookedOnly::FUtils::GetQualifiedName(Module, "SomeLatentFloat")) == 1123.31f, "FAnimationAnimNextRuntimeTest_GraphExecuteLatent -> Unexpected SomeLatentFloat value");
-
-		ParamStack->PopLayer(LayerHandle);
-		GraphInstance.Release();
-
-		FParamStack::DetachFromCurrentThread();
+		AddErrorIfFalse(Messages.Num() == 7, "FAnimationAnimNextRuntimeTest_GraphExecuteLatent - > Unexpected message count");
+		AddErrorIfFalse(Messages[0] == "UpdateCount == 1", "FAnimationAnimNextRuntimeTest_GraphExecuteLatent -> Unexpected update count");
+		AddErrorIfFalse(Messages[1] == "SomeInt32 == 78", "FAnimationAnimNextRuntimeTest_GraphExecuteLatent -> Unexpected SomeInt32 value");
+		AddErrorIfFalse(Messages[2] == "SomeFloat == 142.33", "FAnimationAnimNextRuntimeTest_GraphExecuteLatent -> Unexpected SomeFloat value");
+		AddErrorIfFalse(Messages[3] == "SomeLatentInt32 == 33", "FAnimationAnimNextRuntimeTest_GraphExecuteLatent -> Unexpected SomeLatentInt32 value");
+		AddErrorIfFalse(Messages[4] == "SomeOtherLatentInt32 == 34", "FAnimationAnimNextRuntimeTest_GraphExecuteLatent -> Unexpected SomeOtherLatentInt32 value");
+		AddErrorIfFalse(Messages[5] == "SomeLatentFloat == 1123.31", "FAnimationAnimNextRuntimeTest_GraphExecuteLatent -> Unexpected SomeLatentFloat value");
+		AddErrorIfFalse(Messages[6] == "EvaluateCount == 1", "FAnimationAnimNextRuntimeTest_GraphExecuteLatent -> Unexpected evaluate count");
 	}
 
 	Tests::FUtils::CleanupAfterTests();
diff --git a/Engine/Plugins/Experimental/AnimNextAnimGraph/Source/AnimNextAnimGraphTestSuite/Private/AnimNextAnimGraphTraitGraphTest.h b/Engine/Plugins/Experimental/AnimNextAnimGraph/Source/AnimNextAnimGraphTestSuite/Private/AnimNextAnimGraphTraitGraphTest.h
--- a/Engine/Plugins/Experimental/AnimNextAnimGraph/Source/AnimNextAnimGraphTestSuite/Private/AnimNextAnimGraphTraitGraphTest.h
+++ b/Engine/Plugins/Experimental/AnimNextAnimGraph/Source/AnimNextAnimGraphTestSuite/Private/AnimNextAnimGraphTraitGraphTest.h
@@ -15,6 +15,12 @@ struct FTestTraitSharedData : public FAnimNextTraitSharedData
 {
 	GENERATED_BODY()
 
+	UPROPERTY(meta = (Input, Inline))
+	int32 UpdateCount = 0;
+
+	UPROPERTY(meta = (Input, Inline))
+	int32 EvaluateCount = 0;
+
 	UPROPERTY(meta = (Input, Inline))
 	int32 SomeInt32 = 3;
 
diff --git a/Engine/Plugins/Experimental/AnimNextAnimGraph/Source/AnimNextAnimGraphTestSuite/Private/AnimNextAnimGraphTraitInterfacesTest.cpp b/Engine/Plugins/Experimental/AnimNextAnimGraph/Source/AnimNextAnimGraphTestSuite/Private/AnimNextAnimGraphTraitInterfacesTest.cpp
--- a/Engine/Plugins/Experimental/AnimNextAnimGraph/Source/AnimNextAnimGraphTestSuite/Private/AnimNextAnimGraphTraitInterfacesTest.cpp
+++ b/Engine/Plugins/Experimental/AnimNextAnimGraph/Source/AnimNextAnimGraphTestSuite/Private/AnimNextAnimGraphTraitInterfacesTest.cpp
@@ -24,8 +24,8 @@
 #include "TraitInterfaces/IEvaluate.h"
 #include "TraitInterfaces/IHierarchy.h"
 #include "TraitInterfaces/IUpdate.h"
-#include "Module/AnimNextModule.h"
-#include "Module/ModuleFactory.h"
+#include "Graph/AnimNextAnimationGraph.h"
+#include "Graph/AnimNextAnimationGraphFactory.h"
 
 //****************************************************************************
 // AnimNext Runtime TraitInterfaces Tests
@@ -477,9 +477,9 @@ bool FAnimationAnimNextRuntimeTest_IHierarchy::RunTest(const FString& InParamete
 		AUTO_REGISTER_ANIM_TRAIT(FTraitWithOneChild)
 		AUTO_REGISTER_ANIM_TRAIT(FTraitWithChildren)
 
-		UFactory* ModuleFactory = NewObject<UAnimNextModuleFactory>();
-		UAnimNextModule* Module = CastChecked<UAnimNextModule>(ModuleFactory->FactoryCreateNew(UAnimNextModule::StaticClass(), GetTransientPackage(), TEXT("TestAnimNextGraph"), RF_Transient, nullptr, nullptr, NAME_None));
-		UE_RETURN_ON_ERROR(Module != nullptr, "FAnimationAnimNextRuntimeTest_IHierarchy -> Failed to create module");
+		UFactory* GraphFactory = NewObject<UAnimNextAnimationGraphFactory>();
+		UAnimNextAnimationGraph* AnimationGraph = CastChecked<UAnimNextAnimationGraph>(GraphFactory->FactoryCreateNew(UAnimNextAnimationGraph::StaticClass(), GetTransientPackage(), TEXT("TestAnimNextGraph"), RF_Transient, nullptr, nullptr, NAME_None));
+		UE_RETURN_ON_ERROR(AnimationGraph != nullptr, "FAnimationAnimNextRuntimeTest_GraphTraitEvent -> Failed to create animation graph");
 
 		FScopedClearNodeTemplateRegistry ScopedClearNodeTemplateRegistry;
 		FNodeTemplateRegistry& Registry = FNodeTemplateRegistry::Get();
@@ -591,10 +591,10 @@ bool FAnimationAnimNextRuntimeTest_IHierarchy::RunTest(const FString& InParamete
 		}
 
 		// Read our graph
-		FTestUtils::LoadFromArchiveBuffer(*Module, NodeHandles, GraphSharedDataArchiveBuffer);
+		FTestUtils::LoadFromArchiveBuffer(*AnimationGraph, NodeHandles, GraphSharedDataArchiveBuffer);
 
 		FAnimNextGraphInstancePtr GraphInstance;
-		Module->AllocateInstance(GraphInstance);
+		AnimationGraph->AllocateInstance(GraphInstance);
 
 		FExecutionContext Context(GraphInstance);
 
@@ -698,9 +698,9 @@ bool FAnimationAnimNextRuntimeTest_IUpdate::RunTest(const FString& InParameters)
 		AUTO_REGISTER_ANIM_TRAIT(FTraitWithOneChild)
 		AUTO_REGISTER_ANIM_TRAIT(FTraitWithChildren)
 
-		UFactory* ModuleFactory = NewObject<UAnimNextModuleFactory>();
-		UAnimNextModule* Module = CastChecked<UAnimNextModule>(ModuleFactory->FactoryCreateNew(UAnimNextModule::StaticClass(), GetTransientPackage(), TEXT("TestAnimNextGraph"), RF_Transient, nullptr, nullptr, NAME_None));
-		UE_RETURN_ON_ERROR(Module != nullptr, "FAnimationAnimNextRuntimeTest_IUpdate -> Failed to create module");
+		UFactory* GraphFactory = NewObject<UAnimNextAnimationGraphFactory>();
+		UAnimNextAnimationGraph* AnimationGraph = CastChecked<UAnimNextAnimationGraph>(GraphFactory->FactoryCreateNew(UAnimNextAnimationGraph::StaticClass(), GetTransientPackage(), TEXT("TestAnimNextGraph"), RF_Transient, nullptr, nullptr, NAME_None));
+		UE_RETURN_ON_ERROR(AnimationGraph != nullptr, "FAnimationAnimNextRuntimeTest_GraphTraitEvent -> Failed to create animation graph");
 
 		FScopedClearNodeTemplateRegistry ScopedClearNodeTemplateRegistry;
 		FNodeTemplateRegistry& Registry = FNodeTemplateRegistry::Get();
@@ -790,10 +790,10 @@ bool FAnimationAnimNextRuntimeTest_IUpdate::RunTest(const FString& InParameters)
 		}
 
 		// Read our graph
-		FTestUtils::LoadFromArchiveBuffer(*Module, NodeHandles, GraphSharedDataArchiveBuffer);
+		FTestUtils::LoadFromArchiveBuffer(*AnimationGraph, NodeHandles, GraphSharedDataArchiveBuffer);
 
 		FAnimNextGraphInstancePtr GraphInstance;
-		Module->AllocateInstance(GraphInstance);
+		AnimationGraph->AllocateInstance(GraphInstance);
 
 		FExecutionContext Context(GraphInstance);
 
@@ -840,9 +840,9 @@ bool FAnimationAnimNextRuntimeTest_IEvaluate::RunTest(const FString& InParameter
 		AUTO_REGISTER_ANIM_TRAIT(FTraitWithOneChild)
 		AUTO_REGISTER_ANIM_TRAIT(FTraitWithChildren)
 
-		UFactory* ModuleFactory = NewObject<UAnimNextModuleFactory>();
-		UAnimNextModule* Module = CastChecked<UAnimNextModule>(ModuleFactory->FactoryCreateNew(UAnimNextModule::StaticClass(), GetTransientPackage(), TEXT("TestAnimNextGraph"), RF_Transient, nullptr, nullptr, NAME_None));
-		UE_RETURN_ON_ERROR(Module != nullptr, "FAnimationAnimNextRuntimeTest_IEvaluate -> Failed to create module");
+		UFactory* GraphFactory = NewObject<UAnimNextAnimationGraphFactory>();
+		UAnimNextAnimationGraph* AnimationGraph = CastChecked<UAnimNextAnimationGraph>(GraphFactory->FactoryCreateNew(UAnimNextAnimationGraph::StaticClass(), GetTransientPackage(), TEXT("TestAnimNextGraph"), RF_Transient, nullptr, nullptr, NAME_None));
+		UE_RETURN_ON_ERROR(AnimationGraph != nullptr, "FAnimationAnimNextRuntimeTest_GraphTraitEvent -> Failed to create animation graph");
 
 		FScopedClearNodeTemplateRegistry ScopedClearNodeTemplateRegistry;
 		FNodeTemplateRegistry& Registry = FNodeTemplateRegistry::Get();
@@ -933,10 +933,10 @@ bool FAnimationAnimNextRuntimeTest_IEvaluate::RunTest(const FString& InParameter
 		}
 
 		// Read our graph
-		FTestUtils::LoadFromArchiveBuffer(*Module, NodeHandles, GraphSharedDataArchiveBuffer);
+		FTestUtils::LoadFromArchiveBuffer(*AnimationGraph, NodeHandles, GraphSharedDataArchiveBuffer);
 
 		FAnimNextGraphInstancePtr GraphInstance;
-		Module->AllocateInstance(GraphInstance);
+		AnimationGraph->AllocateInstance(GraphInstance);
 
 		{
 			TArray<FTraitUID> EvaluatedTraits;
@@ -979,9 +979,9 @@ bool FAnimationAnimNextRuntimeTest_IScopedInterface::RunTest(const FString& InPa
 		AUTO_REGISTER_ANIM_TRAIT(FScopedTagTrait)
 		AUTO_REGISTER_ANIM_TRAIT(FTestScopedTagTrait)
 
-		UFactory* ModuleFactory = NewObject<UAnimNextModuleFactory>();
-		UAnimNextModule* Module = CastChecked<UAnimNextModule>(ModuleFactory->FactoryCreateNew(UAnimNextModule::StaticClass(), GetTransientPackage(), TEXT("TestAnimNextGraph"), RF_Transient, nullptr, nullptr, NAME_None));
-		UE_RETURN_ON_ERROR(Module != nullptr, "FAnimationAnimNextRuntimeTest_IScopedInterface -> Failed to create module");
+		UFactory* GraphFactory = NewObject<UAnimNextAnimationGraphFactory>();
+		UAnimNextAnimationGraph* AnimationGraph = CastChecked<UAnimNextAnimationGraph>(GraphFactory->FactoryCreateNew(UAnimNextAnimationGraph::StaticClass(), GetTransientPackage(), TEXT("TestAnimNextGraph"), RF_Transient, nullptr, nullptr, NAME_None));
+		UE_RETURN_ON_ERROR(AnimationGraph != nullptr, "FAnimationAnimNextRuntimeTest_GraphTraitEvent -> Failed to create animation graph");
 
 		FScopedClearNodeTemplateRegistry ScopedClearNodeTemplateRegistry;
 		FNodeTemplateRegistry& Registry = FNodeTemplateRegistry::Get();
@@ -1070,10 +1070,10 @@ bool FAnimationAnimNextRuntimeTest_IScopedInterface::RunTest(const FString& InPa
 		}
 
 		// Read our graph
-		FTestUtils::LoadFromArchiveBuffer(*Module, NodeHandles, GraphSharedDataArchiveBuffer);
+		FTestUtils::LoadFromArchiveBuffer(*AnimationGraph, NodeHandles, GraphSharedDataArchiveBuffer);
 
 		FAnimNextGraphInstancePtr GraphInstance;
-		Module->AllocateInstance(GraphInstance);
+		AnimationGraph->AllocateInstance(GraphInstance);
 
 		FExecutionContext Context(GraphInstance);
 
diff --git a/Engine/Plugins/Experimental/AnimNextAnimGraph/Source/AnimNextAnimGraphUncookedOnly/Private/Graph/AnimGraph/AnimGraphNode_AnimNextGraph.cpp b/Engine/Plugins/Experimental/AnimNextAnimGraph/Source/AnimNextAnimGraphUncookedOnly/Private/Graph/AnimGraph/AnimGraphNode_AnimNextGraph.cpp
--- a/Engine/Plugins/Experimental/AnimNextAnimGraph/Source/AnimNextAnimGraphUncookedOnly/Private/Graph/AnimGraph/AnimGraphNode_AnimNextGraph.cpp
+++ b/Engine/Plugins/Experimental/AnimNextAnimGraph/Source/AnimNextAnimGraphUncookedOnly/Private/Graph/AnimGraph/AnimGraphNode_AnimNextGraph.cpp
@@ -10,7 +10,6 @@
 #include "Kismet2/BlueprintEditorUtils.h"
 #include "RigVMBlueprintGeneratedClass.h"
 #include "ControlRigBlueprint.h"
-#include "Graph/AnimGraph/AnimBlueprintExtension_AnimNextParameters.h"
 
 #include UE_INLINE_GENERATED_CPP_BY_NAME(AnimGraphNode_AnimNextGraph)
 
@@ -23,14 +22,12 @@ UAnimGraphNode_AnimNextGraph::UAnimGraphNode_AnimNextGraph(const FObjectInitiali
 
 FText UAnimGraphNode_AnimNextGraph::GetNodeTitle(ENodeTitleType::Type TitleType) const
 {
-	// display control rig here
-	return LOCTEXT("AnimGraphNode_AnimNextModule_Title", "AnimNext Module");
+	return LOCTEXT("AnimGraphNode_AnimNextModule_Title", "AnimNext Animation Graph");
 }
 
 FText UAnimGraphNode_AnimNextGraph::GetTooltipText() const
 {
-	// display control rig here
-	return LOCTEXT("AnimGraphNode_AnimNextModule_Tooltip", "Evaluates a AnimNext Module");
+	return LOCTEXT("AnimGraphNode_AnimNextAnimationGraph_Tooltip", "Evaluates a AnimNext Animation Graph");
 }
 
 FText UAnimGraphNode_AnimNextGraph::GetMenuCategory() const
@@ -42,16 +39,15 @@ void UAnimGraphNode_AnimNextGraph::PreloadRequiredAssets()
 {
 	Super::PreloadRequiredAssets();
 
-	if (Node.Module)
+	if (Node.AnimationGraph)
 	{
-		PreloadObject(Node.Module);
-		PreloadObject(Node.Module->EditorData);
+		PreloadObject(Node.AnimationGraph);
+		PreloadObject(Node.AnimationGraph->EditorData);
 	}
 }
 
 void UAnimGraphNode_AnimNextGraph::GetRequiredExtensions(TArray<TSubclassOf<UAnimBlueprintExtension>>& OutExtensions) const
 {
-	OutExtensions.Add(UAnimBlueprintExtension_AnimNextParameters::StaticClass());
 }
 
 void UAnimGraphNode_AnimNextGraph::CreateCustomPins(TArray<UEdGraphPin*>* OldPins)
@@ -152,25 +148,6 @@ void UAnimGraphNode_AnimNextGraph::GetVariables(bool bInput, TMap<FName, FRigVME
 
 	OutVariables.Reset();
 
-	if (URigVMBlueprintGeneratedClass* TargetClass = Cast<URigVMBlueprintGeneratedClass>(GetTargetClass()))
-	{
-		if (UControlRigBlueprint* RigBlueprint = Cast<UControlRigBlueprint>(TargetClass->ClassGeneratedBy))
-		{
-			//RigBlueprint->CleanupVariables();
-			UControlRig* ControlRig = TargetClass->GetDefaultObject<UControlRig>();
-			if (ControlRig)
-			{
-				const TArray<FRigVMExternalVariable>& PublicVariables = ControlRig->GetPublicVariables();
-				for (const FRigVMExternalVariable& PublicVariable : PublicVariables)
-				{
-					if (!bInput || !PublicVariable.bIsReadOnly)
-					{
-						OutVariables.Add(PublicVariable.Name, PublicVariable);
-					}
-				}
-			}
-		}
-	}
 }
 
 void UAnimGraphNode_AnimNextGraph::PostEditChangeProperty(FPropertyChangedEvent& PropertyChangedEvent)
@@ -184,7 +161,7 @@ void UAnimGraphNode_AnimNextGraph::PostEditChangeProperty(FPropertyChangedEvent&
 
 	if (ChangedProperty)
 	{
-		if (ChangedProperty->GetFName() == GET_MEMBER_NAME_CHECKED(FAnimNode_AnimNextGraph, Module))
+		if (ChangedProperty->GetFName() == GET_MEMBER_NAME_CHECKED(FAnimNode_AnimNextGraph, AnimationGraph))
 		{
 			bRequiresNodeReconstruct = true;
 			RebuildExposedProperties();
@@ -210,7 +187,7 @@ void UAnimGraphNode_AnimNextGraph::CustomizePinData(UEdGraphPin* Pin, FName Sour
 
 UObject* UAnimGraphNode_AnimNextGraph::GetJumpTargetForDoubleClick() const
 {
-	return Node.Module;
+	return Node.AnimationGraph;
 }
 
 #undef LOCTEXT_NAMESPACE
diff --git a/Engine/Plugins/Experimental/AnimNextChooser/Source/AnimNextChooser/Private/ChooserParameters.cpp b/Engine/Plugins/Experimental/AnimNextChooser/Source/AnimNextChooser/Private/ChooserParameters.cpp
--- a/Engine/Plugins/Experimental/AnimNextChooser/Source/AnimNextChooser/Private/ChooserParameters.cpp
+++ b/Engine/Plugins/Experimental/AnimNextChooser/Source/AnimNextChooser/Private/ChooserParameters.cpp
@@ -1,6 +1,6 @@
 // Copyright Epic Games, Inc. All Rights Reserved.
 #include "ChooserParameters.h"
-#include "Param/ParamStack.h"
+//#include "Param/ParamStack.h"
 
 #if WITH_EDITOR
 #include "StructUtils/PropertyBag.h"
@@ -8,10 +8,11 @@
 
 bool FBoolAnimProperty::GetValue(FChooserEvaluationContext& Context, bool& OutResult) const
 {
-	UE::AnimNext::FParamResult Result;
-	UE::AnimNext::FParamStack& ParamStack = UE::AnimNext::FParamStack::Get();
-	OutResult = ParamStack.GetParamValue<bool>(UE::AnimNext::FParamId(PropertyName), &Result);
-	return Result.IsSuccessful();
+//	UE::AnimNext::FParamResult Result;
+//	UE::AnimNext::FParamStack& ParamStack = UE::AnimNext::FParamStack::Get();
+//	OutResult = ParamStack.GetParamValue<bool>(UE::AnimNext::FParamId(PropertyName), &Result);
+//	return Result.IsSuccessful();
+	return false;
 }
 
 bool FBoolAnimProperty::SetValue(FChooserEvaluationContext& Context, bool InValue) const
@@ -35,10 +36,11 @@ void FBoolAnimProperty::GetDisplayName(FText& OutName) const
 
 bool FFloatAnimProperty::GetValue(FChooserEvaluationContext& Context, double& OutResult) const
 {
-	UE::AnimNext::FParamResult Result;
-	UE::AnimNext::FParamStack& ParamStack = UE::AnimNext::FParamStack::Get();
-	OutResult = ParamStack.GetParamValue<double>(UE::AnimNext::FParamId(PropertyName), &Result);
-	return Result.IsSuccessful();
+//	UE::AnimNext::FParamResult Result;
+//	UE::AnimNext::FParamStack& ParamStack = UE::AnimNext::FParamStack::Get();
+//	OutResult = ParamStack.GetParamValue<double>(UE::AnimNext::FParamId(PropertyName), &Result);
+//	return Result.IsSuccessful();
+	return false;
 }
 
 bool FFloatAnimProperty::SetValue(FChooserEvaluationContext& Context, double InValue) const
diff --git a/Engine/Plugins/Experimental/Workspace/Source/WorkspaceEditor/Private/AssetDocumentSummoner.cpp b/Engine/Plugins/Experimental/Workspace/Source/WorkspaceEditor/Private/AssetDocumentSummoner.cpp
--- a/Engine/Plugins/Experimental/Workspace/Source/WorkspaceEditor/Private/AssetDocumentSummoner.cpp
+++ b/Engine/Plugins/Experimental/Workspace/Source/WorkspaceEditor/Private/AssetDocumentSummoner.cpp
@@ -54,9 +54,9 @@ void FAssetDocumentSummoner::SaveState(TSharedPtr<SDockTab> Tab, TSharedPtr<FTab
 		UObject* Object = Payload->IsValid() ? FTabPayload_UObject::CastChecked<UObject>(Payload) : nullptr;
 		if(Object)
 		{
-			TSharedRef<SWorkspaceTabWrapper> TabWrapper = StaticCastSharedRef<SWorkspaceTabWrapper>(Tab->GetContent());			
+			TSharedRef<SWorkspaceTabWrapper> TabWrapper = StaticCastSharedRef<SWorkspaceTabWrapper>(Tab->GetContent());
 			FWorkspaceEditorModule& WorkspaceEditorModule = FModuleManager::LoadModuleChecked<FWorkspaceEditorModule>("WorkspaceEditor");
-			const FObjectDocumentArgs* DocumentArgs = WorkspaceEditorModule.FindObjectDocumentType(Object->GetClass()->GetClassPathName());
+			const FObjectDocumentArgs* DocumentArgs = WorkspaceEditorModule.FindObjectDocumentType(Object);
 			if(DocumentArgs && DocumentArgs->OnGetDocumentState.IsBound())
 			{
 				WorkspaceEditor->RecordDocumentState(DocumentArgs->OnGetDocumentState.Execute(FWorkspaceEditorContext(WorkspaceEditor.ToSharedRef(), Object), TabWrapper->GetContent()));
@@ -78,7 +78,7 @@ TAttribute<FText> FAssetDocumentSummoner::ConstructTabNameForObject(UObject* Doc
 	}
 
 	FWorkspaceEditorModule& WorkspaceEditorModule = FModuleManager::LoadModuleChecked<FWorkspaceEditorModule>("WorkspaceEditor");
-	if(const FObjectDocumentArgs* DocumentArgs = WorkspaceEditorModule.FindObjectDocumentType(DocumentID->GetClass()->GetClassPathName()))
+	if(const FObjectDocumentArgs* DocumentArgs = WorkspaceEditorModule.FindObjectDocumentType(DocumentID))
 	{
 		if(DocumentArgs->OnGetTabName.IsBound())
 		{
@@ -103,7 +103,7 @@ bool FAssetDocumentSummoner::IsPayloadSupported(TSharedRef<FTabPayload> Payload)
 	if(Object)
 	{
 		FWorkspaceEditorModule& WorkspaceEditorModule = FModuleManager::LoadModuleChecked<FWorkspaceEditorModule>("WorkspaceEditor");
-		if(const FObjectDocumentArgs* DocumentArgs = WorkspaceEditorModule.FindObjectDocumentType(Object->GetClass()->GetClassPathName()))
+		if(const FObjectDocumentArgs* DocumentArgs = WorkspaceEditorModule.FindObjectDocumentType(Object))
 		{
 			return AllowedClassPaths.Contains(Object->GetClass()->GetClassPathName()) || bAllowUnsupportedClasses;
 		}
@@ -142,7 +142,7 @@ TSharedRef<SWidget> FAssetDocumentSummoner::CreateTabBodyForObject(const FWorkfl
 	check(WorkspaceEditor.IsValid());
 	
 	const FWorkspaceEditorModule& WorkspaceEditorModule = FModuleManager::LoadModuleChecked<FWorkspaceEditorModule>("WorkspaceEditor");
-	if(const FObjectDocumentArgs* DocumentArgs = WorkspaceEditorModule.FindObjectDocumentType(DocumentID->GetClass()->GetClassPathName()))
+	if(const FObjectDocumentArgs* DocumentArgs = WorkspaceEditorModule.FindObjectDocumentType(DocumentID))
 	{
 		if(DocumentArgs->OnMakeDocumentWidget.IsBound())
 		{
@@ -161,7 +161,7 @@ const FSlateBrush* FAssetDocumentSummoner::GetTabIconForObject(const FWorkflowTa
 	if(TSharedPtr<FWorkspaceEditor> WorkspaceEditor = HostingAppPtr.Pin())
 	{
 		FWorkspaceEditorModule& WorkspaceEditorModule = FModuleManager::LoadModuleChecked<FWorkspaceEditorModule>("WorkspaceEditor");
-		if(const FObjectDocumentArgs* DocumentArgs = WorkspaceEditorModule.FindObjectDocumentType(DocumentID->GetClass()->GetClassPathName()))
+		if(const FObjectDocumentArgs* DocumentArgs = WorkspaceEditorModule.FindObjectDocumentType(DocumentID))
 		{
 			if(DocumentArgs->OnGetTabIcon.IsBound())
 			{
diff --git a/Engine/Plugins/Experimental/Workspace/Source/WorkspaceEditor/Private/Outliner/WorkspaceOutlinerMode.cpp b/Engine/Plugins/Experimental/Workspace/Source/WorkspaceEditor/Private/Outliner/WorkspaceOutlinerMode.cpp
--- a/Engine/Plugins/Experimental/Workspace/Source/WorkspaceEditor/Private/Outliner/WorkspaceOutlinerMode.cpp
+++ b/Engine/Plugins/Experimental/Workspace/Source/WorkspaceEditor/Private/Outliner/WorkspaceOutlinerMode.cpp
@@ -14,6 +14,8 @@
 #include "AssetRegistry/AssetRegistryModule.h"
 #include "ScopedTransaction.h"
 #include "ToolMenus.h"
+#include "WorkspaceEditorCommands.h"
+#include "Framework/Commands/GenericCommands.h"
 
 #define LOCTEXT_NAMESPACE "FWorkspaceOutlinerMode"
 
@@ -21,6 +23,8 @@ namespace UE::Workspace
 {
 FWorkspaceOutlinerMode::FWorkspaceOutlinerMode(SSceneOutliner* InSceneOutliner, const TWeakObjectPtr<UWorkspace>& InWeakWorkspace, const TWeakPtr<IWorkspaceEditor>& InWeakWorkspaceEditor) : ISceneOutlinerMode(InSceneOutliner), WeakWorkspace(InWeakWorkspace), WeakWorkspaceEditor(InWeakWorkspaceEditor)
 {
+	CommandList = MakeShared<FUICommandList>();
+	
 	if (UWorkspace* Workspace = WeakWorkspace.Get())
 	{
 		Workspace->ModifiedDelegate.AddRaw(this, &FWorkspaceOutlinerMode::OnWorkspaceModified);
@@ -66,71 +70,20 @@ TSharedPtr<SWidget> FWorkspaceOutlinerMode::CreateContextMenu()
 				const UWorkspaceItemMenuContext* MenuContext = InMenu->FindContext<UWorkspaceItemMenuContext>();
 				if(EditorContext && MenuContext)
 				{
+					FToolMenuSection& CommonSection = InMenu->AddSection("Common", LOCTEXT("CommonSectionLabel", "Common"));
+					CommonSection.AddMenuEntryWithCommandList(FWorkspaceAssetEditorCommands::Get().Open, MenuContext->WeakCommandList.Pin());
+					CommonSection.AddMenuEntryWithCommandList(FGenericCommands::Get().Delete, MenuContext->WeakCommandList.Pin());
+					CommonSection.AddMenuEntryWithCommandList(FGenericCommands::Get().Rename, MenuContext->WeakCommandList.Pin());
+					// TODO JDB more possible actions? (find in content browser?)
+
 					const bool bSelectionContainsTopLevelAsset = MenuContext->SelectedExports.Num() && MenuContext->SelectedExports.ContainsByPredicate([](const FWorkspaceOutlinerItemExport& Export)
 					{
 						return Export.GetParentIdentifier() == NAME_None;
 					});
-
+					
 					FToolMenuSection& AssetsSection = InMenu->AddSection("Assets", LOCTEXT("AssetSectionLabel", "Assets"));
 					if (bSelectionContainsTopLevelAsset)
 					{
-						AssetsSection.AddMenuEntry(TEXT("OpenAsset"),
-							FText::FormatOrdered(LOCTEXT("OpenAssetLabel", "Open {0}|plural(one=Asset,other=Assets)"), MenuContext->SelectedExports.Num()),
-							FText::FormatOrdered(LOCTEXT("OpenAssetTooltip", "Opens the selected {0}|plural(one=Asset,other=Assets)"), MenuContext->SelectedExports.Num()),
-							FSlateIcon(FAppStyle::Get().GetStyleSetName(), "SystemWideCommands.SummonOpenAssetDialog"),
-							FUIAction(FExecuteAction::CreateLambda([WeakEditor=EditorContext->Toolkit, SelectedExports=MenuContext->SelectedExports]()
-							{
-								if (const TSharedPtr<UE::Workspace::IWorkspaceEditor> SharedWorkspaceEditor = StaticCastSharedPtr<UE::Workspace::IWorkspaceEditor>(WeakEditor.Pin()))
-								{
-									TSet<FSoftObjectPath> AssetPaths;
-									for (const FWorkspaceOutlinerItemExport& ItemExport : SelectedExports)
-									{
-										if (ItemExport.GetParentIdentifier() == NAME_None)
-										{
-											AssetPaths.Add(ItemExport.GetAssetPath());	
-										}
-									}
-
-									for (const FSoftObjectPath& AssetPath : AssetPaths)
-									{
-										SharedWorkspaceEditor->OpenAssets({AssetPath.TryLoad()});	
-									}
-								}
-							}))
-						);
-
-						AssetsSection.AddMenuEntry(TEXT("RemoveAsset"),
-							FText::FormatOrdered(LOCTEXT("RemoveAssetLabel", "Remove {0}|plural(one=Asset,other=Assets)"), MenuContext->SelectedExports.Num()),
-							FText::FormatOrdered(LOCTEXT("RemoveAssetTooltip", "Removes the selected {0}|plural(one=Asset,other=Assets) from the Workspace"), MenuContext->SelectedExports.Num()),
-							FSlateIcon(FAppStyle::Get().GetStyleSetName(), "Icons.Delete"),
-							FUIAction(FExecuteAction::CreateLambda([WeakEditor=EditorContext->Toolkit, SelectedExports=MenuContext->SelectedExports, EditingObjects=EditorContext->GetEditingObjects()]()
-							{
-								if (const TSharedPtr<UE::Workspace::IWorkspaceEditor> SharedWorkspaceEditor = StaticCastSharedPtr<UE::Workspace::IWorkspaceEditor>(WeakEditor.Pin()))
-								{
-									TSet<FSoftObjectPath> AssetPaths;
-									for (const FWorkspaceOutlinerItemExport& ItemExport : SelectedExports)
-									{
-										if (ItemExport.GetParentIdentifier() == NAME_None)
-										{
-											AssetPaths.Add(ItemExport.GetAssetPath());	
-										}
-									}
-
-									if (EditingObjects.Num() > 0)
-									{
-										if (UWorkspace* Workspace = Cast<UWorkspace>(EditingObjects[0]))
-										{
-											FScopedTransaction Transaction(LOCTEXT("RemoveAssets", "Remove assets from workspace"));
-											for (const FSoftObjectPath& AssetPath : AssetPaths)
-											{
-												Workspace->RemoveAsset(AssetPath.TryLoad());
-											}
-										}
-									}
-								}
-							}))
-						);
-
 						AssetsSection.AddMenuEntry(TEXT("BrowseToAsset"),
 							LOCTEXT("BrowseToAssetLabel", "Browse to Asset"),
 							LOCTEXT("BrowseToAssetTooltip", "Browse to the selected assets in the content browser"),
@@ -241,7 +194,8 @@ TSharedPtr<SWidget> FWorkspaceOutlinerMode::CreateContextMenu()
 
 	{
 		UWorkspaceItemMenuContext* MenuContext = NewObject<UWorkspaceItemMenuContext>();
-
+		MenuContext->WeakCommandList = CommandList;
+		
 		Algo::TransformIf(SceneOutliner->GetSelectedItems(), MenuContext->SelectedExports,
 			[](const FSceneOutlinerTreeItemPtr& InItem)
 			{
@@ -274,41 +228,43 @@ void FWorkspaceOutlinerMode::OnItemClicked(FSceneOutlinerTreeItemPtr Item)
 
 FReply FWorkspaceOutlinerMode::OnKeyDown(const FKeyEvent& InKeyEvent)
 {
-	// TODO JDB these could be in a FUICommandList
-	if (InKeyEvent.GetKey() == EKeys::Enter)
+	if (CommandList->ProcessCommandBindings(InKeyEvent))
 	{
-		OpenItems(SceneOutliner->GetSelectedItems());
-		return FReply::Handled();
-	}
-	else if (InKeyEvent.GetKey() == EKeys::Platform_Delete)
-	{
-		DeleteItems(SceneOutliner->GetSelectedItems());
 		return FReply::Handled();
 	}
 
-	// TODO JDB more possible actions? (find in content browser?)
-	
 	return ISceneOutlinerMode::OnKeyDown(InKeyEvent);
 }
 
 void FWorkspaceOutlinerMode::HandleItemSelection(const FSceneOutlinerItemSelection& Selection)
 {
-	if (Selection.Num() == 1)
+	if (Selection.Num() > 0)
 	{
 		TArray<FSceneOutlinerTreeItemPtr> SelectedItems;
-		Selection.Get(SelectedItems);			
+		Selection.Get(SelectedItems);
 		
 		if (TSharedPtr<IWorkspaceEditor> SharedWorkspaceEditor = WeakWorkspaceEditor.Pin())
 		{
-			if(FWorkspaceOutlinerTreeItem* TreeItem = SelectedItems[0]->CastTo<FWorkspaceOutlinerTreeItem>())
+			TArray<FWorkspaceOutlinerItemExport> Exports;
+			Exports.Reserve(SelectedItems.Num());
+			for(const FSceneOutlinerTreeItemPtr& SelectedItem : SelectedItems)
 			{
-				if(TreeItem->Export.GetData().IsValid())
+				if(FWorkspaceOutlinerTreeItem* TreeItem = SelectedItem->CastTo<FWorkspaceOutlinerTreeItem>())
 				{
-					TSharedPtr<FStructOnScope> ExportDataView = MakeShared<FStructOnScope>(TreeItem->Export.GetData().GetScriptStruct(), TreeItem->Export.GetData().GetMutableMemory());
-					// TODO JDB handle struct selections
-					//SharedWorkspaceEditor->SetDetailsStruct(ExportDataView);
+					if(TreeItem->Export.GetData().IsValid())
+					{
+						Exports.Add(TreeItem->Export);
+					}
 				}
 			}
+			SharedWorkspaceEditor->OnOutlinerSelectionChanged().Broadcast(Exports);
+		}
+	}
+	else
+	{
+		if (TSharedPtr<IWorkspaceEditor> SharedWorkspaceEditor = WeakWorkspaceEditor.Pin())
+		{
+			SharedWorkspaceEditor->OnOutlinerSelectionChanged().Broadcast(TConstArrayView<FWorkspaceOutlinerItemExport>());
 		}
 	}
 }
@@ -323,6 +279,55 @@ void FWorkspaceOutlinerMode::OnItemSelectionChanged(FSceneOutlinerTreeItemPtr It
 	}
 }
 
+bool FWorkspaceOutlinerMode::CanDeleteItem(const ISceneOutlinerTreeItem& Item) const
+{
+	if (const FWorkspaceOutlinerTreeItem* TreeItem = Item.CastTo<FWorkspaceOutlinerTreeItem>())
+	{
+		if(TreeItem->ItemDetails.IsValid())
+		{
+			return TreeItem->ItemDetails->CanDelete(TreeItem->Export);
+		}
+		else if(TreeItem->Export.GetParentIdentifier() == NAME_None)
+		{
+			// No details, but allow deletion if we are a root item
+			return true;
+		}
+	}
+	return false;
+}
+
+bool FWorkspaceOutlinerMode::CanRenameItem(const ISceneOutlinerTreeItem& Item) const
+{
+	if (const FWorkspaceOutlinerTreeItem* TreeItem = Item.CastTo<FWorkspaceOutlinerTreeItem>())
+	{
+		if(TreeItem->ItemDetails.IsValid())
+		{
+			return TreeItem->ItemDetails->CanRename(TreeItem->Export);
+		}
+	}
+	return false;
+}
+
+void FWorkspaceOutlinerMode::BindCommands(const TSharedRef<FUICommandList>& OutCommandList)
+{
+	CommandList->MapAction( 
+		FWorkspaceAssetEditorCommands::Get().Open, 
+		FExecuteAction::CreateRaw(this, &FWorkspaceOutlinerMode::Open)
+	);
+	
+	CommandList->MapAction( 
+		FGenericCommands::Get().Delete, 
+		FExecuteAction::CreateRaw(this, &FWorkspaceOutlinerMode::Delete),
+		FCanExecuteAction::CreateRaw(this, &FWorkspaceOutlinerMode::CanDelete)
+	);
+	
+	CommandList->MapAction( 
+		FGenericCommands::Get().Rename, 
+		FExecuteAction::CreateRaw(this, &FWorkspaceOutlinerMode::Rename),
+		FCanExecuteAction::CreateRaw(this, &FWorkspaceOutlinerMode::CanRename)
+	);
+}
+
 TUniquePtr<ISceneOutlinerHierarchy> FWorkspaceOutlinerMode::CreateHierarchy()
 {
 	return MakeUnique<FWorkspaceOutlinerHierarchy>(this, WeakWorkspace);
@@ -348,7 +353,7 @@ void FWorkspaceOutlinerMode::OpenItems(TArrayView<const FSceneOutlinerTreeItemPt
 			if(TreeItem->Export.GetParentIdentifier() == NAME_None)
 			{
 				if (const TSharedPtr<UE::Workspace::IWorkspaceEditor> SharedWorkspaceEditor = StaticCastSharedPtr<UE::Workspace::IWorkspaceEditor>(WeakWorkspaceEditor.Pin()))
-				{			
+				{
 					SharedWorkspaceEditor->OpenAssets({TreeItem->Export.GetAssetPath().TryLoad()});
 				}
 			}
@@ -369,31 +374,128 @@ void FWorkspaceOutlinerMode::OpenItems(TArrayView<const FSceneOutlinerTreeItemPt
 	}
 }
 
-void FWorkspaceOutlinerMode::DeleteItems(TArrayView<const FSceneOutlinerTreeItemPtr> Items) const
+void FWorkspaceOutlinerMode::Open()
 {
-	FARFilter Filter;	
-	for (const FSceneOutlinerTreeItemPtr& Item : Items)
+	OpenItems(SceneOutliner->GetSelectedItems());
+}
+
+void FWorkspaceOutlinerMode::Delete()
+{
+	const FSceneOutlinerItemSelection& Selection = SceneOutliner->GetSelection();
+	if (Selection.Num() == 0)
 	{
-		if (const FWorkspaceOutlinerTreeItem* TreeItem = Item->CastTo<FWorkspaceOutlinerTreeItem>())
+		return;
+	}
+
+	UWorkspace* Workspace = WeakWorkspace.Get();
+	if (Workspace == nullptr)
+	{
+		return;
+	}
+
+	// Parse out items into differing types and forward to details
+	TMap<TSharedPtr<IWorkspaceOutlinerItemDetails>, TArray<FWorkspaceOutlinerItemExport>> DetailsMap;
+	
+	FARFilter Filter;
+	for (TWeakPtr<ISceneOutlinerTreeItem>& SelectedItem : Selection.SelectedItems)
+	{
+		FSceneOutlinerTreeItemPtr Item = SelectedItem.Pin();
+		if(!Item.IsValid())
 		{
-			if(TreeItem->Export.GetParentIdentifier() == NAME_None)
+			continue;
+		}
+
+		const FWorkspaceOutlinerTreeItem* TreeItem = Item->CastTo<FWorkspaceOutlinerTreeItem>();
+		if (TreeItem == nullptr)
+		{
+			continue;
+		}
+
+		if(TreeItem->Export.GetParentIdentifier() == NAME_None)
+		{
+			Filter.SoftObjectPaths.AddUnique(TreeItem->Export.GetAssetPath());
+		}
+		else
+		{
+			const TSharedPtr<IWorkspaceOutlinerItemDetails> SharedDetails = FWorkspaceEditorModule::GetOutlinerItemDetails(MakeOutlinerDetailsId(TreeItem->Export));
+			if(!SharedDetails.IsValid())
 			{
-				Filter.SoftObjectPaths.AddUnique(TreeItem->Export.GetAssetPath());
+				continue;
 			}
-		}		
+
+			TArray<FWorkspaceOutlinerItemExport>& ItemsToDelete = DetailsMap.FindOrAdd(SharedDetails);
+			ItemsToDelete.Add(TreeItem->Export);
+		}
 	}
 
-	if (UWorkspace* Workspace = WeakWorkspace.Get())
+	TArray<FAssetData> AssetDataEntriesToRemove;
+	const IAssetRegistry& AssetRegistry = FAssetRegistryModule::GetRegistry();
+	AssetRegistry.GetAssets(Filter, AssetDataEntriesToRemove);
+	
+	if(DetailsMap.Num() > 0 || AssetDataEntriesToRemove.Num() > 0)
 	{
-		TArray<FAssetData> AssetDataEntriesToRemove;
-		const IAssetRegistry& AssetRegistry = FAssetRegistryModule::GetRegistry();
-		AssetRegistry.GetAssets(Filter, AssetDataEntriesToRemove);
-		if (AssetDataEntriesToRemove.Num())
+		FScopedTransaction Transaction(LOCTEXT("RemoveItems", "Remove items from workspace"));
+
+		if(AssetDataEntriesToRemove.Num() > 0)
 		{
-			FScopedTransaction Transaction(LOCTEXT("RemoveAssets", "Remove assets from workspace"));
 			Workspace->RemoveAssets(AssetDataEntriesToRemove);
-		}		
+		}
+
+		for(const TPair<TSharedPtr<IWorkspaceOutlinerItemDetails>, TArray<FWorkspaceOutlinerItemExport>>& ItemsToDeletePair : DetailsMap)
+		{
+			ItemsToDeletePair.Key->Delete(ItemsToDeletePair.Value);
+		}
+	}
+}
+
+bool FWorkspaceOutlinerMode::CanDelete() const
+{
+	const FSceneOutlinerItemSelection& Selection = SceneOutliner->GetSelection();
+	if (Selection.Num() == 0)
+	{
+		return false;
+	}
+
+	for(TWeakPtr<ISceneOutlinerTreeItem> SelectedItem : Selection.SelectedItems)
+	{
+		FSceneOutlinerTreeItemPtr ItemToDelete = SelectedItem.Pin();
+		if(!ItemToDelete.IsValid())
+		{
+			continue;
+		}
+
+		if(!CanDeleteItem(*ItemToDelete) || !ItemToDelete->CanInteract())
+		{
+			return false;
+		}
+	}
+	return true;
+}
+
+void FWorkspaceOutlinerMode::Rename()
+{
+	const FSceneOutlinerItemSelection& Selection = SceneOutliner->GetSelection();
+	if (Selection.Num() == 1)
+	{
+		FSceneOutlinerTreeItemPtr ItemToRename = Selection.SelectedItems[0].Pin();
+
+		if (ItemToRename.IsValid() && CanRenameItem(*ItemToRename) && ItemToRename->CanInteract())
+		{
+			SceneOutliner->SetPendingRenameItem(ItemToRename);
+			SceneOutliner->ScrollItemIntoView(ItemToRename);
+		}
+	}
+}
+
+bool FWorkspaceOutlinerMode::CanRename() const
+{
+	const FSceneOutlinerItemSelection& Selection = SceneOutliner->GetSelection();
+	if (Selection.Num() == 1)
+	{
+		FSceneOutlinerTreeItemPtr ItemToRename = Selection.SelectedItems[0].Pin();
+		return ItemToRename.IsValid() && CanRenameItem(*ItemToRename) && ItemToRename->CanInteract();
 	}
+	return false;
 }
 
 void FWorkspaceOutlinerMode::OnAssetRegistryAssetUpdate(const FAssetData& AssetData)
