Add blend stack trait:
-IBlendStack API for pushing a graph
-FBlendStackCoreTrait used as base trait in state machine-like setups where state machine is the additive trait
-FBlendStackTrait for standalone base trait
-FBlendStackRequesterTrait for example additive trait that can push graphs to a blend stack

diff --git a/Engine/Plugins/Experimental/AnimNextAnimGraph/Source/AnimNextAnimGraph/Private/TraitInterfaces/IBlendStack.cpp b/Engine/Plugins/Experimental/AnimNextAnimGraph/Source/AnimNextAnimGraph/Private/TraitInterfaces/IBlendStack.cpp
new file mode 100644
--- /dev/null
+++ b/Engine/Plugins/Experimental/AnimNextAnimGraph/Source/AnimNextAnimGraph/Private/TraitInterfaces/IBlendStack.cpp
@@ -0,0 +1,48 @@
+// Copyright Epic Games, Inc. All Rights Reserved.
+
+#include "TraitInterfaces/IBlendStack.h"
+
+#include "TraitCore/ExecutionContext.h"
+
+namespace UE::AnimNext
+{
+	AUTO_REGISTER_ANIM_TRAIT_INTERFACE(IBlendStack)
+
+#if WITH_EDITOR
+	const FText& IBlendStack::GetDisplayName() const
+	{
+		static FText InterfaceName = NSLOCTEXT("TraitInterfaces", "TraitInterface_IBlendStack_Name", "Sub Graph");
+		return InterfaceName;
+	}
+	const FText& IBlendStack::GetDisplayShortName() const
+	{
+		static FText InterfaceShortName = NSLOCTEXT("TraitInterfaces", "TraitInterface_IBlendStack_ShortName", "SBG");
+		return InterfaceShortName;
+	}
+
+	void IBlendStack::PushGraph(FExecutionContext& Context, const TTraitBinding<IBlendStack>& Binding, const IBlendStack::FGraphRequest& GraphRequest, FAnimNextGraphInstancePtr& OutGraphInstance) const
+	{
+		TTraitBinding<IBlendStack> SuperBinding;
+		if (Binding.GetStackInterfaceSuper(SuperBinding))
+		{
+			SuperBinding.PushGraph(Context, GraphRequest, OutGraphInstance);
+			return;
+		}
+
+		OutGraphInstance.Release();
+	}
+
+	void IBlendStack::GetActiveGraphRequest(FExecutionContext& Context, const TTraitBinding<IBlendStack>& Binding, IBlendStack::FGraphRequest& OutRequest) const
+	{
+		TTraitBinding<IBlendStack> SuperBinding;
+		if (Binding.GetStackInterfaceSuper(SuperBinding))
+		{
+			SuperBinding.GetActiveGraphRequest(Context, OutRequest);
+			return;
+		}
+
+		OutRequest.AnimationGraph = nullptr;
+	}
+
+#endif // WITH_EDITOR
+}
diff --git a/Engine/Plugins/Experimental/AnimNextAnimGraph/Source/AnimNextAnimGraph/Private/Traits/BlendStackTrait.cpp b/Engine/Plugins/Experimental/AnimNextAnimGraph/Source/AnimNextAnimGraph/Private/Traits/BlendStackTrait.cpp
new file mode 100644
--- /dev/null
+++ b/Engine/Plugins/Experimental/AnimNextAnimGraph/Source/AnimNextAnimGraph/Private/Traits/BlendStackTrait.cpp
@@ -0,0 +1,331 @@
+// Copyright Epic Games, Inc. All Rights Reserved.
+
+#include "Traits/BlendStackTrait.h"
+#include "TraitCore/NodeInstance.h"
+#include "TraitCore/ExecutionContext.h"
+
+namespace UE::AnimNext
+{
+	AUTO_REGISTER_ANIM_TRAIT(FBlendStackCoreTrait)
+	AUTO_REGISTER_ANIM_TRAIT(FBlendStackTrait)
+	AUTO_REGISTER_ANIM_TRAIT(FBlendStackRequesterTrait)
+
+	// Trait required interfaces implementation boilerplate
+	#define TRAIT_INTERFACE_ENUMERATOR(GeneratorMacro) \
+		GeneratorMacro(IDiscreteBlend) \
+		GeneratorMacro(IGarbageCollection) \
+		GeneratorMacro(IHierarchy) \
+		GeneratorMacro(ISmoothBlend) \
+		GeneratorMacro(IUpdateTraversal) \
+		GeneratorMacro(IBlendStack) \
+
+	GENERATE_ANIM_TRAIT_IMPLEMENTATION(FBlendStackCoreTrait, TRAIT_INTERFACE_ENUMERATOR, NULL_ANIM_TRAIT_INTERFACE_ENUMERATOR, NULL_ANIM_TRAIT_EVENT_ENUMERATOR)
+	#undef TRAIT_INTERFACE_ENUMERATOR
+
+	// Trait required interfaces implementation boilerplate
+	#define TRAIT_INTERFACE_ENUMERATOR(GeneratorMacro) \
+		GeneratorMacro(IUpdate) \
+
+	GENERATE_ANIM_TRAIT_IMPLEMENTATION(FBlendStackTrait, TRAIT_INTERFACE_ENUMERATOR, NULL_ANIM_TRAIT_INTERFACE_ENUMERATOR, NULL_ANIM_TRAIT_EVENT_ENUMERATOR)
+
+		// Trait required interfaces implementation boilerplate
+#define TRAIT_REQUIRED_INTERFACE_ENUMERATOR(GeneratorMacroRequired) \
+		GeneratorMacroRequired(IBlendStack) \
+
+		GENERATE_ANIM_TRAIT_IMPLEMENTATION(FBlendStackRequesterTrait, TRAIT_INTERFACE_ENUMERATOR, TRAIT_REQUIRED_INTERFACE_ENUMERATOR, NULL_ANIM_TRAIT_EVENT_ENUMERATOR)
+#undef TRAIT_INTERFACE_ENUMERATOR
+#undef TRAIT_REQUIRED_INTERFACE_ENUMERATOR
+
+
+
+	void FBlendStackCoreTrait::FInstanceData::Construct(const FExecutionContext& Context, const FTraitBinding& Binding)
+	{
+		FTrait::FInstanceData::Construct(Context, Binding);
+		IGarbageCollection::RegisterWithGC(Context, Binding);
+	}
+
+	void FBlendStackCoreTrait::FInstanceData::Destruct(const FExecutionContext& Context, const FTraitBinding& Binding)
+	{
+		FTrait::FInstanceData::Destruct(Context, Binding);
+		IGarbageCollection::UnregisterWithGC(Context, Binding);
+	}
+
+	void FBlendStackCoreTrait::FGraphState::Initialize(const IBlendStack::FGraphRequest& GraphRequest)
+	{
+		Request = GraphRequest;
+		State = FBlendStackCoreTrait::EGraphState::Active;
+		bNewlyCreated = true;
+	}
+
+	void FBlendStackCoreTrait::FGraphState::Terminate()
+	{
+		Instance.Release();
+		ChildPtr.Reset();
+		bNewlyCreated = false;
+		State = FBlendStackCoreTrait::EGraphState::Inactive;
+	}
+
+	int32 FBlendStackCoreTrait::FindFreeGraphIndexOrAdd(FInstanceData& InstanceData)
+	{
+		// Find an empty graph we can use
+		const int32 NumGraphs = InstanceData.ChildGraphs.Num();
+		for (int32 ChildIndex = 0; ChildIndex < NumGraphs; ++ChildIndex)
+		{
+			if (InstanceData.ChildGraphs[ChildIndex].State == FBlendStackCoreTrait::EGraphState::Inactive)
+			{
+				// This graph is inactive, we can re-use it
+				return ChildIndex;
+			}
+		}
+
+		// All graphs are in use, add a new one
+		return InstanceData.ChildGraphs.AddDefaulted();
+	}
+
+	uint32 FBlendStackCoreTrait::GetNumChildren(const FExecutionContext& Context, const TTraitBinding<IHierarchy>& Binding) const
+	{
+		const FInstanceData* InstanceData = Binding.GetInstanceData<FInstanceData>();
+		return InstanceData->ChildGraphs.Num();
+	}
+
+	void FBlendStackCoreTrait::GetChildren(const FExecutionContext& Context, const TTraitBinding<IHierarchy>& Binding, FChildrenArray& Children) const
+	{
+		const FInstanceData* InstanceData = Binding.GetInstanceData<FInstanceData>();
+
+		for (const FBlendStackTrait::FGraphState& ChildGraph : InstanceData->ChildGraphs)
+		{
+			// Even if the request is inactive, we queue an empty handle
+			Children.Add(ChildGraph.ChildPtr);
+		}
+	}
+
+	void FBlendStackCoreTrait::QueueChildrenForTraversal(FUpdateTraversalContext& Context, const TTraitBinding<IUpdateTraversal>& Binding, const FTraitUpdateState& TraitState, FUpdateTraversalQueue& TraversalQueue) const
+	{
+		FInstanceData* InstanceData = Binding.GetInstanceData<FInstanceData>();
+		const int32 NumGraphs = InstanceData->ChildGraphs.Num();
+
+		TTraitBinding<IDiscreteBlend> DiscreteBlendTrait;
+		Binding.GetStackInterface(DiscreteBlendTrait);
+
+		for (int32 ChildIndex = 0; ChildIndex < NumGraphs; ++ChildIndex)
+		{
+			FGraphState& Graph = InstanceData->ChildGraphs[ChildIndex];
+			const float BlendWeight = DiscreteBlendTrait.GetBlendWeight(Context, ChildIndex);
+
+			FTraitUpdateState ChildGraphTraitState = TraitState
+				.WithWeight(BlendWeight)
+				.AsBlendingOut(ChildIndex != InstanceData->CurrentlyActiveGraphIndex)
+				.AsNewlyRelevant(Graph.bNewlyCreated);
+			Graph.bNewlyCreated = false;
+
+			TraversalQueue.Push(InstanceData->ChildGraphs[ChildIndex].ChildPtr, ChildGraphTraitState);
+		}
+	}
+
+	int32 FBlendStackCoreTrait::GetBlendDestinationChildIndex(FExecutionContext& Context, const TTraitBinding<IDiscreteBlend>& Binding) const
+	{
+		const FInstanceData* InstanceData = Binding.GetInstanceData<FInstanceData>();
+		return InstanceData->CurrentlyActiveGraphIndex;
+	}
+
+	void FBlendStackCoreTrait::OnBlendTransition(FExecutionContext& Context, const TTraitBinding<IDiscreteBlend>& Binding, int32 OldChildIndex, int32 NewChildIndex) const
+	{
+		TTraitBinding<IDiscreteBlend> DiscreteBlendTrait;
+		Binding.GetStackInterface(DiscreteBlendTrait);
+
+		// We initiate immediately when we transition
+		DiscreteBlendTrait.OnBlendInitiated(Context, NewChildIndex);
+
+		// We terminate immediately when we transition
+		DiscreteBlendTrait.OnBlendTerminated(Context, OldChildIndex);
+	}
+
+	void FBlendStackCoreTrait::OnBlendInitiated(FExecutionContext& Context, const TTraitBinding<IDiscreteBlend>& Binding, int32 ChildIndex) const
+	{
+		FInstanceData* InstanceData = Binding.GetInstanceData<FInstanceData>();
+		if (InstanceData->ChildGraphs.IsValidIndex(ChildIndex))
+		{
+			FGraphState& Graph = InstanceData->ChildGraphs[ChildIndex];
+
+			//@TODO: Remove or implement entry points once we decide if we still need them.
+			static const FName EntryPoint = NAME_None;
+			Graph.Request.AnimationGraph->AllocateInstance(Binding.GetTraitPtr().GetNodeInstance()->GetOwner(), Graph.Instance, EntryPoint);
+			Graph.ChildPtr = Graph.Instance.GetGraphRootPtr();
+		}
+	}
+
+	void FBlendStackCoreTrait::OnBlendTerminated(FExecutionContext& Context, const TTraitBinding<IDiscreteBlend>& Binding, int32 ChildIndex) const
+	{
+		FInstanceData* InstanceData = Binding.GetInstanceData<FInstanceData>();
+		if (InstanceData->ChildGraphs.IsValidIndex(ChildIndex))
+		{
+			FGraphState& Graph = InstanceData->ChildGraphs[ChildIndex];
+
+			// Deallocate our graph
+			Graph.Terminate();
+		}
+	}
+
+	float FBlendStackCoreTrait::GetBlendTime(FExecutionContext& Context, const TTraitBinding<ISmoothBlend>& Binding, int32 ChildIndex) const
+	{
+		const FInstanceData* InstanceData = Binding.GetInstanceData<FInstanceData>();
+		if (InstanceData->ChildGraphs.IsValidIndex(ChildIndex))
+		{
+			const FGraphState& Graph = InstanceData->ChildGraphs[ChildIndex];
+			return Graph.Request.BlendArgs.BlendTime;
+		}
+		else
+		{
+			// Unknown child
+			return 0.0f;
+		}
+	}
+
+	EAlphaBlendOption FBlendStackCoreTrait::GetBlendType(FExecutionContext& Context, const TTraitBinding<ISmoothBlend>& Binding, int32 ChildIndex) const
+	{
+		const FInstanceData* InstanceData = Binding.GetInstanceData<FInstanceData>();
+		if (InstanceData->ChildGraphs.IsValidIndex(ChildIndex))
+		{
+			const FGraphState& Graph = InstanceData->ChildGraphs[ChildIndex];
+			return Graph.Request.BlendArgs.BlendOption;
+		}
+		else
+		{
+			// Unknown child
+			return EAlphaBlendOption::Linear;
+		}
+	}
+
+	UCurveFloat* FBlendStackCoreTrait::GetCustomBlendCurve(FExecutionContext& Context, const TTraitBinding<ISmoothBlend>& Binding, int32 ChildIndex) const
+	{
+		const FInstanceData* InstanceData = Binding.GetInstanceData<FInstanceData>();
+		if (InstanceData->ChildGraphs.IsValidIndex(ChildIndex))
+		{
+			const FGraphState& Graph = InstanceData->ChildGraphs[ChildIndex];
+			return Graph.Request.BlendArgs.CustomCurve;
+		}
+		else
+		{
+			// Unknown child
+			return nullptr;
+		}
+	}
+
+	void FBlendStackCoreTrait::AddReferencedObjects(const FExecutionContext& Context, const TTraitBinding<IGarbageCollection>& Binding, FReferenceCollector& Collector) const
+	{
+		IGarbageCollection::AddReferencedObjects(Context, Binding, Collector);
+
+		FInstanceData* InstanceData = Binding.GetInstanceData<FInstanceData>();
+
+		for (FBlendStackCoreTrait::FGraphState& Graph : InstanceData->ChildGraphs)
+		{
+			Collector.AddReferencedObject(Graph.Request.AnimationGraph);
+		}
+	}
+
+	void FBlendStackCoreTrait::PushGraph(FExecutionContext& Context, const TTraitBinding<IBlendStack>& Binding, const IBlendStack::FGraphRequest& GraphRequest, FAnimNextGraphInstancePtr& OutGraphInstance) const
+	{
+		if (GraphRequest.AnimationGraph == nullptr)
+		{
+			OutGraphInstance.Release();
+			return;
+		}
+
+		//@TODO: Add depth limit and saturation policies
+		
+		FInstanceData* InstanceData = Binding.GetInstanceData<FInstanceData>();
+
+		const int32 OldChildIndex = InstanceData->CurrentlyActiveGraphIndex;
+		const int32 NewChildIndex = FBlendStackCoreTrait::FindFreeGraphIndexOrAdd(*InstanceData);
+		FBlendStackCoreTrait::FGraphState& Graph = InstanceData->ChildGraphs[NewChildIndex];
+		Graph.Initialize(GraphRequest);
+
+		TTraitBinding<IDiscreteBlend> DiscreteBlendTrait;
+		Binding.GetStackInterface(DiscreteBlendTrait);
+		DiscreteBlendTrait.OnBlendTransition(Context, OldChildIndex, NewChildIndex);
+		
+		InstanceData->CurrentlyActiveGraphIndex = NewChildIndex;
+		OutGraphInstance = Graph.Instance;
+	}
+
+	void FBlendStackCoreTrait::GetActiveGraphRequest(FExecutionContext& Context, const TTraitBinding<IBlendStack>& Binding, IBlendStack::FGraphRequest& OutRequest) const
+	{
+		FInstanceData* InstanceData = Binding.GetInstanceData<FInstanceData>();
+		const int32 CurrentlyActiveGraphIndex = InstanceData->CurrentlyActiveGraphIndex;
+		if (CurrentlyActiveGraphIndex != INDEX_NONE)
+		{
+			OutRequest = InstanceData->ChildGraphs[CurrentlyActiveGraphIndex].Request;
+			return;
+		}
+
+		return IBlendStack::GetActiveGraphRequest(Context, Binding, OutRequest);
+	}
+
+	void FBlendStackTrait::PreUpdate(FUpdateTraversalContext& Context, const TTraitBinding<IUpdate>& Binding, const FTraitUpdateState& TraitState) const
+	{
+		FInstanceData* InstanceData = Binding.GetInstanceData<FInstanceData>();
+		const FSharedData* SharedData = Binding.GetSharedData<FSharedData>();
+ 
+		const UAnimNextAnimationGraph* DesiredAnimationGraph = SharedData->GetAnimationGraph(Binding);
+		const int32 CurrentlyActiveGraphIndex = InstanceData->CurrentlyActiveGraphIndex;
+
+		const bool bForceBlend = SharedData->GetbForceBlend(Binding);
+		const bool bIsEmpty = InstanceData->CurrentlyActiveGraphIndex == INDEX_NONE;
+		if (DesiredAnimationGraph && (bForceBlend || bIsEmpty || (DesiredAnimationGraph != InstanceData->ChildGraphs[CurrentlyActiveGraphIndex].Request.AnimationGraph)))
+		{
+			TTraitBinding<IBlendStack> BlendStackTrait;
+			Binding.GetStackInterface(BlendStackTrait);
+
+			FAnimNextGraphInstancePtr NewGraphInstance;
+			IBlendStack::FGraphRequest GraphRequest;
+			GraphRequest.AnimationGraph = DesiredAnimationGraph;
+			GraphRequest.BlendArgs.BlendTime = SharedData->GetBlendTime(Binding);
+
+			BlendStackTrait.PushGraph(Context, GraphRequest, NewGraphInstance);
+		}
+	}
+
+	void FBlendStackRequesterTrait::PreUpdate(FUpdateTraversalContext& Context, const TTraitBinding<IUpdate>& Binding, const FTraitUpdateState& TraitState) const
+	{
+		const FSharedData* SharedData = Binding.GetSharedData<FSharedData>();
+
+		const UAnimNextAnimationGraph* DesiredAnimationGraph = SharedData->GetAnimationGraph(Binding);
+		if (DesiredAnimationGraph != nullptr)
+		{
+			TTraitBinding<IBlendStack> BlendStackTrait;
+			Binding.GetStackInterface(BlendStackTrait);
+
+			IBlendStack::FGraphRequest ActiveGraphRequest;
+			BlendStackTrait.GetActiveGraphRequest(Context, ActiveGraphRequest);
+			const bool bForceBlend = SharedData->GetbForceBlend(Binding);
+			if (bForceBlend || (DesiredAnimationGraph != ActiveGraphRequest.AnimationGraph))
+			{
+				FAnimNextGraphInstancePtr NewGraphInstance;
+				IBlendStack::FGraphRequest GraphRequest;
+				GraphRequest.AnimationGraph = DesiredAnimationGraph;
+				GraphRequest.BlendArgs.BlendTime = SharedData->GetBlendTime(Binding);
+
+				BlendStackTrait.PushGraph(Context, GraphRequest, NewGraphInstance);
+			}
+		}
+	}
+
+	void FBlendStackRequesterTrait::OnBecomeRelevant(FUpdateTraversalContext& Context, const TTraitBinding<IUpdate>& Binding, const FTraitUpdateState& TraitState) const
+	{
+		const FSharedData* SharedData = Binding.GetSharedData<FSharedData>();
+
+		TTraitBinding<IBlendStack> BlendStackTrait;
+		Binding.GetStackInterface(BlendStackTrait);
+
+		const UAnimNextAnimationGraph* DesiredAnimationGraph = SharedData->GetAnimationGraph(Binding);
+		if (DesiredAnimationGraph != nullptr)
+		{
+			FAnimNextGraphInstancePtr NewGraphInstance;
+			IBlendStack::FGraphRequest GraphRequest;
+			GraphRequest.AnimationGraph = DesiredAnimationGraph;
+			GraphRequest.BlendArgs.BlendTime = SharedData->GetBlendTime(Binding);
+
+			BlendStackTrait.PushGraph(Context, GraphRequest, NewGraphInstance);
+		}
+	}
+}
diff --git a/Engine/Plugins/Experimental/AnimNextAnimGraph/Source/AnimNextAnimGraph/Private/Traits/BlendStackTrait.h b/Engine/Plugins/Experimental/AnimNextAnimGraph/Source/AnimNextAnimGraph/Private/Traits/BlendStackTrait.h
new file mode 100644
--- /dev/null
+++ b/Engine/Plugins/Experimental/AnimNextAnimGraph/Source/AnimNextAnimGraph/Private/Traits/BlendStackTrait.h
@@ -0,0 +1,202 @@
+// Copyright Epic Games, Inc. All Rights Reserved.
+
+#pragma once
+
+#include "TraitCore/Trait.h"
+#include "TraitInterfaces/IDiscreteBlend.h"
+#include "TraitInterfaces/IGarbageCollection.h"
+#include "TraitInterfaces/IHierarchy.h"
+#include "TraitInterfaces/IUpdate.h"
+#include "Graph/AnimNextAnimationGraph.h"
+#include "Graph/AnimNextGraphInstancePtr.h"
+#include "TraitInterfaces/IBlendStack.h"
+#include "TraitInterfaces/ISmoothBlend.h"
+
+#include "BlendStackTrait.generated.h"
+
+USTRUCT(meta = (DisplayName = "Blend Stack Core"))
+struct FAnimNextBlendStackCoreTraitSharedData : public FAnimNextTraitSharedData
+{
+	GENERATED_BODY()
+
+	// This struct is empty but required so that we can give a nice display name to the trait
+};
+
+USTRUCT(meta = (DisplayName = "Blend Stack"))
+struct FAnimNextBlendStackTraitSharedData : public FAnimNextTraitSharedData
+{
+	GENERATED_BODY()
+
+	/** Asset to use as a sub-graph */
+	UPROPERTY(EditAnywhere, Category = "Default")
+	TObjectPtr<const UAnimNextAnimationGraph> AnimationGraph;
+
+	/** Blend time to use when sub-graph pin input changes */
+	UPROPERTY(EditAnywhere, Category = "Default")
+	float BlendTime = 0.2f;
+
+	/** Force a blend to the current graph, even if there's been no change in selection. */
+	UPROPERTY(EditAnywhere, Category = "Default")
+	bool bForceBlend = false;
+
+	// Latent pin support boilerplate
+	#define TRAIT_LATENT_PROPERTIES_ENUMERATOR(GeneratorMacro) \
+		GeneratorMacro(AnimationGraph) \
+		GeneratorMacro(BlendTime) \
+		GeneratorMacro(bForceBlend) \
+
+	GENERATE_TRAIT_LATENT_PROPERTIES(FAnimNextBlendStackTraitSharedData, TRAIT_LATENT_PROPERTIES_ENUMERATOR)
+	#undef TRAIT_LATENT_PROPERTIES_ENUMERATOR
+};
+
+USTRUCT(meta = (DisplayName = "Blend Stack Requester"))
+struct FAnimNextBlendStackRequesterTraitSharedData : public FAnimNextTraitSharedData
+{
+	GENERATED_BODY()
+
+	/** Asset to use as a sub-graph */
+	UPROPERTY(EditAnywhere, Category = "Default")
+	TObjectPtr<const UAnimNextAnimationGraph> AnimationGraph;
+
+	/** Blend time to use when sub-graph pin input changes */
+	UPROPERTY(EditAnywhere, Category = "Default")
+	float BlendTime = 0.2f;
+
+	/** Force a blend to the current graph, even if there's been no change in selection. */
+	UPROPERTY(EditAnywhere, Category = "Default")
+	bool bForceBlend = false;
+
+	// Latent pin support boilerplate
+	#define TRAIT_LATENT_PROPERTIES_ENUMERATOR(GeneratorMacro) \
+		GeneratorMacro(AnimationGraph) \
+		GeneratorMacro(BlendTime) \
+		GeneratorMacro(bForceBlend) \
+
+	GENERATE_TRAIT_LATENT_PROPERTIES(FAnimNextBlendStackRequesterTraitSharedData, TRAIT_LATENT_PROPERTIES_ENUMERATOR)
+	#undef TRAIT_LATENT_PROPERTIES_ENUMERATOR
+};
+
+namespace UE::AnimNext
+{
+	/**
+	 * FBlendStackCoreTrait
+	 * 
+	 * Used as base trait in state machine-like setups where state machine is the additive trait that pushes graphs and blend settings to the blend stack.
+	 */
+	struct FBlendStackCoreTrait : FBaseTrait, IUpdateTraversal, IHierarchy, IDiscreteBlend, ISmoothBlend, IGarbageCollection, IBlendStack
+	{
+		DECLARE_ANIM_TRAIT(FBlendStackCoreTrait, 0xda7b7f8d, FBaseTrait)
+
+		enum class EGraphState : uint8
+		{
+			Inactive,
+			Active,
+		};
+
+		struct FGraphState
+		{
+			// The blend settings to use
+			IBlendStack::FGraphRequest Request;
+
+			// The graph instance
+			FAnimNextGraphInstancePtr Instance;
+
+			// Our child handle
+			// Weak handle to the graph instance's root
+			FTraitPtr ChildPtr;
+
+			FBlendStackCoreTrait::EGraphState State;
+
+			// Whether or not this graph was previously relevant
+			bool bNewlyCreated = false;
+
+			void Initialize(const IBlendStack::FGraphRequest& GraphRequest);
+			void Terminate();
+		};
+
+		struct FInstanceData : FTrait::FInstanceData
+		{
+			// List of current graphs in the blend stack
+			TArray<FGraphState> ChildGraphs;
+
+			// The index of the currently active graph
+			// All other graphs are blending out
+			int32 CurrentlyActiveGraphIndex = INDEX_NONE;
+
+			void Construct(const FExecutionContext& Context, const FTraitBinding& Binding);
+			void Destruct(const FExecutionContext& Context, const FTraitBinding& Binding);
+		};
+
+		using FSharedData = FAnimNextBlendStackCoreTraitSharedData;
+
+		// Internal impl
+		static int32 FindFreeGraphIndexOrAdd(FInstanceData& InstanceData);
+
+		// IUpdateTraversal impl
+		virtual void QueueChildrenForTraversal(FUpdateTraversalContext& Context, const TTraitBinding<IUpdateTraversal>& Binding, const FTraitUpdateState& TraitState, FUpdateTraversalQueue& TraversalQueue) const override;
+
+		// IHierarchy impl
+		virtual uint32 GetNumChildren(const FExecutionContext& Context, const TTraitBinding<IHierarchy>& Binding) const override;
+		virtual void GetChildren(const FExecutionContext& Context, const TTraitBinding<IHierarchy>& Binding, FChildrenArray& Children) const override;
+
+		// IDiscreteBlend impl
+		virtual int32 GetBlendDestinationChildIndex(FExecutionContext& Context, const TTraitBinding<IDiscreteBlend>& Binding) const override;
+		virtual void OnBlendTransition(FExecutionContext& Context, const TTraitBinding<IDiscreteBlend>& Binding, int32 OldChildIndex, int32 NewChildIndex) const override;
+		virtual void OnBlendInitiated(FExecutionContext& Context, const TTraitBinding<IDiscreteBlend>& Binding, int32 ChildIndex) const override;
+		virtual void OnBlendTerminated(FExecutionContext& Context, const TTraitBinding<IDiscreteBlend>& Binding, int32 ChildIndex) const override;
+
+		//@TODO: Add inertial blend.
+		// ISmoothBlend impl
+		virtual float GetBlendTime(FExecutionContext& Context, const TTraitBinding<ISmoothBlend>& Binding, int32 ChildIndex) const override;
+		virtual EAlphaBlendOption GetBlendType(FExecutionContext& Context, const TTraitBinding<ISmoothBlend>& Binding, int32 ChildIndex) const override;
+		virtual UCurveFloat* GetCustomBlendCurve(FExecutionContext& Context, const TTraitBinding<ISmoothBlend>& Binding, int32 ChildIndex) const override;
+
+		// IGarbageCollection impl
+		virtual void AddReferencedObjects(const FExecutionContext& Context, const TTraitBinding<IGarbageCollection>& Binding, FReferenceCollector& Collector) const override;
+
+		// IBlendStack impl
+		virtual void PushGraph(FExecutionContext& Context, const TTraitBinding<IBlendStack>& Binding, const IBlendStack::FGraphRequest& GraphRequest, FAnimNextGraphInstancePtr& OutGraphInstance) const;
+		virtual void GetActiveGraphRequest(FExecutionContext& Context, const TTraitBinding<IBlendStack>& Binding, IBlendStack::FGraphRequest& OutRequest) const;
+
+#if WITH_EDITOR
+		virtual bool IsHidden() const override { return true; }
+#endif
+	};
+
+	/**
+	 * FBlendStackTrait
+	 * 
+	 * Standalone blend stack that exposes desired graph and blend time as pins. 
+	 * A new graph will be pushed every time the newest graph does not match the pin (or bForceBlend is true).
+	 */	
+	struct FBlendStackTrait : FBlendStackCoreTrait, IUpdate
+	{
+		DECLARE_ANIM_TRAIT(FBlendStackTrait, 0x46678513, FBlendStackCoreTrait)
+
+		using FSharedData = FAnimNextBlendStackTraitSharedData;
+
+		// IUpdate impl
+		virtual void PreUpdate(FUpdateTraversalContext& Context, const TTraitBinding<IUpdate>& Binding, const FTraitUpdateState& TraitState) const override;
+
+#if WITH_EDITOR
+		virtual bool IsHidden() const override { return false; }
+#endif
+	};
+
+	/**
+	 * FBlendStackRequesterTrait
+	 * 
+	 * Example additive trait used to push graphs to a Blend Stack base trait.
+	 * Same behavior as FBlendStackTrait, but as additive (i.e. FBlendStackCoreTrait base + FBlendStackRequesterTrait as additive).
+	 */	
+	struct FBlendStackRequesterTrait : FAdditiveTrait, IUpdate
+	{
+		DECLARE_ANIM_TRAIT(FBlendStackRequesterTrait, 0x3076bf31, FAdditiveTrait)
+
+		using FSharedData = FAnimNextBlendStackRequesterTraitSharedData;
+
+		// IUpdate impl
+		virtual void PreUpdate(FUpdateTraversalContext& Context, const TTraitBinding<IUpdate>& Binding, const FTraitUpdateState& TraitState) const override;
+		virtual void OnBecomeRelevant(FUpdateTraversalContext& Context, const TTraitBinding<IUpdate>& Binding, const FTraitUpdateState& TraitState) const override;
+	};
+}
diff --git a/Engine/Plugins/Experimental/AnimNextAnimGraph/Source/AnimNextAnimGraph/Public/TraitInterfaces/IBlendStack.h b/Engine/Plugins/Experimental/AnimNextAnimGraph/Source/AnimNextAnimGraph/Public/TraitInterfaces/IBlendStack.h
new file mode 100644
--- /dev/null
+++ b/Engine/Plugins/Experimental/AnimNextAnimGraph/Source/AnimNextAnimGraph/Public/TraitInterfaces/IBlendStack.h
@@ -0,0 +1,69 @@
+// Copyright Epic Games, Inc. All Rights Reserved.
+
+#pragma once
+
+#include "CoreMinimal.h"
+#include "TraitCore/ITraitInterface.h"
+#include "TraitCore/TraitBinding.h"
+#include "AlphaBlend.h"
+#include "Graph/AnimNextAnimationGraph.h"
+#include "Graph/AnimNextGraphInstancePtr.h"
+
+namespace UE::AnimNext
+{
+	/**
+	 * IBlendStack
+	 *
+	 * This interface exposes anything needed to request a new subgraph
+	 */
+	struct ANIMNEXTANIMGRAPH_API IBlendStack : ITraitInterface
+	{
+		DECLARE_ANIM_TRAIT_INTERFACE(IBlendStack, 0xf3468e64)
+
+		struct FGraphRequest
+		{
+			/**
+			 * TODO: Add more blend options here as we need them.
+			 * Consider making dynamic payload if we want to implement a special blend framework.
+			 */
+
+			// Blend-in duration for this graph request 
+			FAlphaBlendArgs BlendArgs;
+
+			// The template graph to use for the new graph instance
+			TObjectPtr<const UAnimNextAnimationGraph> AnimationGraph;
+		};
+
+
+		// Returns the desired blend time for the specified child
+		virtual void PushGraph(FExecutionContext& Context, const TTraitBinding<IBlendStack>& Binding, const IBlendStack::FGraphRequest& GraphRequest, FAnimNextGraphInstancePtr& OutGraphInstance) const;
+		virtual void GetActiveGraphRequest(FExecutionContext& Context, const TTraitBinding<IBlendStack>& Binding, IBlendStack::FGraphRequest& OutRequest) const;
+
+#if WITH_EDITOR
+		virtual const FText& GetDisplayName() const override;
+		virtual const FText& GetDisplayShortName() const override;
+#endif // WITH_EDITOR
+	};
+
+	/**
+	 * Specialization for trait binding.
+	 */
+	template<>
+	struct TTraitBinding<IBlendStack> : FTraitBinding
+	{
+		// @see IBlendStack::PushGraph
+		void PushGraph(FExecutionContext& Context, const IBlendStack::FGraphRequest& GraphRequest, FAnimNextGraphInstancePtr& OutGraphInstance) const
+		{
+			GetInterface()->PushGraph(Context, *this, GraphRequest, OutGraphInstance);
+		}
+
+		// @see IBlendStack::GetActiveGraphRequest
+		void GetActiveGraphRequest(FExecutionContext& Context, IBlendStack::FGraphRequest& OutRequest) const
+		{
+			GetInterface()->GetActiveGraphRequest(Context, *this, OutRequest);
+		}
+
+	protected:
+		const IBlendStack* GetInterface() const { return GetInterfaceTyped<IBlendStack>(); }
+	};
+}
