Fix for build error on non-editor Anim Next

diff --git a/Engine/Plugins/Experimental/AnimNextAnimGraph/Source/AnimNextAnimGraph/Public/TraitInterfaces/IBlendStack.h b/Engine/Plugins/Experimental/AnimNextAnimGraph/Source/AnimNextAnimGraph/Public/TraitInterfaces/IBlendStack.h
--- a/Engine/Plugins/Experimental/AnimNextAnimGraph/Source/AnimNextAnimGraph/Public/TraitInterfaces/IBlendStack.h
+++ b/Engine/Plugins/Experimental/AnimNextAnimGraph/Source/AnimNextAnimGraph/Public/TraitInterfaces/IBlendStack.h
@@ -14,7 +14,7 @@ namespace UE::AnimNext
 	/**
 	 * IBlendStack
 	 *
-	 * This interface exposes anything needed to request a new subgraph
+	 * This interface exposes anything needed to push a new subgraph
 	 */
 	struct ANIMNEXTANIMGRAPH_API IBlendStack : ITraitInterface
 	{
@@ -34,9 +34,10 @@ namespace UE::AnimNext
 			TObjectPtr<const UAnimNextAnimationGraph> AnimationGraph;
 		};
 
-
-		// Returns the desired blend time for the specified child
+		// Pushes a new subgraph along with blend settings defined in GraphRequest. Outputs the in-place created subgraph OutGraphInstance
 		virtual void PushGraph(FExecutionContext& Context, const TTraitBinding<IBlendStack>& Binding, const IBlendStack::FGraphRequest& GraphRequest, FAnimNextGraphInstancePtr& OutGraphInstance) const;
+		
+		// Gets the graph request info from the most recent PushGraph
 		virtual void GetActiveGraphRequest(FExecutionContext& Context, const TTraitBinding<IBlendStack>& Binding, IBlendStack::FGraphRequest& OutRequest) const;
 
 #if WITH_EDITOR
