AnimNext: Data interfaces

diff --git a/Engine/Plugins/Experimental/AnimNext/Source/AnimNextUncookedOnly/Private/UncookedOnlyUtils.cpp b/Engine/Plugins/Experimental/AnimNext/Source/AnimNextUncookedOnly/Private/UncookedOnlyUtils.cpp
--- a/Engine/Plugins/Experimental/AnimNext/Source/AnimNextUncookedOnly/Private/UncookedOnlyUtils.cpp
+++ b/Engine/Plugins/Experimental/AnimNext/Source/AnimNextUncookedOnly/Private/UncookedOnlyUtils.cpp
@@ -26,6 +26,9 @@
 #include "Variables/IVariableBindingType.h"
 #include "Variables/RigUnit_CopyModuleProxyVariables.h"
 #include "WorkspaceAssetRegistryInfo.h"
+#include "Entries/AnimNextDataInterfaceEntry.h"
+#include "DataInterface/AnimNextDataInterface.h"
+#include "DataInterface/AnimNextDataInterface_EditorData.h"
 
 #define LOCTEXT_NAMESPACE "AnimNextUncookedOnlyUtils"
 
@@ -47,47 +50,150 @@ void FUtils::CompileVariables(UAnimNextRigVMAsset* InAsset)
 {
 	check(InAsset);
 
-	UAnimNextRigVMAssetEditorData* EditorData = GetEditorData(InAsset);
+	UAnimNextDataInterface* DataInterface = Cast<UAnimNextDataInterface>(InAsset);
+	if(DataInterface == nullptr)
+	{
+		// Currently only support data interface types (TODO: could make UAnimNextDataInterface the common base rather than UAnimNextRigVMAsset)
+		return;
+	}
+
+	FMessageLog Log("AnimNextCompilerResults");
+
+	UAnimNextDataInterface_EditorData* EditorData = GetEditorData<UAnimNextDataInterface_EditorData>(DataInterface);
 
 	struct FStructEntryInfo
 	{
+		const UAnimNextDataInterface* FromDataInterface;
 		FName Name;
 		FAnimNextParamType Type;
 		EAnimNextExportAccessSpecifier AccessSpecifier;
+		bool bAutoBindDataInterfaceToHost;
 		TConstArrayView<uint8> Value;
 	};
 
+	// Gather all variables in this asset.
+	// Variables are harvested from the valid entries and data interfaces.
+	// Data interface harvesting is performed recursively
+	// The topmost value for a data interface 'wins' if a value is to be supplied
+	TMap<FName, int32> EntryInfoIndexMap;
 	TArray<FStructEntryInfo> StructEntryInfos;
 	StructEntryInfos.Reserve(EditorData->Entries.Num());
-
-	// Gather all parameters in this asset
 	int32 NumPublicVariables = 0;
-	for(const UAnimNextRigVMAssetEntry* Entry : EditorData->Entries)
+
+	auto AddVariable = [&Log, &NumPublicVariables, &StructEntryInfos, &EntryInfoIndexMap](const UAnimNextVariableEntry* InVariable, const UAnimNextDataInterfaceEntry* InFromInterfaceEntry, const UAnimNextDataInterface* InFromInterface, bool bInAutoBindInterface)
 	{
-		const IAnimNextRigVMExportInterface* Export = Cast<IAnimNextRigVMExportInterface>(Entry);
-		const IAnimNextRigVMVariableInterface* Variable = Cast<IAnimNextRigVMVariableInterface>(Entry);
-		if(Export && Variable)
+		const FName Name = InVariable->GetExportName();
+		const FAnimNextParamType& Type = InVariable->GetExportType();
+		if(!Type.IsValid())
+		{
+			Log.Error(FText::Format(LOCTEXT("InvalidVariableTypeFound", "Variable '{0}' with invalid type found"), FText::FromName(Name)));
+			return;
+		}
+
+		const EAnimNextExportAccessSpecifier AccessSpecifier = InVariable->GetExportAccessSpecifier();
+
+		// Check for type conflicts
+		int32* ExistingIndexPtr = EntryInfoIndexMap.Find(Name);
+		if(ExistingIndexPtr)
+		{
+			const FStructEntryInfo& ExistingInfo = StructEntryInfos[*ExistingIndexPtr];
+			if(ExistingInfo.Type != Type)
+			{
+				Log.Error(FText::Format(LOCTEXT("ConflictingVariableTypeFound", "Variable '{0}' with conflicting type found ({1} vs {2})"), FText::FromName(Name), FText::FromString(ExistingInfo.Type.ToString()), FText::FromString(Type.ToString())));
+				return;
+			}
+
+			if(ExistingInfo.AccessSpecifier != AccessSpecifier)
+			{
+				Log.Error(FText::Format(LOCTEXT("ConflictingVariableAccessFound", "Variable '{0}' with conflicting access specifier found ({1} vs {2})"), FText::FromName(Name), FText::FromString(UEnum::GetValueAsString(ExistingInfo.AccessSpecifier)), FText::FromString(UEnum::GetValueAsString(AccessSpecifier))));
+				return;
+			}
+		}
+		else
 		{
-			const FAnimNextParamType& Type = Export->GetExportType();
-			ensure(Type.IsValid());
-			const FName Name = Export->GetExportName();
-			const EAnimNextExportAccessSpecifier AccessSpecifier = Export->GetExportAccessSpecifier();
 			if(AccessSpecifier == EAnimNextExportAccessSpecifier::Public)
 			{
 				NumPublicVariables++;
 			}
+		}
+
+		// Check the overrides to see if this variable's default is overriden
+		const FProperty* OverrideProperty = nullptr;
+		TConstArrayView<uint8> OverrideValue;
+		if(InFromInterfaceEntry != nullptr)
+		{
+			InFromInterfaceEntry->FindValueOverrideRecursive(Name, OverrideProperty, OverrideValue);
+		}
+
+		TConstArrayView<uint8> Value;
+		if(!OverrideValue.IsEmpty())
+		{
+			Value = OverrideValue;
+		}
+		else
+		{
+			Value = TConstArrayView<uint8>(InVariable->GetValuePtr(), Type.GetSize());
+		}
+
+		if(ExistingIndexPtr)
+		{
+			StructEntryInfos[*ExistingIndexPtr].Value = Value;
+		}
+		else
+		{
+			int32 Index = StructEntryInfos.Add(
+			{
+				InFromInterface,
+				Name,
+				FAnimNextParamType(Type.GetValueType(), Type.GetContainerType(), Type.GetValueTypeObject()),
+				AccessSpecifier,
+				bInAutoBindInterface,
+				Value
+			});
+
+			EntryInfoIndexMap.Add(Name, Index);
+		}
+	};
+
+	auto AddDataInterface = [&Log, &AddVariable, &DataInterface](const UAnimNextDataInterface* InDataInterface, const UAnimNextDataInterfaceEntry* InDataInterfaceEntry, bool bInPublicOnly, bool bInAutoBindInterface, auto& InAddDataInterface) -> void
+	{
+		const UAnimNextDataInterface_EditorData* DataInterfaceEditorData = GetEditorData<const UAnimNextDataInterface_EditorData>(InDataInterface);
+		check(DataInterfaceEditorData != nullptr);
 
-			StructEntryInfos.Add(
+		for(UAnimNextRigVMAssetEntry* OtherEntry : DataInterfaceEditorData->Entries)
+		{
+			if(const UAnimNextVariableEntry* VariableEntry = Cast<UAnimNextVariableEntry>(OtherEntry))
+			{
+				if(!bInPublicOnly || VariableEntry->GetExportAccessSpecifier() == EAnimNextExportAccessSpecifier::Public)
+				{
+					AddVariable(VariableEntry, InDataInterfaceEntry, InDataInterface, bInAutoBindInterface);
+				}
+			}
+			else if(const UAnimNextDataInterfaceEntry* DataInterfaceEntry = Cast<UAnimNextDataInterfaceEntry>(OtherEntry))
+			{
+				UAnimNextDataInterface* SubDataInterface = DataInterfaceEntry->GetDataInterface();
+				if(SubDataInterface == nullptr)
+				{
+					Log.Error(FText::Format(LOCTEXT("MissingDataInterfaceWarning", "Invalid data interface found: {0}"), FText::FromString(DataInterfaceEntry->GetDataInterfacePath().ToString())));
+					return;
+				}
+				else if(DataInterface == SubDataInterface)
+				{
+					Log.Error(FText::Format(LOCTEXT("CircularDataInterfaceRefError", "Circular data interface reference found: {0}"), FText::FromString(DataInterfaceEntry->GetDataInterfacePath().ToString())));
+					return;
+				}
+				else
 				{
-					Name,
-					FAnimNextParamType(Type.GetValueType(),  Type.GetContainerType(), Type.GetValueTypeObject()),
-					AccessSpecifier,
-					TConstArrayView<uint8>(Variable->GetValuePtr(), Type.GetSize())
-				});
+					bool bAutoBindInterface = DataInterfaceEntry->AutomaticBinding == EAnimNextDataInterfaceAutomaticBindingMode::BindSharedInterfaces;
+					InAddDataInterface(SubDataInterface, DataInterfaceEntry, true, bAutoBindInterface, InAddDataInterface);
+				}
+			}
 		}
-	}
+	};
+
+	AddDataInterface(DataInterface, nullptr, false, false, AddDataInterface);
 
-	// Sort public entries first & then by size, largest first, for better packing
+	// Sort public entries first, then by data interface & then by size, largest first, for better packing
 	static_assert(EAnimNextExportAccessSpecifier::Private < EAnimNextExportAccessSpecifier::Public, "Private must be less than Public as parameters are sorted internally according to this assumption");
 	StructEntryInfos.Sort([](const FStructEntryInfo& InLHS, const FStructEntryInfo& InRHS)
 	{
@@ -95,10 +201,18 @@ void FUtils::CompileVariables(UAnimNextRigVMAsset* InAsset)
 		{
 			return InLHS.AccessSpecifier > InRHS.AccessSpecifier;
 		}
-		else
+		else if(InLHS.FromDataInterface != InRHS.FromDataInterface)
+		{
+			return InLHS.FromDataInterface->GetFName().LexicalLess(InRHS.FromDataInterface->GetFName());
+		}
+		else if(InLHS.Type.GetSize() != InRHS.Type.GetSize())
 		{
 			return InLHS.Type.GetSize() > InRHS.Type.GetSize();
 		}
+		else
+		{
+			return InLHS.Name.LexicalLess(InRHS.Name);
+		}
 	});
 
 	if(StructEntryInfos.Num() > 0)
@@ -109,40 +223,68 @@ void FUtils::CompileVariables(UAnimNextRigVMAsset* InAsset)
 		TArray<TConstArrayView<uint8>> Values;
 		Values.Reserve(StructEntryInfos.Num());
 
+		DataInterface->ImplementedInterfaces.Empty();
+
 		for (TEnumerateRef<const FStructEntryInfo> StructEntryInfo : EnumerateRange(StructEntryInfos))
 		{
 			PropertyDescs.Emplace(StructEntryInfo->Name, StructEntryInfo->Type.ContainerType, StructEntryInfo->Type.ValueType, StructEntryInfo->Type.ValueTypeObject);
 			Values.Add(StructEntryInfo->Value);
+
+			if(StructEntryInfo->AccessSpecifier != EAnimNextExportAccessSpecifier::Public)
+			{
+				continue;
+			}
+
+			// Now process any data interfaces (sets of public variables) 
+			auto CheckForExistingDataInterface = [&StructEntryInfo](const FAnimNextImplementedDataInterface& InImplementedDataInterface)
+			{
+				return InImplementedDataInterface.DataInterface == StructEntryInfo->FromDataInterface;
+			};
+
+			FAnimNextImplementedDataInterface* ExistingImplementedDataInterface = DataInterface->ImplementedInterfaces.FindByPredicate(CheckForExistingDataInterface);
+			if(ExistingImplementedDataInterface == nullptr)
+			{
+				FAnimNextImplementedDataInterface& NewImplementedDataInterface = DataInterface->ImplementedInterfaces.AddDefaulted_GetRef();
+				NewImplementedDataInterface.DataInterface = StructEntryInfo->FromDataInterface;
+				NewImplementedDataInterface.VariableIndex = StructEntryInfo.GetIndex();
+				NewImplementedDataInterface.NumVariables = 1;
+				NewImplementedDataInterface.bAutoBindToHost = StructEntryInfo->bAutoBindDataInterfaceToHost; 
+			}
+			else
+			{
+				ExistingImplementedDataInterface->NumVariables++;
+			}
 		}
 
 		// Create new property bags and migrate
-		EPropertyBagResult Result = InAsset->VariableDefaults.ReplaceAllPropertiesAndValues(PropertyDescs, Values);
+		EPropertyBagResult Result = DataInterface->VariableDefaults.ReplaceAllPropertiesAndValues(PropertyDescs, Values);
 		check(Result == EPropertyBagResult::Success);
 
 		if(NumPublicVariables > 0)
 		{
 			TConstArrayView<FPropertyBagPropertyDesc> PublicPropertyDescs(PropertyDescs.GetData(), NumPublicVariables);
 			TConstArrayView<TConstArrayView<uint8>> PublicValues(Values.GetData(), NumPublicVariables);
-			Result = InAsset->PublicVariableDefaults.ReplaceAllPropertiesAndValues(PublicPropertyDescs, PublicValues);
+			Result = DataInterface->PublicVariableDefaults.ReplaceAllPropertiesAndValues(PublicPropertyDescs, PublicValues);
 			check(Result == EPropertyBagResult::Success);
 		}
 		else
 		{
-			InAsset->PublicVariableDefaults.Reset();
+			DataInterface->PublicVariableDefaults.Reset();
 		}
 
 		// Rebuild external variables
-		InAsset->VM->SetExternalVariableDefs(InAsset->GetExternalVariablesImpl(false));
+		DataInterface->VM->SetExternalVariableDefs(DataInterface->GetExternalVariablesImpl(false));
 	}
 	else
 	{
-		InAsset->VariableDefaults.Reset();
-		InAsset->PublicVariableDefaults.Reset();
-		InAsset->VM->ClearExternalVariables(InAsset->ExtendedExecuteContext);
+		DataInterface->ImplementedInterfaces.Empty();
+		DataInterface->VariableDefaults.Reset();
+		DataInterface->PublicVariableDefaults.Reset();
+		DataInterface->VM->ClearExternalVariables(DataInterface->ExtendedExecuteContext);
 	}
 }
 
-URigVMGraph* FUtils::CompileVariableBindings(const FRigVMCompileSettings& InSettings, UAnimNextRigVMAsset* InAsset)
+void FUtils::CompileVariableBindings(const FRigVMCompileSettings& InSettings, UAnimNextRigVMAsset* InAsset, TArray<URigVMGraph*>& OutGraphs)
 {
 	check(InAsset);
 
@@ -158,13 +300,13 @@ URigVMGraph* FUtils::CompileVariableBindings(const FRigVMCompileSettings& InSett
 			continue;
 		}
 
-		const FAnimNextVariableBinding& Binding = Variable->GetBinding();
-		if(!Binding.IsValid())
+		TConstStructView<FAnimNextVariableBindingData> Binding = Variable->GetBinding();
+		if(!Binding.IsValid() || !Binding.Get<FAnimNextVariableBindingData>().IsValid())
 		{
 			continue;
 		}
 
-		TSharedPtr<IVariableBindingType> BindingType = Module.FindVariableBindingType(Binding.BindingData.GetScriptStruct());
+		TSharedPtr<IVariableBindingType> BindingType = Module.FindVariableBindingType(Binding.GetScriptStruct());
 		if(!BindingType.IsValid())
 		{
 			continue;
@@ -172,7 +314,7 @@ URigVMGraph* FUtils::CompileVariableBindings(const FRigVMCompileSettings& InSett
 
 		TArray<IVariableBindingType::FBindingGraphInput>& Group = BindingGroups.FindOrAdd(BindingType.Get());
 		FRigVMTemplateArgumentType RigVMArg = Variable->GetType().ToRigVMTemplateArgument();
-		Group.Add({ Variable->GetVariableName(), RigVMArg.CPPType.ToString(), RigVMArg.CPPTypeObject, Binding.BindingData });
+		Group.Add({ Variable->GetVariableName(), RigVMArg.CPPType.ToString(), RigVMArg.CPPTypeObject, Binding });
 	}
 
 	const bool bHasBindings = BindingGroups.Num() > 0;
@@ -180,7 +322,7 @@ URigVMGraph* FUtils::CompileVariableBindings(const FRigVMCompileSettings& InSett
 	if(!bHasBindings && !bHasPublicVariablesToCopy)
 	{
 		// Nothing to do here
-		return nullptr;
+		return;
 	}
 
 	URigVMGraph* BindingGraph = NewObject<URigVMGraph>(EditorData, NAME_None, RF_Transient);
@@ -191,13 +333,13 @@ URigVMGraph* FUtils::CompileVariableBindings(const FRigVMCompileSettings& InSett
 	if(ExcecuteBindingsNode == nullptr)
 	{
 		InSettings.ReportError(TEXT("Could not spawn Execute Bindings node"));
-		return nullptr;
+		return;
 	}
 	URigVMPin* ExecuteBindingsExecPin = ExcecuteBindingsNode->FindPin(FRigVMStruct::ExecuteContextName.ToString());
 	if(ExecuteBindingsExecPin == nullptr)
 	{
 		InSettings.ReportError(TEXT("Could not find execute pin on Execute Bindings node"));
-		return nullptr;
+		return;
 	}
 	URigVMPin* ExecPin = ExecuteBindingsExecPin;
 
@@ -207,24 +349,25 @@ URigVMGraph* FUtils::CompileVariableBindings(const FRigVMCompileSettings& InSett
 		if(CopyProxyVariablesNode == nullptr)
 		{
 			InSettings.ReportError(TEXT("Could not spawn Copy Module Proxy Variables node"));
-			return nullptr;
+			return;
 		}
 		URigVMPin* CopyProxyVariablesExecPin = CopyProxyVariablesNode->FindPin(FRigVMStruct::ExecuteContextName.ToString());
 		if(ExecPin == nullptr)
 		{
 			InSettings.ReportError(TEXT("Could not find execute pin on Execute Bindings node"));
-			return nullptr;
+			return;
 		}
 		bool bLinkAdded = Controller->AddLink(ExecuteBindingsExecPin, CopyProxyVariablesExecPin, false);
 		if(!bLinkAdded)
 		{
 			InSettings.ReportError(TEXT("Could not link Execute Bindings node"));
-			return nullptr;
+			return;
 		}
 		ExecPin = CopyProxyVariablesExecPin;
 	}
 
 	IVariableBindingType::FBindingGraphFragmentArgs Args;
+	Args.Event = FRigUnit_AnimNextExecuteBindings::StaticStruct();
 	Args.Controller = Controller;
 	Args.BindingGraph = BindingGraph;
 	Args.ExecTail = ExecPin;
@@ -236,7 +379,7 @@ URigVMGraph* FUtils::CompileVariableBindings(const FRigVMCompileSettings& InSett
 		BindingGroupPair.Key->BuildBindingGraphFragment(InSettings, Args, ExecPin, Location);
 	}
 
-	return BindingGraph;
+	OutGraphs.Add(BindingGraph);
 }
 
 UAnimNextRigVMAsset* FUtils::GetAsset(UAnimNextRigVMAssetEditorData* InEditorData)
@@ -618,9 +761,9 @@ void FUtils::GetAssetVariables(const UAnimNextRigVMAssetEditorData* EditorData,
 	OutExports.Variables = ExportSet.Array();
 }
 
-void FUtils::GetAssetVariables(const UAnimNextRigVMAssetEditorData* EditorData, TSet<FAnimNextAssetRegistryExportedVariable>& OutExports)
+void FUtils::GetAssetVariables(const UAnimNextRigVMAssetEditorData* InEditorData, TSet<FAnimNextAssetRegistryExportedVariable>& OutExports)
 {
-	for(const UAnimNextRigVMAssetEntry* Entry : EditorData->Entries)
+	for(const UAnimNextRigVMAssetEntry* Entry : InEditorData->Entries)
 	{
 		if(const IAnimNextRigVMExportInterface* ExportInterface = Cast<IAnimNextRigVMExportInterface>(Entry))
 		{
@@ -632,6 +775,14 @@ void FUtils::GetAssetVariables(const UAnimNextRigVMAssetEditorData* EditorData,
 				AddParamToSet(NewParam, OutExports);
 			}
 		}
+		else if(const UAnimNextDataInterfaceEntry* DataInterfaceEntry = Cast<UAnimNextDataInterfaceEntry>(Entry))
+		{
+			if(DataInterfaceEntry->DataInterface)
+			{
+				UAnimNextDataInterface_EditorData* EditorData = GetEditorData<UAnimNextDataInterface_EditorData>(DataInterfaceEntry->DataInterface.Get());
+				GetAssetVariables(EditorData, OutExports);
+			}
+		}
 	}
 }
 
diff --git a/Engine/Plugins/Experimental/AnimNext/Source/AnimNextUncookedOnly/Private/Variables/UniversalObjectLocatorBindingType.cpp b/Engine/Plugins/Experimental/AnimNext/Source/AnimNextUncookedOnly/Private/Variables/UniversalObjectLocatorBindingType.cpp
--- a/Engine/Plugins/Experimental/AnimNext/Source/AnimNextUncookedOnly/Private/Variables/UniversalObjectLocatorBindingType.cpp
+++ b/Engine/Plugins/Experimental/AnimNext/Source/AnimNextUncookedOnly/Private/Variables/UniversalObjectLocatorBindingType.cpp
@@ -14,6 +14,7 @@
 #include "Framework/Application/SlateApplication.h"
 #include "Framework/PropertyViewer/IFieldExpander.h"
 #include "Framework/PropertyViewer/IFieldIterator.h"
+#include "Module/RigUnit_AnimNextModuleEvents.h"
 #include "Modules/ModuleManager.h"
 #include "Param/ParamCompatibility.h"
 #include "Param/ParamUtils.h"
@@ -680,6 +681,11 @@ FText FUniversalObjectLocatorBindingType::GetTooltipText(TConstStructView<FAnimN
 
 void FUniversalObjectLocatorBindingType::BuildBindingGraphFragment(const FRigVMCompileSettings& InSettings, const FBindingGraphFragmentArgs& InArgs, URigVMPin*& OutExecTail, FVector2D& OutLocation) const
 {
+	if(InArgs.Event != FRigUnit_AnimNextExecuteBindings::StaticStruct())
+	{
+		return;
+	}
+	
 	URigVMPin* ExecTail = InArgs.ExecTail;
 
 	struct FBindingData
diff --git a/Engine/Plugins/Experimental/AnimNext/Source/AnimNextUncookedOnly/Public/Variables/AnimNextVariableBinding.h b/Engine/Plugins/Experimental/AnimNext/Source/AnimNextUncookedOnly/Public/Variables/AnimNextVariableBinding.h
--- a/Engine/Plugins/Experimental/AnimNext/Source/AnimNextUncookedOnly/Public/Variables/AnimNextVariableBinding.h
+++ b/Engine/Plugins/Experimental/AnimNext/Source/AnimNextUncookedOnly/Public/Variables/AnimNextVariableBinding.h
@@ -20,7 +20,7 @@ struct FAnimNextVariableBinding
 		return BindingData.IsValid() && BindingData.Get<FAnimNextVariableBindingData>().IsValid();
 	}
 
-	UPROPERTY(EditAnywhere, Category = "Binding", NoClear, meta=(ExcludeBaseStruct))
+	UPROPERTY(EditAnywhere, Category = "Binding", meta=(ExcludeBaseStruct))
 	TInstancedStruct<FAnimNextVariableBindingData> BindingData;
 };
 
diff --git a/Engine/Plugins/Experimental/AnimNext/Source/AnimNextUncookedOnly/Public/Variables/IAnimNextRigVMVariableInterface.h b/Engine/Plugins/Experimental/AnimNext/Source/AnimNextUncookedOnly/Public/Variables/IAnimNextRigVMVariableInterface.h
--- a/Engine/Plugins/Experimental/AnimNext/Source/AnimNextUncookedOnly/Public/Variables/IAnimNextRigVMVariableInterface.h
+++ b/Engine/Plugins/Experimental/AnimNext/Source/AnimNextUncookedOnly/Public/Variables/IAnimNextRigVMVariableInterface.h
@@ -8,6 +8,7 @@
 
 struct FAnimNextParamType;
 struct FAnimNextVariableBinding;
+struct FAnimNextVariableBindingData;
 
 UINTERFACE(meta=(CannotImplementInterfaceInBlueprint))
 class ANIMNEXTUNCOOKEDONLY_API UAnimNextRigVMVariableInterface : public UInterface
@@ -32,8 +33,11 @@ public:
 	// Set the variable name
 	virtual void SetVariableName(FName InName, bool bSetupUndoRedo = true) = 0;
 
+	// Set the default value
+	virtual bool SetDefaultValue(TConstArrayView<uint8> InValue, bool bSetupUndoRedo = true) = 0;
+
 	// Set the default value from a string
-	virtual bool SetDefaultValue(const FString& InDefaultValue, bool bSetupUndoRedo = true) = 0;
+	virtual bool SetDefaultValueFromString(const FString& InDefaultValue, bool bSetupUndoRedo = true) = 0;
 
 	// Access the backing storage property bag for the parameter
 	virtual const FInstancedPropertyBag& GetPropertyBag() const = 0;
@@ -44,8 +48,13 @@ public:
 		return const_cast<FInstancedPropertyBag&>(GetPropertyBag());
 	}
 
+	virtual bool GetDefaultValue(const FProperty*& OutProperty, TConstArrayView<uint8>& OutValue) const = 0;
+	
+	// Set the binding type for this variable
+	virtual void SetBindingType(UScriptStruct* InBindingTypeStruct, bool bSetupUndoRedo) = 0;
+	
 	// Get the binding for this variable, if any
-	virtual const FAnimNextVariableBinding& GetBinding() const = 0;
+	virtual TConstStructView<FAnimNextVariableBindingData> GetBinding() const = 0;
 
 	// Access the memory for the internal value
 	const uint8* GetValuePtr() const
diff --git a/Engine/Plugins/Experimental/AnimNext/Source/AnimNextUncookedOnly/Public/Variables/IVariableBindingType.h b/Engine/Plugins/Experimental/AnimNext/Source/AnimNextUncookedOnly/Public/Variables/IVariableBindingType.h
--- a/Engine/Plugins/Experimental/AnimNext/Source/AnimNextUncookedOnly/Public/Variables/IVariableBindingType.h
+++ b/Engine/Plugins/Experimental/AnimNext/Source/AnimNextUncookedOnly/Public/Variables/IVariableBindingType.h
@@ -66,10 +66,17 @@ protected:
 
 	struct FBindingGraphFragmentArgs
 	{
+		// The event (e.g. FRigUnit_AnimNextExecuteBindings) that is currently being processed
+		UScriptStruct* Event;
+		// All inputs, corresponding to variables
 		TConstArrayView<FBindingGraphInput> Inputs;
+		// Controller to use for instantiation 
 		URigVMController* Controller;
+		// Graph to instantiate nodes into
 		URigVMGraph* BindingGraph;
+		// The exec pin of the last node that was instantiated, for chaining 
 		URigVMPin* ExecTail;
+		// The current spawn location, useful for making user-readable graphs
 		FVector2D CurrentLocation;
 	};
 	
diff --git a/Engine/Plugins/Experimental/AnimNextAnimGraph/Config/BaseAnimNextAnimGraph.ini b/Engine/Plugins/Experimental/AnimNextAnimGraph/Config/BaseAnimNextAnimGraph.ini
--- a/Engine/Plugins/Experimental/AnimNextAnimGraph/Config/BaseAnimNextAnimGraph.ini
+++ b/Engine/Plugins/Experimental/AnimNextAnimGraph/Config/BaseAnimNextAnimGraph.ini
@@ -25,4 +25,7 @@
 +StructRedirects=(OldName="/Script/AnimNext.AnimNode_AnimNextGraph", NewName="/Script/AnimNextAnimGraph.AnimNode_AnimNextGraph")
 
 ; AnimNextGraph -> AnimNextModule
-+PropertyRedirects=(OldName="/Script/AnimNextGraph.AnimNode_AnimNextGraph.AnimNextGraph", NewName="Module")
\ No newline at end of file
++PropertyRedirects=(OldName="/Script/AnimNextGraph.AnimNode_AnimNextGraph.AnimNextGraph", NewName="Module")
+
+; Rename K2 Node
++ClassRedirects=(OldName="/Script/AnimNextAnimGraphUncookedOnly.K2Node_PlayAnim", NewName="/Script/AnimNextAnimGraphUncookedOnly.K2Node_AnimNextPlayAnim")
\ No newline at end of file
diff --git a/Engine/Plugins/Experimental/AnimNextAnimGraph/Source/AnimNextAnimGraph/Internal/PlayAnim/PlayAnimCallbackProxy.h b/Engine/Plugins/Experimental/AnimNextAnimGraph/Source/AnimNextAnimGraph/Internal/PlayAnim/PlayAnimCallbackProxy.h
--- a/Engine/Plugins/Experimental/AnimNextAnimGraph/Source/AnimNextAnimGraph/Internal/PlayAnim/PlayAnimCallbackProxy.h
+++ b/Engine/Plugins/Experimental/AnimNextAnimGraph/Source/AnimNextAnimGraph/Internal/PlayAnim/PlayAnimCallbackProxy.h
@@ -36,6 +36,16 @@ class UPlayAnimCallbackProxy : public UObject
 		FAnimNextPlayAnimBlendSettings BlendInSettings = FAnimNextPlayAnimBlendSettings(),
 		FAnimNextPlayAnimBlendSettings BlendOutSettings = FAnimNextPlayAnimBlendSettings());
 
+	// Called to perform the query internally
+	UFUNCTION(BlueprintCallable, meta = (BlueprintInternalUseOnly = "true"))
+	static ANIMNEXTANIMGRAPH_API UPlayAnimCallbackProxy* CreateProxyObjectForPlayAsset(
+		class UAnimNextComponent* AnimNextComponent,
+		FName SlotName,
+		UObject* Asset,
+		const FInstancedStruct& Payload,
+		FAnimNextPlayAnimBlendSettings BlendInSettings = FAnimNextPlayAnimBlendSettings(),
+		FAnimNextPlayAnimBlendSettings BlendOutSettings = FAnimNextPlayAnimBlendSettings());
+	
 public:
 	//~ Begin UObject Interface
 	ANIMNEXTANIMGRAPH_API virtual void BeginDestroy() override;
@@ -46,7 +56,16 @@ protected:
 	ANIMNEXTANIMGRAPH_API void OnPlayAnimInterrupted(const UE::AnimNext::FPlayAnimRequest& Request);
 	ANIMNEXTANIMGRAPH_API void OnPlayAnimBlendingOut(const UE::AnimNext::FPlayAnimRequest& Request);
 
-	// Attempts to play a montage with the specified settings. Returns whether it started or not.
+	// Attempts to play an object with the specified payload. Returns whether it started or not.
+	ANIMNEXTANIMGRAPH_API bool Play(
+		class UAnimNextComponent* AnimNextComponent,
+		FName SlotName,
+		UObject* Object,
+		FInstancedStruct&& Payload,
+		const UE::AnimNext::FPlayAnimBlendSettings& BlendInSettings,
+		const UE::AnimNext::FPlayAnimBlendSettings& BlendOutSettings);
+	
+	// Attempts to play an animation with the specified settings. Returns whether it started or not.
 	ANIMNEXTANIMGRAPH_API bool Play(
 		class UAnimNextComponent* AnimNextComponent,
 		FName SlotName,
diff --git a/Engine/Plugins/Experimental/AnimNextAnimGraph/Source/AnimNextAnimGraph/Private/PlayAnim/PlayAnimCallbackProxy.cpp b/Engine/Plugins/Experimental/AnimNextAnimGraph/Source/AnimNextAnimGraph/Private/PlayAnim/PlayAnimCallbackProxy.cpp
--- a/Engine/Plugins/Experimental/AnimNextAnimGraph/Source/AnimNextAnimGraph/Private/PlayAnim/PlayAnimCallbackProxy.cpp
+++ b/Engine/Plugins/Experimental/AnimNextAnimGraph/Source/AnimNextAnimGraph/Private/PlayAnim/PlayAnimCallbackProxy.cpp
@@ -1,6 +1,7 @@
 // Copyright Epic Games, Inc. All Rights Reserved.
 
 #include "PlayAnim/PlayAnimCallbackProxy.h"
+#include "Animation/AnimSequence.h"
 
 #include UE_INLINE_GENERATED_CPP_BY_NAME(PlayAnimCallbackProxy)
 
@@ -27,6 +28,21 @@ UPlayAnimCallbackProxy* UPlayAnimCallbackProxy::CreateProxyObjectForPlayAnim(
 	return Proxy;
 }
 
+UPlayAnimCallbackProxy* UPlayAnimCallbackProxy::CreateProxyObjectForPlayAsset(
+	UAnimNextComponent* AnimNextComponent,
+	FName SlotName,
+	UObject* Asset,
+	const FInstancedStruct& Payload,
+	UE::AnimNext::FPlayAnimBlendSettings BlendInSettings,
+	UE::AnimNext::FPlayAnimBlendSettings BlendOutSettings)
+{
+	UPlayAnimCallbackProxy* Proxy = NewObject<UPlayAnimCallbackProxy>();
+	Proxy->SetFlags(RF_StrongRefOnFrame);
+	FInstancedStruct PayloadCopy(Payload);
+	Proxy->Play(AnimNextComponent, SlotName, Asset, MoveTemp(PayloadCopy), BlendInSettings, BlendOutSettings);
+	return Proxy;
+}
+
 bool UPlayAnimCallbackProxy::Play(
 	UAnimNextComponent* AnimNextComponent,
 	FName SlotName,
@@ -35,24 +51,40 @@ bool UPlayAnimCallbackProxy::Play(
 	float StartPosition,
 	const UE::AnimNext::FPlayAnimBlendSettings& BlendInSettings,
 	const UE::AnimNext::FPlayAnimBlendSettings& BlendOutSettings)
+{
+	FInstancedStruct Payload;
+	Payload.InitializeAs<FAnimNextPlayAnimPayload>();
+	FAnimNextPlayAnimPayload& PlayAnimPayload = Payload.GetMutable<FAnimNextPlayAnimPayload>();
+	PlayAnimPayload.AnimationObject = AnimationObject;
+	PlayAnimPayload.PlayRate = PlayRate;
+	PlayAnimPayload.StartPosition = StartPosition;
+	return Play(AnimNextComponent, SlotName, AnimationObject, MoveTemp(Payload), BlendInSettings, BlendOutSettings);
+}
+
+bool UPlayAnimCallbackProxy::Play(
+	UAnimNextComponent* AnimNextComponent,
+	FName SlotName,
+	UObject* Object,
+	FInstancedStruct&& Payload,
+	const UE::AnimNext::FPlayAnimBlendSettings& BlendInSettings,
+	const UE::AnimNext::FPlayAnimBlendSettings& BlendOutSettings)
 {
 	bool bPlayedSuccessfully = false;
 	if (AnimNextComponent != nullptr)
 	{
 		UE::AnimNext::FPlayAnimRequestArgs RequestArgs;
 		RequestArgs.SlotName = SlotName;
-		RequestArgs.AnimationObject = AnimationObject;
-		RequestArgs.PlayRate = PlayRate;
-		RequestArgs.StartPosition = StartPosition;
+		RequestArgs.Object = Object;
 		RequestArgs.BlendInSettings = BlendInSettings;
 		RequestArgs.BlendOutSettings = BlendOutSettings;
+		RequestArgs.Payload = MoveTemp(Payload);
 
 		auto Request = UE::AnimNext::MakePlayAnimRequest();
 		Request->OnCompleted.BindUObject(this, &UPlayAnimCallbackProxy::OnPlayAnimCompleted);
 		Request->OnInterrupted.BindUObject(this, &UPlayAnimCallbackProxy::OnPlayAnimInterrupted);
 		Request->OnBlendingOut.BindUObject(this, &UPlayAnimCallbackProxy::OnPlayAnimBlendingOut);
 
-		bPlayedSuccessfully = Request->Play(RequestArgs, AnimNextComponent);
+		bPlayedSuccessfully = Request->Play(MoveTemp(RequestArgs), AnimNextComponent);
 
 		PlayingRequest = Request;
 		bWasInterrupted = false;
diff --git a/Engine/Plugins/Experimental/AnimNextAnimGraph/Source/AnimNextAnimGraph/Private/PlayAnim/PlayAnimRequest.cpp b/Engine/Plugins/Experimental/AnimNextAnimGraph/Source/AnimNextAnimGraph/Private/PlayAnim/PlayAnimRequest.cpp
--- a/Engine/Plugins/Experimental/AnimNextAnimGraph/Source/AnimNextAnimGraph/Private/PlayAnim/PlayAnimRequest.cpp
+++ b/Engine/Plugins/Experimental/AnimNextAnimGraph/Source/AnimNextAnimGraph/Private/PlayAnim/PlayAnimRequest.cpp
@@ -7,10 +7,10 @@
 
 namespace UE::AnimNext
 {
-	bool FPlayAnimRequest::Play(const FPlayAnimRequestArgs& InRequestArgs, UAnimNextComponent* InComponent)
+	bool FPlayAnimRequest::Play(FPlayAnimRequestArgs&& InRequestArgs, UAnimNextComponent* InComponent)
 	{
 		check(IsInGameThread());
-		if (InRequestArgs.AnimationObject == nullptr || InComponent == nullptr)
+		if (!InRequestArgs.Payload.IsValid() || InComponent == nullptr)
 		{
 			return false;	// Nothing to play
 		}
@@ -20,7 +20,7 @@ namespace UE::AnimNext
 			return false;	// Already playing, cannot play again
 		}
 
-		RequestArgs = InRequestArgs;
+		RequestArgs = MoveTemp(InRequestArgs);
 		Component = InComponent;
 		Status = EPlayAnimStatus::Pending;
 
@@ -66,6 +66,11 @@ namespace UE::AnimNext
 		return RequestArgs;
 	}
 
+	FPlayAnimRequestArgs& FPlayAnimRequest::GetMutableArgs()
+	{
+		return RequestArgs;
+	}
+
 	EPlayAnimStatus FPlayAnimRequest::GetStatus() const
 	{
 		check(IsInGameThread());
@@ -154,7 +159,8 @@ namespace UE::AnimNext
 
 	void FPlayAnimRequest::AddReferencedObjects(FReferenceCollector& Collector)
 	{
-		Collector.AddReferencedObject(RequestArgs.AnimationObject);
+		Collector.AddReferencedObject(RequestArgs.Object);
+		Collector.AddPropertyReferencesWithStructARO(FInstancedPropertyBag::StaticStruct(), &RequestArgs.Payload);
 		Collector.AddReferencedObject(Component);
 	}
 }
diff --git a/Engine/Plugins/Experimental/AnimNextAnimGraph/Source/AnimNextAnimGraph/Private/PlayAnim/PlayAnimSlotTrait.cpp b/Engine/Plugins/Experimental/AnimNextAnimGraph/Source/AnimNextAnimGraph/Private/PlayAnim/PlayAnimSlotTrait.cpp
--- a/Engine/Plugins/Experimental/AnimNextAnimGraph/Source/AnimNextAnimGraph/Private/PlayAnim/PlayAnimSlotTrait.cpp
+++ b/Engine/Plugins/Experimental/AnimNextAnimGraph/Source/AnimNextAnimGraph/Private/PlayAnim/PlayAnimSlotTrait.cpp
@@ -3,10 +3,13 @@
 #include "PlayAnim/PlayAnimSlotTrait.h"
 
 #include "Animation/AnimSequence.h"
+#include "DataInterface/AnimNextDataInterfaceHost.h"
+#include "DataInterface/DataInterfaceStructAdapter.h"
 #include "TraitCore/ExecutionContext.h"
 #include "TraitCore/NodeInstance.h"
 #include "TraitInterfaces/ITimeline.h"
 #include "Graph/AnimNextGraphInstance.h"
+#include "Logging/StructuredLog.h"
 #include "Module/ModuleEvents.h"
 
 namespace UE::AnimNext
@@ -159,13 +162,19 @@ namespace UE::AnimNext
 			if (Request)
 			{
 				// This is a new pending request, lookup the sub-graph to use with our chooser and the desired animation object
+
+				// TODO: Rather than this built-in choosing logic, this could be implemented by a new optional trait:
+				// - Implement some new trait that converts an event payload's object into a graph + payload struct for variable binding
+				// - Could possibly remove Object from FAnimNextPlayAnimRequestArgs, or AnimationObject from FAnimNextPlayAnimPayload, as they are duplicated
+				// - Call this trait's interface here and use the resulting animation graph (binding variables in OnBlendInitiated)
+				// - Reentrancy needs to handle 'same graph, different parameters'
 				const FAnimNextPlayAnimRequestArgs& RequestArgs = Request->GetArgs();
-				if (RequestArgs.AnimationObject)
+				if (RequestArgs.Object)
 				{
 					if (const UChooserTable* Chooser = SharedData->GetSubGraphChooser(Binding))
 					{
 						FAnimNextPlayAnimChooserParameters ChooserParameters;
-						ChooserParameters.AnimationObjectType = RequestArgs.AnimationObject->GetClass();
+						ChooserParameters.AnimationObjectType = RequestArgs.Object->GetClass();
 
 						FChooserEvaluationContext ChooserContext;
 						ChooserContext.AddStructParam(ChooserParameters);
@@ -376,27 +385,9 @@ namespace UE::AnimNext
 				SlotRequest.AnimationGraph->AllocateInstance(Binding.GetTraitPtr().GetNodeInstance()->GetOwner(), SlotRequest.GraphInstance, EntryPoint);
 				SlotRequest.ChildPtr = SlotRequest.GraphInstance.GetGraphRootPtr();
 
-				// Setup our graph parameters
-				const FAnimNextPlayAnimRequestArgs& RequestArgs = SlotRequest.Request->GetArgs();
-
-				const FString SubGraphPath = SlotRequest.AnimationGraph->GetPathName();
-
-				const FName AnimationObjectName(FString::Printf(TEXT("%s:AnimationObject"), *SubGraphPath));
-				const FName StartPositionName(FString::Printf(TEXT("%s:StartPosition"), *SubGraphPath));
-				const FName PlayRateName(FString::Printf(TEXT("%s:PlayRate"), *SubGraphPath));
-				const FName IsLoopingName(FString::Printf(TEXT("%s:IsLooping"), *SubGraphPath));
-
-				FInstancedPropertyBag SubGraphParameters;
-				SubGraphParameters.AddProperty(AnimationObjectName, EPropertyBagPropertyType::Object, UAnimSequence::StaticClass());
-				SubGraphParameters.SetValueObject(AnimationObjectName, RequestArgs.AnimationObject);
-				SubGraphParameters.AddProperty(StartPositionName, EPropertyBagPropertyType::Double);
-				SubGraphParameters.SetValueFloat(StartPositionName, RequestArgs.StartPosition);
-				SubGraphParameters.AddProperty(PlayRateName, EPropertyBagPropertyType::Double);
-				SubGraphParameters.SetValueFloat(PlayRateName, RequestArgs.PlayRate);
-				SubGraphParameters.AddProperty(IsLoopingName, EPropertyBagPropertyType::Bool);
-				SubGraphParameters.SetValueBool(IsLoopingName, false);
-
-				// TODO: forwarding variables to playanim 
+				// Note: args are mutable here as bindings allow writes!
+				FDataInterfaceStructAdapter VariableBinding(SlotRequest.AnimationGraph, SlotRequest.Request->GetMutableArgs().Payload);
+				SlotRequest.GraphInstance.BindPublicVariables({ &VariableBinding });
 
 				// TODO: Validate that our child implements the ITimeline interface
 
diff --git a/Engine/Plugins/Experimental/AnimNextAnimGraph/Source/AnimNextAnimGraph/Public/PlayAnim/PlayAnimRequest.h b/Engine/Plugins/Experimental/AnimNextAnimGraph/Source/AnimNextAnimGraph/Public/PlayAnim/PlayAnimRequest.h
--- a/Engine/Plugins/Experimental/AnimNextAnimGraph/Source/AnimNextAnimGraph/Public/PlayAnim/PlayAnimRequest.h
+++ b/Engine/Plugins/Experimental/AnimNextAnimGraph/Source/AnimNextAnimGraph/Public/PlayAnim/PlayAnimRequest.h
@@ -8,9 +8,11 @@
 #include "TraitCore/TraitEvent.h"
 #include "TraitInterfaces/ITimeline.h"
 #include "PlayAnim/PlayAnimStatus.h"
+#include "StructUtils/PropertyBag.h"
 
 #include "PlayAnimRequest.generated.h"
 
+class UAnimNextDataInterface;
 class FReferenceCollector;
 class UAnimNextModule;
 class UAnimNextComponent;
@@ -58,30 +60,45 @@ struct FAnimNextPlayAnimBlendSettings
 };
 
 /**
- * PlayAnim Request Arguments
+ * PlayAnim payload
  *
- * Encapsulates the parameters required to initiate a Play animation request.
+ * Encapsulates the data interface used to play a simple animation
  */
-USTRUCT()
-struct FAnimNextPlayAnimRequestArgs
+USTRUCT(BlueprintType)
+struct FAnimNextPlayAnimPayload
 {
 	GENERATED_BODY()
 
-	// The slot name to play the animation object on with this request
-	UPROPERTY()
-	FName SlotName;
-
 	// The animation object to play with this request
-	UPROPERTY()
+	UPROPERTY(EditAnywhere, BlueprintReadWrite, Category = "PlayAnim")
 	TObjectPtr<UAnimSequence> AnimationObject;
 
 	// The play rate of the request
-	UPROPERTY()
-	float PlayRate = 1.0f;
+	UPROPERTY(EditAnywhere, BlueprintReadWrite, Category = "PlayAnim")
+	double PlayRate = 1.0f;
 
 	// The timeline start position of the request
+	UPROPERTY(EditAnywhere, BlueprintReadWrite, Category = "PlayAnim")
+	double StartPosition = 0.0f;
+
+	// Whether to loop the animation
+	UPROPERTY(EditAnywhere, BlueprintReadWrite, Category = "PlayAnim")
+	bool IsLooping = false;
+};
+
+/**
+ * PlayAnim Request Arguments
+ *
+ * Encapsulates the parameters required to initiate a Play animation request.
+ */
+USTRUCT()
+struct FAnimNextPlayAnimRequestArgs
+{
+	GENERATED_BODY()
+
+	// The slot name to play the animation object on with this request
 	UPROPERTY()
-	float StartPosition = 0.0f;
+	FName SlotName;
 
 	// The blend settings to use when blending in
 	UPROPERTY()
@@ -90,6 +107,15 @@ struct FAnimNextPlayAnimRequestArgs
 	// The blend settings to use when blending out (if not interrupted)
 	UPROPERTY()
 	FAnimNextPlayAnimBlendSettings BlendOutSettings;
+
+	// Object to 'play'.
+	// The animation graph to be instantiated for this request will be chosen by interrogating this object's class.
+	UPROPERTY()
+	TObjectPtr<UObject> Object;
+
+	// Payload that will be applied to the animation graph's variables via its data interfaces.
+	UPROPERTY()
+	FInstancedStruct Payload;
 };
 
 namespace UE::AnimNext
@@ -130,7 +156,7 @@ namespace UE::AnimNext
 		FAnimNextOnPlayAnimBlendingOut OnBlendingOut;
 
 		// Sends this request to the specified component and it will attempt to play with the requested arguments
-		bool Play(const FPlayAnimRequestArgs& InRequestArgs, UAnimNextComponent* InComponent);
+		bool Play(FPlayAnimRequestArgs&& InRequestArgs, UAnimNextComponent* InComponent);
 
 		// Interrupts this request and request that we transition to the source input on the playing slot
 		void Stop();
@@ -138,6 +164,9 @@ namespace UE::AnimNext
 		// Returns the arguments this request is playing
 		const FPlayAnimRequestArgs& GetArgs() const;
 
+		// Returns the arguments this request is playing
+		FPlayAnimRequestArgs& GetMutableArgs();
+
 		// Returns the request status
 		EPlayAnimStatus GetStatus() const;
 
diff --git a/Engine/Plugins/Experimental/AnimNextAnimGraph/Source/AnimNextAnimGraphTestSuite/Private/AnimNextAnimGraphTraitGraphTest.cpp b/Engine/Plugins/Experimental/AnimNextAnimGraph/Source/AnimNextAnimGraphTestSuite/Private/AnimNextAnimGraphTraitGraphTest.cpp
--- a/Engine/Plugins/Experimental/AnimNextAnimGraph/Source/AnimNextAnimGraphTestSuite/Private/AnimNextAnimGraphTraitGraphTest.cpp
+++ b/Engine/Plugins/Experimental/AnimNextAnimGraph/Source/AnimNextAnimGraphTestSuite/Private/AnimNextAnimGraphTraitGraphTest.cpp
@@ -5,11 +5,14 @@
 #include "AnimNextRuntimeTest.h"
 #include "AnimNextTest.h"
 #include "AssetToolsModule.h"
+#include "IAnimNextRigVMExportInterface.h"
 #include "UncookedOnlyUtils.h"
+#include "Animation/AnimSequence.h"
 #include "TraitCore/TraitRegistry.h"
 #include "TraitInterfaces/IEvaluate.h"
 #include "TraitInterfaces/IUpdate.h"
 #include "Editor/Transactor.h"
+#include "Entries/AnimNextVariableEntry.h"
 #include "Graph/AnimNextAnimationGraph_EditorData.h"
 #include "Graph/AnimNextAnimationGraph.h"
 #include "Graph/AnimNextAnimationGraphFactory.h"
@@ -619,4 +622,299 @@ bool FAnimationAnimNextRuntimeTest_GraphExecuteLatent::RunTest(const FString& In
 	return true;
 }
 
+
+IMPLEMENT_SIMPLE_AUTOMATION_TEST(FAnimationAnimNextRuntimeTest_Variables, "Animation.AnimNext.Runtime.Graph.Variables", EAutomationTestFlags::EditorContext | EAutomationTestFlags::EngineFilter)
+
+bool FAnimationAnimNextRuntimeTest_Variables::RunTest(const FString& InParameters)
+{
+	using namespace UE::AnimNext;
+
+	{
+		UFactory* GraphFactory = NewObject<UAnimNextAnimationGraphFactory>();
+		UAnimNextAnimationGraph* AnimationGraph = CastChecked<UAnimNextAnimationGraph>(GraphFactory->FactoryCreateNew(UAnimNextAnimationGraph::StaticClass(), GetTransientPackage(), TEXT("TestAnimNextGraph"), RF_Transient, nullptr, nullptr, NAME_None));
+		UE_RETURN_ON_ERROR(AnimationGraph != nullptr, "FAnimationAnimNextRuntimeTest_Variables -> Failed to create animation graph");
+
+		auto AddPublicVariable = [this](UAnimNextAnimationGraph* InAnimationGraph, FName InName, const auto& InValue)
+		{
+			using ValueType = std::remove_reference_t<decltype(InValue)>;
+			FAnimNextParamType Type = FAnimNextParamType::GetType<ValueType>();
+			UAnimNextVariableEntry* VariableEntry = UAnimNextRigVMAssetLibrary::AddVariable(InAnimationGraph, InName, Type.GetValueType(), Type.GetContainerType(), Type.GetValueTypeObject(), TEXT(""), false, false);
+			UE_RETURN_ON_ERROR(VariableEntry != nullptr, "FAnimationAnimNextRuntimeTest_Variables::AddPublicVariable -> Failed to create variable");
+			VariableEntry->SetExportAccessSpecifier(EAnimNextExportAccessSpecifier::Public, false);
+			UE_RETURN_ON_ERROR(VariableEntry->SetDefaultValue(InValue, false), "FAnimationAnimNextRuntimeTest_Variables::AddPublicVariable -> Failed to set variable default value");
+			return true;
+		};
+
+		AddPublicVariable(AnimationGraph, "Bool", true);
+		AddPublicVariable(AnimationGraph, "Byte", (uint8)42);
+		AddPublicVariable(AnimationGraph, "Int32", (int32)-4679222);
+		AddPublicVariable(AnimationGraph, "UInt32", (uint32)3415919103);
+		AddPublicVariable(AnimationGraph, "Int64", (int64)-3415919105);
+		AddPublicVariable(AnimationGraph, "UInt64", (uint64)34159191067);
+		AddPublicVariable(AnimationGraph, "Float", 1.0f);
+		AddPublicVariable(AnimationGraph, "Double", 1.0);
+		AddPublicVariable(AnimationGraph, "Name", FName("Test"));
+		AddPublicVariable(AnimationGraph, "String", FString("Test"));
+		AddPublicVariable(AnimationGraph, "Text", NSLOCTEXT("Tests", "Test", "Test"));
+		AddPublicVariable(AnimationGraph, "Enum", EPropertyBagPropertyType::Double);
+		AddPublicVariable(AnimationGraph, "Struct", FVector::OneVector);
+		AddPublicVariable(AnimationGraph, "DerivedStruct", FTestDerivedVector());
+		AddPublicVariable(AnimationGraph, "Object", UAnimNextDataInterface::StaticClass()->GetDefaultObject<UAnimNextDataInterface>());
+		AddPublicVariable(AnimationGraph, "SoftObject", FSoftObjectPath(UAnimNextDataInterface::StaticClass()->GetDefaultObject<UAnimNextDataInterface>()));
+		AddPublicVariable(AnimationGraph, "Class", UAnimNextDataInterface::StaticClass());
+		AddPublicVariable(AnimationGraph, "SoftClass", FSoftClassPath(UAnimNextDataInterface::StaticClass()));
+
+		FAnimNextGraphInstancePtr GraphInstance;
+		AnimationGraph->AllocateInstance(GraphInstance);
+
+		// Bool/Integers + conversions
+		{
+			// Gets
+			bool TestBool = false;
+			UE_RETURN_ON_ERROR(GraphInstance.GetImpl()->GetVariable("Bool", TestBool) == EPropertyBagResult::Success, "FAnimationAnimNextRuntimeTest_Variables -> GetVariable failed");
+			UE_RETURN_ON_ERROR(TestBool == true, "FAnimationAnimNextRuntimeTest_Variables -> Variable value did not match");
+
+			uint8 TestByte = 0;
+			UE_RETURN_ON_ERROR(GraphInstance.GetImpl()->GetVariable("Byte", TestByte) == EPropertyBagResult::Success, "FAnimationAnimNextRuntimeTest_Variables -> GetVariable failed");
+			UE_RETURN_ON_ERROR(TestByte == 42, "FAnimationAnimNextRuntimeTest_Variables -> Variable value did not match");
+
+			int32 TestInt32 = 0;
+			UE_RETURN_ON_ERROR(GraphInstance.GetImpl()->GetVariable("Int32", TestInt32) == EPropertyBagResult::Success, "FAnimationAnimNextRuntimeTest_Variables -> GetVariable failed");
+			UE_RETURN_ON_ERROR(TestInt32 == -4679222, "FAnimationAnimNextRuntimeTest_Variables -> Variable value did not match");
+
+			uint32 TestUInt32 = 0;
+			UE_RETURN_ON_ERROR(GraphInstance.GetImpl()->GetVariable("UInt32", TestUInt32) == EPropertyBagResult::Success, "FAnimationAnimNextRuntimeTest_Variables -> GetVariable failed");
+			UE_RETURN_ON_ERROR(TestUInt32 == 3415919103, "FAnimationAnimNextRuntimeTest_Variables -> Variable value did not match");
+
+			int64 TestInt64 = 0;
+			UE_RETURN_ON_ERROR(GraphInstance.GetImpl()->GetVariable("Int64", TestInt64) == EPropertyBagResult::Success, "FAnimationAnimNextRuntimeTest_Variables -> GetVariable failed");
+			UE_RETURN_ON_ERROR(TestInt64 == -3415919105, "FAnimationAnimNextRuntimeTest_Variables -> Variable value did not match");
+
+			uint64 TestUInt64 = 0;
+			UE_RETURN_ON_ERROR(GraphInstance.GetImpl()->GetVariable("UInt64", TestUInt64) == EPropertyBagResult::Success, "FAnimationAnimNextRuntimeTest_Variables -> GetVariable failed");
+			UE_RETURN_ON_ERROR(TestUInt64 == 34159191067, "FAnimationAnimNextRuntimeTest_Variables -> Variable value did not match");
+
+			// Conversions
+
+			// Bool
+			TestByte = 0;
+			UE_RETURN_ON_ERROR(GraphInstance.GetImpl()->GetVariable("Bool", TestByte) == EPropertyBagResult::Success, "FAnimationAnimNextRuntimeTest_Variables -> GetVariable failed");
+			UE_RETURN_ON_ERROR(TestByte == 1, "FAnimationAnimNextRuntimeTest_Variables -> Variable value did not match");
+
+			TestInt32 = 0;
+			UE_RETURN_ON_ERROR(GraphInstance.GetImpl()->GetVariable("Bool", TestInt32) == EPropertyBagResult::Success, "FAnimationAnimNextRuntimeTest_Variables -> GetVariable failed");
+			UE_RETURN_ON_ERROR(TestInt32 == 1, "FAnimationAnimNextRuntimeTest_Variables -> Variable value did not match");
+
+			TestUInt32 = 0;
+			UE_RETURN_ON_ERROR(GraphInstance.GetImpl()->GetVariable("Bool", TestUInt32) == EPropertyBagResult::Success, "FAnimationAnimNextRuntimeTest_Variables -> GetVariable failed");
+			UE_RETURN_ON_ERROR(TestUInt32 == 1, "FAnimationAnimNextRuntimeTest_Variables -> Variable value did not match");
+
+			TestInt64 = 0;
+			UE_RETURN_ON_ERROR(GraphInstance.GetImpl()->GetVariable("Bool", TestInt64) == EPropertyBagResult::Success, "FAnimationAnimNextRuntimeTest_Variables -> GetVariable failed");
+			UE_RETURN_ON_ERROR(TestInt64 == 1, "FAnimationAnimNextRuntimeTest_Variables -> Variable value did not match");
+
+			TestUInt64 = 0;
+			UE_RETURN_ON_ERROR(GraphInstance.GetImpl()->GetVariable("Bool", TestUInt64) == EPropertyBagResult::Success, "FAnimationAnimNextRuntimeTest_Variables -> GetVariable failed");
+			UE_RETURN_ON_ERROR(TestUInt64 == 1, "FAnimationAnimNextRuntimeTest_Variables -> Variable value did not match");
+
+			// Byte
+			TestBool = false;
+			UE_RETURN_ON_ERROR(GraphInstance.GetImpl()->GetVariable("Byte", TestBool) == EPropertyBagResult::Success, "FAnimationAnimNextRuntimeTest_Variables -> GetVariable failed");
+			UE_RETURN_ON_ERROR(TestBool == true, "FAnimationAnimNextRuntimeTest_Variables -> Variable value did not match");
+
+			TestInt32 = 0;
+			UE_RETURN_ON_ERROR(GraphInstance.GetImpl()->GetVariable("Byte", TestInt32) == EPropertyBagResult::Success, "FAnimationAnimNextRuntimeTest_Variables -> GetVariable failed");
+			UE_RETURN_ON_ERROR(TestInt32 == 42, "FAnimationAnimNextRuntimeTest_Variables -> Variable value did not match");
+
+			TestUInt32 = 0;
+			UE_RETURN_ON_ERROR(GraphInstance.GetImpl()->GetVariable("Byte", TestUInt32) == EPropertyBagResult::Success, "FAnimationAnimNextRuntimeTest_Variables -> GetVariable failed");
+			UE_RETURN_ON_ERROR(TestUInt32 == 42, "FAnimationAnimNextRuntimeTest_Variables -> Variable value did not match");
+
+			TestInt64 = 0;
+			UE_RETURN_ON_ERROR(GraphInstance.GetImpl()->GetVariable("Byte", TestInt64) == EPropertyBagResult::Success, "FAnimationAnimNextRuntimeTest_Variables -> GetVariable failed");
+			UE_RETURN_ON_ERROR(TestInt64 == 42, "FAnimationAnimNextRuntimeTest_Variables -> Variable value did not match");
+
+			TestUInt64 = 0;
+			UE_RETURN_ON_ERROR(GraphInstance.GetImpl()->GetVariable("Byte", TestUInt64) == EPropertyBagResult::Success, "FAnimationAnimNextRuntimeTest_Variables -> GetVariable failed");
+			UE_RETURN_ON_ERROR(TestUInt64 == 42, "FAnimationAnimNextRuntimeTest_Variables -> Variable value did not match");
+
+			// Int32
+			TestBool = false;
+			UE_RETURN_ON_ERROR(GraphInstance.GetImpl()->GetVariable("Int32", TestBool) == EPropertyBagResult::Success, "FAnimationAnimNextRuntimeTest_Variables -> GetVariable failed");
+			UE_RETURN_ON_ERROR(TestBool == true, "FAnimationAnimNextRuntimeTest_Variables -> Variable value did not match");
+
+			TestByte = 0;
+			UE_RETURN_ON_ERROR(GraphInstance.GetImpl()->GetVariable("Int32", TestByte) == EPropertyBagResult::Success, "FAnimationAnimNextRuntimeTest_Variables -> GetVariable failed");
+			UE_RETURN_ON_ERROR(TestByte == 202/*(uint8)-4679222*/, "FAnimationAnimNextRuntimeTest_Variables -> Variable value did not match");
+
+			TestUInt32 = 0;
+			UE_RETURN_ON_ERROR(GraphInstance.GetImpl()->GetVariable("Int32", TestUInt32) == EPropertyBagResult::Success, "FAnimationAnimNextRuntimeTest_Variables -> GetVariable failed");
+			UE_RETURN_ON_ERROR(TestUInt32 == 4290288074/*(uint32)-4679222*/, "FAnimationAnimNextRuntimeTest_Variables -> Variable value did not match");
+
+			TestInt64 = 0;
+			UE_RETURN_ON_ERROR(GraphInstance.GetImpl()->GetVariable("Int32", TestInt64) == EPropertyBagResult::Success, "FAnimationAnimNextRuntimeTest_Variables -> GetVariable failed");
+			UE_RETURN_ON_ERROR(TestInt64 == -4679222, "FAnimationAnimNextRuntimeTest_Variables -> Variable value did not match");
+
+			TestUInt64 = 0;
+			UE_RETURN_ON_ERROR(GraphInstance.GetImpl()->GetVariable("Int32", TestUInt64) == EPropertyBagResult::Success, "FAnimationAnimNextRuntimeTest_Variables -> GetVariable failed");
+			UE_RETURN_ON_ERROR(TestUInt64 == 4290288074/*(uint64)-4679222*/, "FAnimationAnimNextRuntimeTest_Variables -> Variable value did not match");
+
+			// Uint32
+			TestBool = false;
+			UE_RETURN_ON_ERROR(GraphInstance.GetImpl()->GetVariable("UInt32", TestBool) == EPropertyBagResult::Success, "FAnimationAnimNextRuntimeTest_Variables -> GetVariable failed");
+			UE_RETURN_ON_ERROR(TestBool == true, "FAnimationAnimNextRuntimeTest_Variables -> Variable value did not match");
+
+			TestByte = 0;
+			UE_RETURN_ON_ERROR(GraphInstance.GetImpl()->GetVariable("UInt32", TestByte) == EPropertyBagResult::Success, "FAnimationAnimNextRuntimeTest_Variables -> GetVariable failed");
+			UE_RETURN_ON_ERROR(TestByte == 255/*(uint8)3415919103*/, "FAnimationAnimNextRuntimeTest_Variables -> Variable value did not match");
+
+			TestInt32 = 0;
+			UE_RETURN_ON_ERROR(GraphInstance.GetImpl()->GetVariable("UInt32", TestInt32) == EPropertyBagResult::Success, "FAnimationAnimNextRuntimeTest_Variables -> GetVariable failed");
+			UE_RETURN_ON_ERROR(TestInt32 == -879048193/*(int32)3415919103*/, "FAnimationAnimNextRuntimeTest_Variables -> Variable value did not match");
+
+			TestInt64 = 0;
+			UE_RETURN_ON_ERROR(GraphInstance.GetImpl()->GetVariable("UInt32", TestInt64) == EPropertyBagResult::Success, "FAnimationAnimNextRuntimeTest_Variables -> GetVariable failed");
+			UE_RETURN_ON_ERROR(TestInt64 == 3415919103, "FAnimationAnimNextRuntimeTest_Variables -> Variable value did not match");
+
+			TestUInt64 = 0;
+			UE_RETURN_ON_ERROR(GraphInstance.GetImpl()->GetVariable("UInt32", TestUInt64) == EPropertyBagResult::Success, "FAnimationAnimNextRuntimeTest_Variables -> GetVariable failed");
+			UE_RETURN_ON_ERROR(TestUInt64 == 3415919103, "FAnimationAnimNextRuntimeTest_Variables -> Variable value did not match");
+
+			// Int64
+			TestBool = false;
+			UE_RETURN_ON_ERROR(GraphInstance.GetImpl()->GetVariable("Int64", TestBool) == EPropertyBagResult::Success, "FAnimationAnimNextRuntimeTest_Variables -> GetVariable failed");
+			UE_RETURN_ON_ERROR(TestBool == true, "FAnimationAnimNextRuntimeTest_Variables -> Variable value did not match");
+
+			TestByte = 0;
+			UE_RETURN_ON_ERROR(GraphInstance.GetImpl()->GetVariable("Int64", TestByte) == EPropertyBagResult::Success, "FAnimationAnimNextRuntimeTest_Variables -> GetVariable failed");
+			UE_RETURN_ON_ERROR(TestByte == 255/*(uint8)-3415919105*/, "FAnimationAnimNextRuntimeTest_Variables -> Variable value did not match");
+
+			TestInt32 = 0;
+			UE_RETURN_ON_ERROR(GraphInstance.GetImpl()->GetVariable("Int64", TestInt32) == EPropertyBagResult::Success, "FAnimationAnimNextRuntimeTest_Variables -> GetVariable failed");
+			UE_RETURN_ON_ERROR(TestInt32 == 879048191/*(int32)-3415919105*/, "FAnimationAnimNextRuntimeTest_Variables -> Variable value did not match");
+
+			TestUInt32 = 0;
+			UE_RETURN_ON_ERROR(GraphInstance.GetImpl()->GetVariable("Int64", TestUInt32) == EPropertyBagResult::Success, "FAnimationAnimNextRuntimeTest_Variables -> GetVariable failed");
+			UE_RETURN_ON_ERROR(TestUInt32 == 879048191/*(uint32)-3415919105*/, "FAnimationAnimNextRuntimeTest_Variables -> Variable value did not match");
+
+			TestUInt64 = 0;
+			UE_RETURN_ON_ERROR(GraphInstance.GetImpl()->GetVariable("Int64", TestUInt64) == EPropertyBagResult::Success, "FAnimationAnimNextRuntimeTest_Variables -> GetVariable failed");
+			UE_RETURN_ON_ERROR(TestUInt64 == 18446744070293632511ull/*(uint64)-3415919105*/, "FAnimationAnimNextRuntimeTest_Variables -> Variable value did not match");
+		}
+
+		// Float/double + conversions
+		{
+			// Gets
+			float TestFloat = 0.0f;
+			UE_RETURN_ON_ERROR(GraphInstance.GetImpl()->GetVariable("Float", TestFloat) == EPropertyBagResult::Success, "FAnimationAnimNextRuntimeTest_Variables -> GetVariable failed");
+			UE_RETURN_ON_ERROR(TestFloat == 1.0f, "FAnimationAnimNextRuntimeTest_Variables -> Variable value did not match");
+
+			double TestDouble = 0.0;
+			UE_RETURN_ON_ERROR(GraphInstance.GetImpl()->GetVariable("Double", TestDouble) == EPropertyBagResult::Success, "FAnimationAnimNextRuntimeTest_Variables -> GetVariable failed");
+			UE_RETURN_ON_ERROR(TestDouble == 1.0, "FAnimationAnimNextRuntimeTest_Variables -> Variable value did not match");
+
+			// Conversions
+			TestDouble = 0.0f;
+			UE_RETURN_ON_ERROR(GraphInstance.GetImpl()->GetVariable("Float", TestDouble) == EPropertyBagResult::Success, "FAnimationAnimNextRuntimeTest_Variables -> GetVariable failed");
+			UE_RETURN_ON_ERROR(TestDouble == 1.0f, "FAnimationAnimNextRuntimeTest_Variables -> Variable value did not match");
+
+			TestFloat = 0.0;
+			UE_RETURN_ON_ERROR(GraphInstance.GetImpl()->GetVariable("Double", TestFloat) == EPropertyBagResult::Success, "FAnimationAnimNextRuntimeTest_Variables -> GetVariable failed");
+			UE_RETURN_ON_ERROR(TestFloat == 1.0, "FAnimationAnimNextRuntimeTest_Variables -> Variable value did not match");
+		}
+		
+		FName TestFName = NAME_None;
+		UE_RETURN_ON_ERROR(GraphInstance.GetImpl()->GetVariable("Name", TestFName) == EPropertyBagResult::Success, "FAnimationAnimNextRuntimeTest_Variables -> GetVariable failed");
+		UE_RETURN_ON_ERROR(TestFName == "Test", "FAnimationAnimNextRuntimeTest_Variables -> Variable value did not match");
+
+		FString TestString;
+		UE_RETURN_ON_ERROR(GraphInstance.GetImpl()->GetVariable("String", TestString) == EPropertyBagResult::Success, "FAnimationAnimNextRuntimeTest_Variables -> GetVariable failed");
+		UE_RETURN_ON_ERROR(TestString == "Test", "FAnimationAnimNextRuntimeTest_Variables -> Variable value did not match");
+
+		FText TestText;
+		UE_RETURN_ON_ERROR(GraphInstance.GetImpl()->GetVariable("Text", TestText) == EPropertyBagResult::Success, "FAnimationAnimNextRuntimeTest_Variables -> GetVariable failed");
+		UE_RETURN_ON_ERROR(TestText.EqualTo(NSLOCTEXT("Tests", "Test", "Test")), "FAnimationAnimNextRuntimeTest_Variables -> Variable value did not match");
+
+		// Enums
+		{
+			EPropertyBagPropertyType TestEnum = EPropertyBagPropertyType::None;
+			UE_RETURN_ON_ERROR(GraphInstance.GetImpl()->GetVariable("Enum", TestEnum) == EPropertyBagResult::Success, "FAnimationAnimNextRuntimeTest_Variables -> GetVariable failed");
+			UE_RETURN_ON_ERROR(TestEnum == EPropertyBagPropertyType::Double, "FAnimationAnimNextRuntimeTest_Variables -> Variable value did not match");
+
+			// Test mismatched enum fails
+			EPropertyBagContainerType MismatchedEnum = EPropertyBagContainerType::None;
+			UE_RETURN_ON_ERROR(GraphInstance.GetImpl()->GetVariable("Enum", MismatchedEnum) != EPropertyBagResult::Success, "FAnimationAnimNextRuntimeTest_Variables -> GetVariable succeeded");
+		}
+
+		// Structs
+		{
+			FVector TestStruct = FVector::ZeroVector;
+			UE_RETURN_ON_ERROR(GraphInstance.GetImpl()->GetVariable("Struct", TestStruct) == EPropertyBagResult::Success, "FAnimationAnimNextRuntimeTest_Variables -> GetVariable failed");
+			UE_RETURN_ON_ERROR(TestStruct == FVector::OneVector, "FAnimationAnimNextRuntimeTest_Variables -> Variable value did not match");
+
+			// Test conversion from base -> derived fails 
+			FTestDerivedVector TestDerived;
+			UE_RETURN_ON_ERROR(GraphInstance.GetImpl()->GetVariable("Struct", TestDerived) != EPropertyBagResult::Success, "FAnimationAnimNextRuntimeTest_Variables -> GetVariable succeeded");
+
+			// Test conversion from derived -> base
+			UE_RETURN_ON_ERROR(GraphInstance.GetImpl()->GetVariable("DerivedStruct", TestStruct) == EPropertyBagResult::Success, "FAnimationAnimNextRuntimeTest_Variables -> GetVariable failed");
+			UE_RETURN_ON_ERROR(TestStruct == FVector::OneVector, "FAnimationAnimNextRuntimeTest_Variables -> Variable value did not match");
+		}
+
+		// Objects
+		{
+			UAnimNextDataInterface* TestObject = nullptr;
+			UE_RETURN_ON_ERROR(GraphInstance.GetImpl()->GetVariable("Object", TestObject) == EPropertyBagResult::Success, "FAnimationAnimNextRuntimeTest_Variables -> GetVariable failed");
+			UE_RETURN_ON_ERROR(TestObject == UAnimNextDataInterface::StaticClass()->GetDefaultObject<UAnimNextDataInterface>(), "FAnimationAnimNextRuntimeTest_Variables -> Variable value did not match");
+
+			// Test unrelated object fails
+			UAnimSequence* TestAnimSequence = nullptr;
+			UE_RETURN_ON_ERROR(GraphInstance.GetImpl()->GetVariable("Object", TestAnimSequence) != EPropertyBagResult::Success, "FAnimationAnimNextRuntimeTest_Variables -> GetVariable succeeded");
+
+			// Test Derived -> Base succeeds
+			UObject* BaseObject = nullptr;
+			UE_RETURN_ON_ERROR(GraphInstance.GetImpl()->GetVariable("Object", BaseObject) == EPropertyBagResult::Success, "FAnimationAnimNextRuntimeTest_Variables -> GetVariable failed");
+			UE_RETURN_ON_ERROR(BaseObject == UAnimNextDataInterface::StaticClass()->GetDefaultObject<UAnimNextDataInterface>(), "FAnimationAnimNextRuntimeTest_Variables -> Variable value did not match");
+
+			// Test Base -> Derived fails
+			UAnimNextAnimationGraph* DerivedObject = nullptr;
+			UE_RETURN_ON_ERROR(GraphInstance.GetImpl()->GetVariable("Object", DerivedObject) != EPropertyBagResult::Success, "FAnimationAnimNextRuntimeTest_Variables -> GetVariable succeeded");
+		}
+
+		FSoftObjectPath TestSoftObjectPath;
+		UE_RETURN_ON_ERROR(GraphInstance.GetImpl()->GetVariable("SoftObject", TestSoftObjectPath) == EPropertyBagResult::Success, "FAnimationAnimNextRuntimeTest_Variables -> GetVariable failed");
+		UE_RETURN_ON_ERROR(TestSoftObjectPath == FSoftObjectPath(UAnimNextDataInterface::StaticClass()->GetDefaultObject<UAnimNextDataInterface>()), "FAnimationAnimNextRuntimeTest_Variables -> Variable value did not match");
+
+		// Classes
+		{
+			UClass* TestClass = nullptr;
+			UE_RETURN_ON_ERROR(GraphInstance.GetImpl()->GetVariable("Class", TestClass) == EPropertyBagResult::Success, "FAnimationAnimNextRuntimeTest_Variables -> GetVariable failed");
+			UE_RETURN_ON_ERROR(TestClass == UAnimNextDataInterface::StaticClass(), "FAnimationAnimNextRuntimeTest_Variables -> Variable value did not match");
+
+
+			// Test unrelated class fails
+			TSubclassOf<UAnimSequence> TestUnrelatedSubclassOf;
+			UE_RETURN_ON_ERROR(GraphInstance.GetImpl()->GetVariable("Class", TestUnrelatedSubclassOf) != EPropertyBagResult::Success, "FAnimationAnimNextRuntimeTest_Variables -> GetVariable succeeded");
+
+			// Test Derived -> Base succeeds
+			TSubclassOf<UObject> BaseSubclassOf;
+			UE_RETURN_ON_ERROR(GraphInstance.GetImpl()->GetVariable("Class", BaseSubclassOf) == EPropertyBagResult::Success, "FAnimationAnimNextRuntimeTest_Variables -> GetVariable failed");
+			UE_RETURN_ON_ERROR(BaseSubclassOf == UAnimNextDataInterface::StaticClass(), "FAnimationAnimNextRuntimeTest_Variables -> Variable value did not match");
+
+			// Test Base -> Derived fails
+			TSubclassOf<UAnimNextAnimationGraph> DerivedSubclassOf;
+			UE_RETURN_ON_ERROR(GraphInstance.GetImpl()->GetVariable("Class", DerivedSubclassOf) != EPropertyBagResult::Success, "FAnimationAnimNextRuntimeTest_Variables -> GetVariable succeeded");
+		}
+
+		FSoftClassPath TestSoftClassPath;
+		UE_RETURN_ON_ERROR(GraphInstance.GetImpl()->GetVariable("SoftClass", TestSoftClassPath) == EPropertyBagResult::Success, "FAnimationAnimNextRuntimeTest_Variables -> GetVariable failed");
+		UE_RETURN_ON_ERROR(TestSoftClassPath == FSoftClassPath(UAnimNextDataInterface::StaticClass()), "FAnimationAnimNextRuntimeTest_Variables -> Variable value did not match");
+	}
+
+	Tests::FUtils::CleanupAfterTests();
+
+	return true;
+}
+
+
 #endif
diff --git a/Engine/Plugins/Experimental/AnimNextAnimGraph/Source/AnimNextAnimGraphTestSuite/Private/AnimNextAnimGraphTraitGraphTest.h b/Engine/Plugins/Experimental/AnimNextAnimGraph/Source/AnimNextAnimGraphTestSuite/Private/AnimNextAnimGraphTraitGraphTest.h
--- a/Engine/Plugins/Experimental/AnimNextAnimGraph/Source/AnimNextAnimGraphTestSuite/Private/AnimNextAnimGraphTraitGraphTest.h
+++ b/Engine/Plugins/Experimental/AnimNextAnimGraph/Source/AnimNextAnimGraphTestSuite/Private/AnimNextAnimGraphTraitGraphTest.h
@@ -44,3 +44,16 @@ struct FTestTraitSharedData : public FAnimNextTraitSharedData
 	GENERATE_TRAIT_LATENT_PROPERTIES(FTestTraitSharedData, TRAIT_LATENT_PROPERTIES_ENUMERATOR)
 	#undef TRAIT_LATENT_PROPERTIES_ENUMERATOR
 };
+
+USTRUCT()
+struct FTestDerivedVector : public FVector
+{
+	GENERATED_BODY()
+	
+	FTestDerivedVector()
+		: FVector(FVector::OneVector)
+	{}
+
+	UPROPERTY()
+	int32 W = 1;
+};
\ No newline at end of file
diff --git a/Engine/Plugins/Experimental/AnimNextAnimGraph/Source/AnimNextAnimGraphUncookedOnly/Private/PlayAnim/K2Node_AnimNextPlayAsset.cpp b/Engine/Plugins/Experimental/AnimNextAnimGraph/Source/AnimNextAnimGraphUncookedOnly/Private/PlayAnim/K2Node_AnimNextPlayAsset.cpp
new file mode 100644
--- /dev/null
+++ b/Engine/Plugins/Experimental/AnimNextAnimGraph/Source/AnimNextAnimGraphUncookedOnly/Private/PlayAnim/K2Node_AnimNextPlayAsset.cpp
@@ -0,0 +1,31 @@
+﻿// Copyright Epic Games, Inc. All Rights Reserved.
+
+#include "PlayAnim/K2Node_AnimNextPlayAsset.h"
+
+#include "PlayAnim/PlayAnimCallbackProxy.h"
+
+#define LOCTEXT_NAMESPACE "K2Node_AnimNextPlayAsset"
+
+UK2Node_AnimNextPlayAsset::UK2Node_AnimNextPlayAsset()
+{
+	ProxyFactoryFunctionName = GET_FUNCTION_NAME_CHECKED(UPlayAnimCallbackProxy, CreateProxyObjectForPlayAsset);
+	ProxyFactoryClass = UPlayAnimCallbackProxy::StaticClass();
+	ProxyClass = UPlayAnimCallbackProxy::StaticClass();
+}
+
+FText UK2Node_AnimNextPlayAsset::GetTooltipText() const
+{
+	return LOCTEXT("K2Node_PlayAsset_Tooltip", "Plays an asset on an AnimNextComponent");
+}
+
+FText UK2Node_AnimNextPlayAsset::GetNodeTitle(ENodeTitleType::Type TitleType) const
+{
+	return LOCTEXT("PlayAsset", "Play Asset");
+}
+
+FText UK2Node_AnimNextPlayAsset::GetMenuCategory() const
+{
+	return LOCTEXT("PlayAssetCategory", "Animation|AnimNext");
+}
+
+#undef LOCTEXT_NAMESPACE
diff --git a/Engine/Plugins/Experimental/AnimNextChooser/AnimNextChooser.uplugin b/Engine/Plugins/Experimental/AnimNextChooser/AnimNextChooser.uplugin
--- a/Engine/Plugins/Experimental/AnimNextChooser/AnimNextChooser.uplugin
+++ b/Engine/Plugins/Experimental/AnimNextChooser/AnimNextChooser.uplugin
@@ -24,6 +24,10 @@
 		{
 			"Name": "Chooser",
 			"Enabled": true
+		},
+		{
+			"Name": "ControlRig",
+			"Enabled": true
 		}
 	],
 	"Modules" :
@@ -32,6 +36,6 @@
 			"Name" : "AnimNextChooser",
 			"Type" : "Runtime",
 			"LoadingPhase" : "Default"
-		},
+		}
 	]
 }
diff --git a/Engine/Plugins/Experimental/AnimNextChooser/Source/AnimNextChooser/Private/ChooserParameters.cpp b/Engine/Plugins/Experimental/AnimNextChooser/Source/AnimNextChooser/Private/ChooserParameters.cpp
--- a/Engine/Plugins/Experimental/AnimNextChooser/Source/AnimNextChooser/Private/ChooserParameters.cpp
+++ b/Engine/Plugins/Experimental/AnimNextChooser/Source/AnimNextChooser/Private/ChooserParameters.cpp
@@ -1,17 +1,34 @@
 // Copyright Epic Games, Inc. All Rights Reserved.
-#include "ChooserParameters.h"
-//#include "Param/ParamStack.h"
 
-#if WITH_EDITOR
+#include "ChooserParameters.h"
+#include "DataInterface/AnimNextDataInterfaceInstance.h"
 #include "StructUtils/PropertyBag.h"
-#endif
+
+namespace UE::AnimNext::Private
+{
+
+// Helper function to find the first data interface instance in the context 
+FAnimNextDataInterfaceInstance* GetFirstDataInterfaceInstance(FChooserEvaluationContext& Context)
+{
+	for(const FStructView& Param : Context.Params)
+	{
+		if(Param.GetScriptStruct() == FAnimNextDataInterfaceInstance::StaticStruct())
+		{
+			return Param.GetPtr<FAnimNextDataInterfaceInstance>();
+		}
+	}
+
+	return nullptr;
+}
+
+}
 
 bool FBoolAnimProperty::GetValue(FChooserEvaluationContext& Context, bool& OutResult) const
 {
-//	UE::AnimNext::FParamResult Result;
-//	UE::AnimNext::FParamStack& ParamStack = UE::AnimNext::FParamStack::Get();
-//	OutResult = ParamStack.GetParamValue<bool>(UE::AnimNext::FParamId(PropertyName), &Result);
-//	return Result.IsSuccessful();
+	if(FAnimNextDataInterfaceInstance* Instance = UE::AnimNext::Private::GetFirstDataInterfaceInstance(Context))
+	{
+		return Instance->GetVariable(VariableName, OutResult) == EPropertyBagResult::Success;
+	}
 	return false;
 }
 
@@ -22,7 +39,7 @@ bool FBoolAnimProperty::SetValue(FChooserEvaluationContext& Context, bool InValu
 
 void FBoolAnimProperty::GetDisplayName(FText& OutName) const
 {
-	FString DisplayName = PropertyName.ToString();
+	FString DisplayName = VariableName.ToString();
 	int Index = -1;
 	DisplayName.FindLastChar(':', Index);
 	if (Index>=0)
@@ -36,10 +53,10 @@ void FBoolAnimProperty::GetDisplayName(FText& OutName) const
 
 bool FFloatAnimProperty::GetValue(FChooserEvaluationContext& Context, double& OutResult) const
 {
-//	UE::AnimNext::FParamResult Result;
-//	UE::AnimNext::FParamStack& ParamStack = UE::AnimNext::FParamStack::Get();
-//	OutResult = ParamStack.GetParamValue<double>(UE::AnimNext::FParamId(PropertyName), &Result);
-//	return Result.IsSuccessful();
+	if(FAnimNextDataInterfaceInstance* Instance = UE::AnimNext::Private::GetFirstDataInterfaceInstance(Context))
+	{
+		return Instance->GetVariable(VariableName, OutResult) == EPropertyBagResult::Success;
+	}
 	return false;
 }
 
@@ -50,7 +67,7 @@ bool FFloatAnimProperty::SetValue(FChooserEvaluationContext& Context, double InV
 
 void FFloatAnimProperty::GetDisplayName(FText& OutName) const
 {
-	FString DisplayName = PropertyName.ToString();
+	FString DisplayName = VariableName.ToString();
 	int Index = -1;
 	DisplayName.FindLastChar(':', Index);
 	if (Index>=0)
diff --git a/Engine/Plugins/Experimental/AnimNextChooser/Source/AnimNextChooser/Private/ChooserParameters.h b/Engine/Plugins/Experimental/AnimNextChooser/Source/AnimNextChooser/Private/ChooserParameters.h
--- a/Engine/Plugins/Experimental/AnimNextChooser/Source/AnimNextChooser/Private/ChooserParameters.h
+++ b/Engine/Plugins/Experimental/AnimNextChooser/Source/AnimNextChooser/Private/ChooserParameters.h
@@ -12,9 +12,9 @@ struct FBoolAnimProperty :  public FChooserParameterBoolBase
 	GENERATED_BODY()
 public:
 
-	UPROPERTY(DisplayName = "Parameter", EditAnywhere, Category="Parameter", meta=(CustomWidget = "ParamName", AllowedParamType = "bool"))
-	FName PropertyName;
-	
+	UPROPERTY(DisplayName = "Variable", EditAnywhere, Category="Variable")
+	FName VariableName;
+
 	virtual bool GetValue(FChooserEvaluationContext& Context, bool& OutResult) const override;
 	virtual bool SetValue(FChooserEvaluationContext& Context, bool InValue) const override;
 	virtual void GetDisplayName(FText& OutName) const override;
@@ -26,9 +26,9 @@ struct FFloatAnimProperty :  public FChooserParameterFloatBase
 	GENERATED_BODY()
 public:
 
-	UPROPERTY(DisplayName = "Parameter", EditAnywhere, Category="Parameter", meta=(CustomWidget = "ParamName", AllowedParamType = "double"))
-	FName PropertyName;
-	
+	UPROPERTY(DisplayName = "Variable", EditAnywhere, Category="Variable")
+	FName VariableName;
+
 	virtual bool GetValue(FChooserEvaluationContext& Context, double& OutResult) const override;
 	virtual bool SetValue(FChooserEvaluationContext& Context, double InValue) const override;
 	virtual void GetDisplayName(FText& OutName) const override;
diff --git a/Engine/Plugins/Experimental/AnimNextChooser/Source/AnimNextChooser/Private/RigUnit_EvaluateChooser.cpp b/Engine/Plugins/Experimental/AnimNextChooser/Source/AnimNextChooser/Private/RigUnit_EvaluateChooser.cpp
--- a/Engine/Plugins/Experimental/AnimNextChooser/Source/AnimNextChooser/Private/RigUnit_EvaluateChooser.cpp
+++ b/Engine/Plugins/Experimental/AnimNextChooser/Source/AnimNextChooser/Private/RigUnit_EvaluateChooser.cpp
@@ -2,21 +2,45 @@
 
 
 #include "RigUnit_EvaluateChooser.h"
+#include "DataInterface/AnimNextDataInterfaceInstance.h"
+#include "ControlRig.h"
 
-
-FRigUnit_EvaluateChooser_Execute()
+static void RunChooserHelper(TConstArrayView<TObjectPtr<UObject>> ContextObjects, FStructView ContextStruct, TObjectPtr<UChooserTable> Chooser, TObjectPtr<UObject>& OutResult)
 {
-	Result = nullptr;
-
-	if (ContextObject && Chooser)
+	if ((!ContextObjects.IsEmpty() || ContextStruct.IsValid()) && Chooser != nullptr)
 	{
 		FChooserEvaluationContext ChooserContext;
-		ChooserContext.AddObjectParam(ContextObject);
-
-		UChooserTable::EvaluateChooser(ChooserContext, Chooser, FObjectChooserBase::FObjectChooserIteratorCallback::CreateLambda([&Result](UObject* InResult)
+		for(TObjectPtr<UObject> ContextObject : ContextObjects)
+		{
+			if(ContextObject != nullptr)
 			{
-				Result = InResult;
-				return FObjectChooserBase::EIteratorStatus::Stop;
-			}));
+				ChooserContext.AddObjectParam(ContextObject);
+			}
+		}
+		if(ContextStruct.IsValid())
+		{
+			ChooserContext.AddStructViewParam(ContextStruct);
+		}
+
+		UChooserTable::EvaluateChooser(ChooserContext, Chooser, FObjectChooserBase::FObjectChooserIteratorCallback::CreateLambda([&OutResult](UObject* InResult)
+		{
+			OutResult = InResult;
+			return FObjectChooserBase::EIteratorStatus::Stop;
+		}));
 	}
 }
+
+FRigUnit_EvaluateChooser_ControlRig_Execute()
+{
+	Result = nullptr;
+
+	RunChooserHelper( { ContextObject, ExecuteContext.ControlRig }, FStructView(), Chooser, Result);
+}
+
+FRigUnit_EvaluateChooser_AnimNext_Execute()
+{
+	Result = nullptr;
+
+	const FAnimNextDataInterfaceInstance& Instance = ExecuteContext.GetInstance();
+	RunChooserHelper({ ContextObject }, FStructView::Make(const_cast<FAnimNextDataInterfaceInstance&>(Instance)), Chooser, Result);
+}
\ No newline at end of file
diff --git a/Engine/Plugins/Experimental/AnimNextChooser/Source/AnimNextChooser/Private/RigUnit_EvaluateChooser.h b/Engine/Plugins/Experimental/AnimNextChooser/Source/AnimNextChooser/Private/RigUnit_EvaluateChooser.h
--- a/Engine/Plugins/Experimental/AnimNextChooser/Source/AnimNextChooser/Private/RigUnit_EvaluateChooser.h
+++ b/Engine/Plugins/Experimental/AnimNextChooser/Source/AnimNextChooser/Private/RigUnit_EvaluateChooser.h
@@ -4,9 +4,9 @@
 #pragma once
 
 #include "CoreMinimal.h"
+#include "AnimNextExecuteContext.h"
 #include "Chooser.h"
 #include "ControlRigDefines.h"
-#include "RigVMFunctions/Animation/RigVMFunction_AnimBase.h"
 #include "Units/RigUnit.h"
 
 #include "RigUnit_EvaluateChooser.generated.h"
@@ -14,19 +14,11 @@
 /*
  * Evaluates a Chooser Table and outputs the selected UObject
  */
-USTRUCT(meta = (DisplayName = "Evaluate Chooser", Keywords="Evaluate Chooser", Category = "Chooser", Varying, NodeColor="0.737911 0.099899 0.099899"))
-struct FRigUnit_EvaluateChooser : public FRigVMFunction_AnimBase
+USTRUCT(meta = (Abstract, Keywords="Evaluate Chooser", Category = "Chooser", Varying, NodeColor="0.737911 0.099899 0.099899"))
+struct FRigUnit_EvaluateChooser : public FRigVMStruct
 {
 	GENERATED_BODY()
 
-	FRigUnit_EvaluateChooser()
-	{
-	}
-
-	/** Execute logic for this rig unit */
-	RIGVM_METHOD()
-	virtual void Execute() override;
-
 	UPROPERTY(meta = (Input))
 	TObjectPtr<UObject> ContextObject;
 	
@@ -35,4 +27,30 @@ struct FRigUnit_EvaluateChooser : public FRigVMFunction_AnimBase
 
 	UPROPERTY(meta = (Output))
 	TObjectPtr<UObject> Result;
+};
+
+/*
+ * Evaluates a Chooser Table in the context of ControlRig
+ */
+USTRUCT(meta = (DisplayName = "Evaluate Chooser", Varying, NodeColor="0.737911 0.099899 0.099899", ExecuteContext="FControlRigExecuteContext"))
+struct FRigUnit_EvaluateChooser_ControlRig : public FRigUnit_EvaluateChooser
+{
+	GENERATED_BODY()
+
+	/** Execute logic for this rig unit */
+	RIGVM_METHOD()
+	virtual void Execute() override;
+};
+
+/*
+ * Evaluates a Chooser Table in the context of AnimNext
+ */
+USTRUCT(meta = (DisplayName = "Evaluate Chooser", Varying, NodeColor="0.737911 0.099899 0.099899", ExecuteContext="FAnimNextExecuteContext"))
+struct FRigUnit_EvaluateChooser_AnimNext : public FRigUnit_EvaluateChooser
+{
+	GENERATED_BODY()
+
+	/** Execute logic for this rig unit */
+	RIGVM_METHOD()
+	virtual void Execute() override;
 };
\ No newline at end of file
diff --git a/Engine/Plugins/Experimental/Chooser/Source/Chooser/Public/IObjectChooser.h b/Engine/Plugins/Experimental/Chooser/Source/Chooser/Public/IObjectChooser.h
--- a/Engine/Plugins/Experimental/Chooser/Source/Chooser/Public/IObjectChooser.h
+++ b/Engine/Plugins/Experimental/Chooser/Source/Chooser/Public/IObjectChooser.h
@@ -64,12 +64,19 @@ struct CHOOSER_API FChooserEvaluationContext
 		AddStructParam(ObjectParams.Last());
 	}
 
+	// Add a struct view Parameter to the Context
+	// the struct will be referred to by reference, and so must have a lifetime that is longer than this context
+	void AddStructViewParam(FStructView Param)
+	{
+		Params.Add(Param);
+	}
+	
 	// Add a struct Parameter to the Context
 	// the struct will be referred to by reference, and so must have a lifetime that is longer than this context
 	template <class T>
 	void AddStructParam(T& Param)
 	{
-		Params.Add(FStructView::Make(Param));
+		AddStructViewParam(FStructView::Make(Param));
 	}
 
 	#if CHOOSER_DEBUGGING_ENABLED
diff --git a/Engine/Plugins/Runtime/RigVM/Source/RigVM/Private/RigVMPropertyUtils.cpp b/Engine/Plugins/Runtime/RigVM/Source/RigVM/Private/RigVMPropertyUtils.cpp
--- a/Engine/Plugins/Runtime/RigVM/Source/RigVM/Private/RigVMPropertyUtils.cpp
+++ b/Engine/Plugins/Runtime/RigVM/Source/RigVM/Private/RigVMPropertyUtils.cpp
@@ -25,6 +25,16 @@ void RigVMPropertyUtils::GetTypeFromProperty(const FProperty* InProperty, FName&
 		OutTypeName = TEXT("int64");
 		OutTypeObject = nullptr;
 	}
+	else if (CastField<FUInt32Property>(InProperty))
+	{
+		OutTypeName = TEXT("uint32");
+		OutTypeObject = nullptr;
+	}
+	else if (CastField<FUInt64Property>(InProperty))
+	{
+		OutTypeName = TEXT("uint64");
+		OutTypeObject = nullptr;
+	}
 	else if (CastField<FFloatProperty>(InProperty))
 	{
 		OutTypeName = TEXT("float");
diff --git a/Engine/Source/Editor/StructUtilsEditor/Public/SInstancedStructPicker.h b/Engine/Source/Editor/StructUtilsEditor/Public/SInstancedStructPicker.h
--- a/Engine/Source/Editor/StructUtilsEditor/Public/SInstancedStructPicker.h
+++ b/Engine/Source/Editor/StructUtilsEditor/Public/SInstancedStructPicker.h
@@ -43,7 +43,7 @@ class STRUCTUTILSEDITOR_API SInstancedStructPicker : public SCompoundWidget
 {
 public:
 	SLATE_BEGIN_ARGS(SInstancedStructPicker) { }
-		SLATE_ARGUMENT(FOnStructPicked, OnStructPicked)
+		SLATE_EVENT(FOnStructPicked, OnStructPicked)
 	SLATE_END_ARGS()
 
 	void Construct(const FArguments& InArgs, TSharedPtr<IPropertyHandle> InStructProperty, TSharedPtr<IPropertyUtilities> InPropertyUtils);
