Add ability for statetree to read playback percentage

#tests PIE

diff --git a/Engine/Plugins/Experimental/AnimNextAnimGraph/Source/AnimNextAnimGraph/Internal/Traits/BlendStackTrait.h b/Engine/Plugins/Experimental/AnimNextAnimGraph/Source/AnimNextAnimGraph/Internal/Traits/BlendStackTrait.h
--- a/Engine/Plugins/Experimental/AnimNextAnimGraph/Source/AnimNextAnimGraph/Internal/Traits/BlendStackTrait.h
+++ b/Engine/Plugins/Experimental/AnimNextAnimGraph/Source/AnimNextAnimGraph/Internal/Traits/BlendStackTrait.h
@@ -2,14 +2,16 @@
 
 #pragma once
 
+#include "Graph/AnimNextAnimationGraph.h"
+#include "Graph/AnimNextGraphInstancePtr.h"
+
 #include "TraitCore/Trait.h"
+#include "TraitInterfaces/IBlendStack.h"
 #include "TraitInterfaces/IDiscreteBlend.h"
 #include "TraitInterfaces/IGarbageCollection.h"
 #include "TraitInterfaces/IHierarchy.h"
 #include "TraitInterfaces/IUpdate.h"
-#include "Graph/AnimNextAnimationGraph.h"
-#include "Graph/AnimNextGraphInstancePtr.h"
-#include "TraitInterfaces/IBlendStack.h"
+#include "TraitInterfaces/ITimeline.h"
 #include "TraitInterfaces/ISmoothBlend.h"
 
 #include "BlendStackTrait.generated.h"
@@ -83,7 +85,7 @@ namespace UE::AnimNext
 	 * 
 	 * Used as base trait in state machine-like setups where state machine is the additive trait that pushes graphs and blend settings to the blend stack.
 	 */
-	struct FBlendStackCoreTrait : FBaseTrait, IUpdateTraversal, IHierarchy, IDiscreteBlend, ISmoothBlend, IGarbageCollection, IBlendStack
+	struct FBlendStackCoreTrait : FBaseTrait, IUpdateTraversal, IHierarchy, IDiscreteBlend, ISmoothBlend, ITimeline, IGarbageCollection, IBlendStack
 	{
 		DECLARE_ANIM_TRAIT(FBlendStackCoreTrait, 0xda7b7f8d, FBaseTrait)
 
@@ -151,6 +153,10 @@ namespace UE::AnimNext
 		virtual EAlphaBlendOption GetBlendType(FExecutionContext& Context, const TTraitBinding<ISmoothBlend>& Binding, int32 ChildIndex) const override;
 		virtual UCurveFloat* GetCustomBlendCurve(FExecutionContext& Context, const TTraitBinding<ISmoothBlend>& Binding, int32 ChildIndex) const override;
 
+		// ITimeline impl
+		virtual float GetPlayRate(FExecutionContext& Context, const TTraitBinding<ITimeline>& Binding) const;
+		virtual FTimelineProgress GetProgress(FExecutionContext& Context, const TTraitBinding<ITimeline>& Binding) const;
+
 		// IGarbageCollection impl
 		virtual void AddReferencedObjects(const FExecutionContext& Context, const TTraitBinding<IGarbageCollection>& Binding, FReferenceCollector& Collector) const override;
 
diff --git a/Engine/Plugins/Experimental/AnimNextAnimGraph/Source/AnimNextAnimGraph/Private/Traits/BlendStackTrait.cpp b/Engine/Plugins/Experimental/AnimNextAnimGraph/Source/AnimNextAnimGraph/Private/Traits/BlendStackTrait.cpp
--- a/Engine/Plugins/Experimental/AnimNextAnimGraph/Source/AnimNextAnimGraph/Private/Traits/BlendStackTrait.cpp
+++ b/Engine/Plugins/Experimental/AnimNextAnimGraph/Source/AnimNextAnimGraph/Private/Traits/BlendStackTrait.cpp
@@ -16,6 +16,7 @@ namespace UE::AnimNext
 		GeneratorMacro(IGarbageCollection) \
 		GeneratorMacro(IHierarchy) \
 		GeneratorMacro(ISmoothBlend) \
+		GeneratorMacro(ITimeline) \
 		GeneratorMacro(IUpdateTraversal) \
 		GeneratorMacro(IBlendStack) \
 
@@ -211,6 +212,48 @@ namespace UE::AnimNext
 		}
 	}
 
+	float FBlendStackCoreTrait::GetPlayRate(FExecutionContext& Context, const TTraitBinding<ITimeline>& Binding) const
+	{
+		FInstanceData* InstanceData = Binding.GetInstanceData<FInstanceData>();
+		const int32 CurrentlyActiveGraphIndex = InstanceData->CurrentlyActiveGraphIndex;
+		if (CurrentlyActiveGraphIndex != INDEX_NONE)
+		{
+			FBlendStackCoreTrait::FGraphState& ActiveGraph = InstanceData->ChildGraphs[CurrentlyActiveGraphIndex];
+
+			FTraitStackBinding ChildTraitStack;
+			ensure(Context.GetStack(ActiveGraph.ChildPtr, ChildTraitStack));
+
+			UE::AnimNext::TTraitBinding<UE::AnimNext::ITimeline> Timeline;
+			if (ChildTraitStack.GetInterface<UE::AnimNext::ITimeline>(Timeline))
+			{
+				return Timeline.GetPlayRate(Context);
+			}
+		}
+
+		return 1.0f;
+	}
+
+	FTimelineProgress FBlendStackCoreTrait::GetProgress(FExecutionContext& Context, const TTraitBinding<ITimeline>& Binding) const
+	{
+		FInstanceData* InstanceData = Binding.GetInstanceData<FInstanceData>();
+		const int32 CurrentlyActiveGraphIndex = InstanceData->CurrentlyActiveGraphIndex;
+		if (CurrentlyActiveGraphIndex != INDEX_NONE)
+		{
+			FBlendStackCoreTrait::FGraphState& ActiveGraph = InstanceData->ChildGraphs[CurrentlyActiveGraphIndex];
+
+			FTraitStackBinding ChildTraitStack;
+			ensure(Context.GetStack(ActiveGraph.ChildPtr, ChildTraitStack));
+
+			UE::AnimNext::TTraitBinding<UE::AnimNext::ITimeline> Timeline;
+			if (ChildTraitStack.GetInterface<UE::AnimNext::ITimeline>(Timeline))
+			{
+				return Timeline.GetProgress(Context);
+			}
+		}
+
+		return FTimelineProgress();
+	}
+
 	void FBlendStackCoreTrait::AddReferencedObjects(const FExecutionContext& Context, const TTraitBinding<IGarbageCollection>& Binding, FReferenceCollector& Collector) const
 	{
 		IGarbageCollection::AddReferencedObjects(Context, Binding, Collector);
diff --git a/Engine/Plugins/Experimental/AnimNextStateTree/Source/AnimNextStateTree/Private/AnimNextStateTreeContext.cpp b/Engine/Plugins/Experimental/AnimNextStateTree/Source/AnimNextStateTree/Private/AnimNextStateTreeContext.cpp
--- a/Engine/Plugins/Experimental/AnimNextStateTree/Source/AnimNextStateTree/Private/AnimNextStateTreeContext.cpp
+++ b/Engine/Plugins/Experimental/AnimNextStateTree/Source/AnimNextStateTree/Private/AnimNextStateTreeContext.cpp
@@ -5,6 +5,7 @@
 #include "TraitCore/TraitStackBinding.h"
 #include "TraitInterfaces/IBlendStack.h"
 #include "TraitInterfaces/ISmoothBlend.h"
+#include "TraitInterfaces/ITimeline.h"
 
 bool FAnimNextStateTreeTraitContext::PushAnimationGraphOntoBlendStack(TNonNullPtr<UAnimNextAnimationGraph> InAnimationGraph, const FAlphaBlendArgs& InBlendArguments) const
 {		
@@ -35,4 +36,20 @@ bool FAnimNextStateTreeTraitContext::PushAnimationGraphOntoBlendStack(TNonNullPt
 	}
 
 	return false;
-}
\ No newline at end of file
+}
+
+float FAnimNextStateTreeTraitContext::QueryPlaybackRatio(TNonNullPtr<UAnimNextAnimationGraph> InAnimationGraph) const
+{
+	if (Binding && Context)
+	{
+		UE::AnimNext::TTraitBinding<UE::AnimNext::ITimeline> Timeline;
+		if (Binding->GetInterface<UE::AnimNext::ITimeline>(Timeline))
+		{
+			UE::AnimNext::FTimelineProgress Progress = Timeline.GetProgress(*Context);
+
+			return Progress.GetPositionRatio();
+		}
+	}
+
+	return 1.0f;
+}
diff --git a/Engine/Plugins/Experimental/AnimNextStateTree/Source/AnimNextStateTree/Private/AnimNextStateTreeGraphInstanceTask.cpp b/Engine/Plugins/Experimental/AnimNextStateTree/Source/AnimNextStateTree/Private/AnimNextStateTreeGraphInstanceTask.cpp
--- a/Engine/Plugins/Experimental/AnimNextStateTree/Source/AnimNextStateTree/Private/AnimNextStateTreeGraphInstanceTask.cpp
+++ b/Engine/Plugins/Experimental/AnimNextStateTree/Source/AnimNextStateTree/Private/AnimNextStateTreeGraphInstanceTask.cpp
@@ -31,6 +31,10 @@ EStateTreeRunStatus FAnimNextStateTreeGraphInstanceTask::Tick(FStateTreeExecutio
 	FInstanceDataType& InstanceData = Context.GetInstanceData(*this);
 	if (InstanceData.bContinueTicking)
 	{
+		FAnimNextStateTreeTraitContext& ExecContext = Context.GetExternalData(TraitContextHandle);
+
+		InstanceData.PlaybackRatio = ExecContext.QueryPlaybackRatio(InstanceData.AnimationGraph.Get());
+
 		return EStateTreeRunStatus::Running;
 	}
 	
diff --git a/Engine/Plugins/Experimental/AnimNextStateTree/Source/AnimNextStateTree/Public/AnimNextStateTreeContext.h b/Engine/Plugins/Experimental/AnimNextStateTree/Source/AnimNextStateTree/Public/AnimNextStateTreeContext.h
--- a/Engine/Plugins/Experimental/AnimNextStateTree/Source/AnimNextStateTree/Public/AnimNextStateTreeContext.h
+++ b/Engine/Plugins/Experimental/AnimNextStateTree/Source/AnimNextStateTree/Public/AnimNextStateTreeContext.h
@@ -24,6 +24,7 @@ struct FAnimNextStateTreeTraitContext
 	FAnimNextStateTreeTraitContext() {}
 
 	bool PushAnimationGraphOntoBlendStack(TNonNullPtr<UAnimNextAnimationGraph> InAnimationGraph, const FAlphaBlendArgs& InBlendArguments) const;
+	float QueryPlaybackRatio(TNonNullPtr<UAnimNextAnimationGraph> InAnimationGraph) const;
 protected:
 	FAnimNextStateTreeTraitContext(UE::AnimNext::FExecutionContext& InContext, const UE::AnimNext::FTraitStackBinding* InBinding) : Context(&InContext), Binding(InBinding) {}
 protected:
diff --git a/Engine/Plugins/Experimental/AnimNextStateTree/Source/AnimNextStateTree/Public/Tasks/AnimNextStateTreeGraphInstanceTask.h b/Engine/Plugins/Experimental/AnimNextStateTree/Source/AnimNextStateTree/Public/Tasks/AnimNextStateTreeGraphInstanceTask.h
--- a/Engine/Plugins/Experimental/AnimNextStateTree/Source/AnimNextStateTree/Public/Tasks/AnimNextStateTreeGraphInstanceTask.h
+++ b/Engine/Plugins/Experimental/AnimNextStateTree/Source/AnimNextStateTree/Public/Tasks/AnimNextStateTreeGraphInstanceTask.h
@@ -22,6 +22,9 @@ struct ANIMNEXTSTATETREE_API FAnimNextGraphInstanceTaskInstanceData
 
 	UPROPERTY(EditAnywhere, Category = Animation)
 	bool bContinueTicking = true;
+
+	UPROPERTY(VisibleAnywhere, Category = Animation)
+	float PlaybackRatio = 1.0f;
 };
 
 // Basic task pushing AnimationGraph onto blend stack
