Fix beta label of AnimNext plugins
[RN] AnimNext is an experimental feature. In a previous release this was accidentally marked as beta.

diff --git a/Engine/Plugins/Experimental/AnimNext/AnimNext.uplugin b/Engine/Plugins/Experimental/AnimNext/AnimNext.uplugin
--- a/Engine/Plugins/Experimental/AnimNext/AnimNext.uplugin
+++ b/Engine/Plugins/Experimental/AnimNext/AnimNext.uplugin
@@ -12,7 +12,7 @@
 	"SupportURL" : "",
 	"EnabledByDefault" : false,
 	"CanContainContent" : false,
-	"IsBetaVersion": true,
+	"IsBetaVersion": false,
 	"IsExperimentalVersion": true,
 	"Installed" : false,
 	"Plugins":
diff --git a/Engine/Plugins/Experimental/AnimNextAnimGraph/AnimNextAnimGraph.uplugin b/Engine/Plugins/Experimental/AnimNextAnimGraph/AnimNextAnimGraph.uplugin
--- a/Engine/Plugins/Experimental/AnimNextAnimGraph/AnimNextAnimGraph.uplugin
+++ b/Engine/Plugins/Experimental/AnimNextAnimGraph/AnimNextAnimGraph.uplugin
@@ -12,7 +12,7 @@
 	"SupportURL" : "",
 	"EnabledByDefault" : false,
 	"CanContainContent" : false,
-	"IsBetaVersion": true,
+	"IsBetaVersion": false,
 	"IsExperimentalVersion": true,
 	"Installed" : false,
 	"Plugins":
diff --git a/Engine/Plugins/Experimental/AnimNextChooser/AnimNextChooser.uplugin b/Engine/Plugins/Experimental/AnimNextChooser/AnimNextChooser.uplugin
--- a/Engine/Plugins/Experimental/AnimNextChooser/AnimNextChooser.uplugin
+++ b/Engine/Plugins/Experimental/AnimNextChooser/AnimNextChooser.uplugin
@@ -12,7 +12,7 @@
 	"SupportURL" : "",
 	"EnabledByDefault" : false,
 	"CanContainContent" : false,
-	"IsBetaVersion": true,
+	"IsBetaVersion": false,
 	"IsExperimentalVersion": true,
 	"Installed" : false,
 	"Plugins":
diff --git a/Engine/Plugins/Experimental/AnimNextStateTree/AnimNextStateTree.uplugin b/Engine/Plugins/Experimental/AnimNextStateTree/AnimNextStateTree.uplugin
--- a/Engine/Plugins/Experimental/AnimNextStateTree/AnimNextStateTree.uplugin
+++ b/Engine/Plugins/Experimental/AnimNextStateTree/AnimNextStateTree.uplugin
@@ -12,7 +12,7 @@
 	"SupportURL" : "",
 	"EnabledByDefault" : false,
 	"CanContainContent" : false,
-	"IsBetaVersion": true,
+	"IsBetaVersion": false,
 	"IsExperimentalVersion": true,
 	"Installed" : false,
 	"Plugins":
