AnimNextAnimGraph - support for payloads in blendstack

diff --git a/Engine/Plugins/Experimental/AnimNextAnimGraph/Source/AnimNextAnimGraph/Internal/Traits/BlendStackTrait.h b/Engine/Plugins/Experimental/AnimNextAnimGraph/Source/AnimNextAnimGraph/Internal/Traits/BlendStackTrait.h
--- a/Engine/Plugins/Experimental/AnimNextAnimGraph/Source/AnimNextAnimGraph/Internal/Traits/BlendStackTrait.h
+++ b/Engine/Plugins/Experimental/AnimNextAnimGraph/Source/AnimNextAnimGraph/Internal/Traits/BlendStackTrait.h
@@ -112,7 +112,7 @@ namespace UE::AnimNext
 			// Whether or not this graph was previously relevant
 			bool bNewlyCreated = false;
 
-			void Initialize(const IBlendStack::FGraphRequest& GraphRequest);
+			void Initialize(IBlendStack::FGraphRequest&& GraphRequest);
 			void Terminate();
 		};
 
@@ -161,8 +161,8 @@ namespace UE::AnimNext
 		virtual void AddReferencedObjects(const FExecutionContext& Context, const TTraitBinding<IGarbageCollection>& Binding, FReferenceCollector& Collector) const override;
 
 		// IBlendStack impl
-		virtual void PushGraph(FExecutionContext& Context, const TTraitBinding<IBlendStack>& Binding, const IBlendStack::FGraphRequest& GraphRequest, FAnimNextGraphInstancePtr& OutGraphInstance) const;
-		virtual void GetActiveGraphRequest(FExecutionContext& Context, const TTraitBinding<IBlendStack>& Binding, IBlendStack::FGraphRequest& OutRequest) const;
+		virtual void PushGraph(FExecutionContext& Context, const TTraitBinding<IBlendStack>& Binding, IBlendStack::FGraphRequest&& GraphRequest) const override;
+		virtual void GetActiveGraph(FExecutionContext& Context, const TTraitBinding<IBlendStack>& Binding, IBlendStack::FGraphRequestPtr& OutGraphRequest) const override;
 
 #if WITH_EDITOR
 		virtual bool IsHidden() const override { return true; }
diff --git a/Engine/Plugins/Experimental/AnimNextAnimGraph/Source/AnimNextAnimGraph/Private/TraitInterfaces/IBlendStack.cpp b/Engine/Plugins/Experimental/AnimNextAnimGraph/Source/AnimNextAnimGraph/Private/TraitInterfaces/IBlendStack.cpp
--- a/Engine/Plugins/Experimental/AnimNextAnimGraph/Source/AnimNextAnimGraph/Private/TraitInterfaces/IBlendStack.cpp
+++ b/Engine/Plugins/Experimental/AnimNextAnimGraph/Source/AnimNextAnimGraph/Private/TraitInterfaces/IBlendStack.cpp
@@ -21,27 +21,24 @@ namespace UE::AnimNext
 	}
 #endif // WITH_EDITOR
 
-	void IBlendStack::PushGraph(FExecutionContext& Context, const TTraitBinding<IBlendStack>& Binding, const IBlendStack::FGraphRequest& GraphRequest, FAnimNextGraphInstancePtr& OutGraphInstance) const
+	void IBlendStack::PushGraph(FExecutionContext& Context, const TTraitBinding<IBlendStack>& Binding, IBlendStack::FGraphRequest&& GraphRequest) const
 	{
 		TTraitBinding<IBlendStack> SuperBinding;
 		if (Binding.GetStackInterfaceSuper(SuperBinding))
 		{
-			SuperBinding.PushGraph(Context, GraphRequest, OutGraphInstance);
-			return;
+			SuperBinding.PushGraph(Context, MoveTemp(GraphRequest));
 		}
-
-		OutGraphInstance.Release();
 	}
 
-	void IBlendStack::GetActiveGraphRequest(FExecutionContext& Context, const TTraitBinding<IBlendStack>& Binding, IBlendStack::FGraphRequest& OutRequest) const
+	void IBlendStack::GetActiveGraph(FExecutionContext& Context, const TTraitBinding<IBlendStack>& Binding, IBlendStack::FGraphRequestPtr& OutGraphRequest) const
 	{
 		TTraitBinding<IBlendStack> SuperBinding;
 		if (Binding.GetStackInterfaceSuper(SuperBinding))
 		{
-			SuperBinding.GetActiveGraphRequest(Context, OutRequest);
+			SuperBinding.GetActiveGraph(Context, OutGraphRequest);
 			return;
 		}
 
-		OutRequest.AnimationGraph = nullptr;
+		OutGraphRequest = nullptr;
 	}
 }
diff --git a/Engine/Plugins/Experimental/AnimNextAnimGraph/Source/AnimNextAnimGraph/Private/Traits/BlendStackTrait.cpp b/Engine/Plugins/Experimental/AnimNextAnimGraph/Source/AnimNextAnimGraph/Private/Traits/BlendStackTrait.cpp
--- a/Engine/Plugins/Experimental/AnimNextAnimGraph/Source/AnimNextAnimGraph/Private/Traits/BlendStackTrait.cpp
+++ b/Engine/Plugins/Experimental/AnimNextAnimGraph/Source/AnimNextAnimGraph/Private/Traits/BlendStackTrait.cpp
@@ -1,6 +1,7 @@
 // Copyright Epic Games, Inc. All Rights Reserved.
 
 #include "Traits/BlendStackTrait.h"
+#include "DataInterface/DataInterfaceStructAdapter.h"
 #include "TraitCore/NodeInstance.h"
 #include "TraitCore/ExecutionContext.h"
 
@@ -51,9 +52,9 @@ namespace UE::AnimNext
 		IGarbageCollection::UnregisterWithGC(Context, Binding);
 	}
 
-	void FBlendStackCoreTrait::FGraphState::Initialize(const IBlendStack::FGraphRequest& GraphRequest)
+	void FBlendStackCoreTrait::FGraphState::Initialize(IBlendStack::FGraphRequest&& GraphRequest)
 	{
-		Request = GraphRequest;
+		Request = MoveTemp(GraphRequest);
 		State = FBlendStackCoreTrait::EGraphState::Active;
 		bNewlyCreated = true;
 	}
@@ -266,11 +267,10 @@ namespace UE::AnimNext
 		}
 	}
 
-	void FBlendStackCoreTrait::PushGraph(FExecutionContext& Context, const TTraitBinding<IBlendStack>& Binding, const IBlendStack::FGraphRequest& GraphRequest, FAnimNextGraphInstancePtr& OutGraphInstance) const
+	void FBlendStackCoreTrait::PushGraph(FExecutionContext& Context, const TTraitBinding<IBlendStack>& Binding, IBlendStack::FGraphRequest&& GraphRequest) const
 	{
 		if (GraphRequest.AnimationGraph == nullptr)
 		{
-			OutGraphInstance.Release();
 			return;
 		}
 
@@ -281,27 +281,53 @@ namespace UE::AnimNext
 		const int32 OldChildIndex = InstanceData->CurrentlyActiveGraphIndex;
 		const int32 NewChildIndex = FBlendStackCoreTrait::FindFreeGraphIndexOrAdd(*InstanceData);
 		FBlendStackCoreTrait::FGraphState& Graph = InstanceData->ChildGraphs[NewChildIndex];
-		Graph.Initialize(GraphRequest);
+		Graph.Initialize(MoveTemp(GraphRequest));
 
 		TTraitBinding<IDiscreteBlend> DiscreteBlendTrait;
 		Binding.GetStackInterface(DiscreteBlendTrait);
 		DiscreteBlendTrait.OnBlendTransition(Context, OldChildIndex, NewChildIndex);
 		
 		InstanceData->CurrentlyActiveGraphIndex = NewChildIndex;
-		OutGraphInstance = Graph.Instance;
+
+		const int32 NumPayloads = Graph.Request.Payloads.Num();
+		if (NumPayloads > 0)
+		{
+			// stack allocating uninitialized arrays to contain FDataInterfaceStructAdapter(s)
+			TArrayView<FDataInterfaceStructAdapter> Adapters((FDataInterfaceStructAdapter*)FMemory_Alloca(NumPayloads * sizeof(FDataInterfaceStructAdapter)), NumPayloads);
+			TArrayView<IDataInterfaceHost*> PayloadVariableBindings((IDataInterfaceHost**)FMemory_Alloca(NumPayloads * sizeof(IDataInterfaceHost*)), NumPayloads);
+
+			// initializing the FDataInterfaceStructAdapter(s) with their respective Graph.Request.Payloads
+			for (int32 PayloadIndex = 0; PayloadIndex < NumPayloads; ++PayloadIndex)
+			{
+				FInstancedStruct& Payload = Graph.Request.Payloads[PayloadIndex];
+				PayloadVariableBindings[PayloadIndex] = new (&Adapters[PayloadIndex]) FDataInterfaceStructAdapter(Graph.Request.AnimationGraph, Payload);
+			}
+			
+			// binding 
+			Graph.Instance.BindPublicVariables(PayloadVariableBindings);
+			
+			// uninitializing the FDataInterfaceStructAdapter(s)
+			for (int32 PayloadIndex = 0; PayloadIndex < NumPayloads; ++PayloadIndex)
+			{
+				Adapters[PayloadIndex].~FDataInterfaceStructAdapter();
+			}
+
+			// abandoning no longer initialized stack allocated arrays to their demise
+		}
 	}
 
-	void FBlendStackCoreTrait::GetActiveGraphRequest(FExecutionContext& Context, const TTraitBinding<IBlendStack>& Binding, IBlendStack::FGraphRequest& OutRequest) const
+	void FBlendStackCoreTrait::GetActiveGraph(FExecutionContext& Context, const TTraitBinding<IBlendStack>& Binding, IBlendStack::FGraphRequestPtr& OutGraphRequest) const
 	{
 		FInstanceData* InstanceData = Binding.GetInstanceData<FInstanceData>();
 		const int32 CurrentlyActiveGraphIndex = InstanceData->CurrentlyActiveGraphIndex;
 		if (CurrentlyActiveGraphIndex != INDEX_NONE)
 		{
-			OutRequest = InstanceData->ChildGraphs[CurrentlyActiveGraphIndex].Request;
+			FGraphState& GraphState = InstanceData->ChildGraphs[CurrentlyActiveGraphIndex];
+			OutGraphRequest = &GraphState.Request;
 			return;
 		}
 
-		return IBlendStack::GetActiveGraphRequest(Context, Binding, OutRequest);
+		return IBlendStack::GetActiveGraph(Context, Binding, OutGraphRequest);
 	}
 
 	void FBlendStackTrait::PreUpdate(FUpdateTraversalContext& Context, const TTraitBinding<IUpdate>& Binding, const FTraitUpdateState& TraitState) const
@@ -319,12 +345,11 @@ namespace UE::AnimNext
 			TTraitBinding<IBlendStack> BlendStackTrait;
 			Binding.GetStackInterface(BlendStackTrait);
 
-			FAnimNextGraphInstancePtr NewGraphInstance;
 			IBlendStack::FGraphRequest GraphRequest;
 			GraphRequest.AnimationGraph = DesiredAnimationGraph;
 			GraphRequest.BlendArgs.BlendTime = SharedData->GetBlendTime(Binding);
 
-			BlendStackTrait.PushGraph(Context, GraphRequest, NewGraphInstance);
+			BlendStackTrait.PushGraph(Context, MoveTemp(GraphRequest));
 		}
 	}
 
@@ -338,17 +363,16 @@ namespace UE::AnimNext
 			TTraitBinding<IBlendStack> BlendStackTrait;
 			Binding.GetStackInterface(BlendStackTrait);
 
-			IBlendStack::FGraphRequest ActiveGraphRequest;
-			BlendStackTrait.GetActiveGraphRequest(Context, ActiveGraphRequest);
+			IBlendStack::FGraphRequestPtr ActiveGraphRequest;
+			BlendStackTrait.GetActiveGraph(Context, ActiveGraphRequest);
 			const bool bForceBlend = SharedData->GetbForceBlend(Binding);
-			if (bForceBlend || (DesiredAnimationGraph != ActiveGraphRequest.AnimationGraph))
+			if (bForceBlend || !ActiveGraphRequest || (DesiredAnimationGraph != ActiveGraphRequest->AnimationGraph))
 			{
-				FAnimNextGraphInstancePtr NewGraphInstance;
 				IBlendStack::FGraphRequest GraphRequest;
 				GraphRequest.AnimationGraph = DesiredAnimationGraph;
 				GraphRequest.BlendArgs.BlendTime = SharedData->GetBlendTime(Binding);
 
-				BlendStackTrait.PushGraph(Context, GraphRequest, NewGraphInstance);
+				BlendStackTrait.PushGraph(Context, MoveTemp(GraphRequest));
 			}
 		}
 	}
@@ -363,12 +387,11 @@ namespace UE::AnimNext
 		const UAnimNextAnimationGraph* DesiredAnimationGraph = SharedData->GetAnimationGraph(Binding);
 		if (DesiredAnimationGraph != nullptr)
 		{
-			FAnimNextGraphInstancePtr NewGraphInstance;
 			IBlendStack::FGraphRequest GraphRequest;
 			GraphRequest.AnimationGraph = DesiredAnimationGraph;
 			GraphRequest.BlendArgs.BlendTime = SharedData->GetBlendTime(Binding);
 
-			BlendStackTrait.PushGraph(Context, GraphRequest, NewGraphInstance);
+			BlendStackTrait.PushGraph(Context, MoveTemp(GraphRequest));
 		}
 	}
 }
diff --git a/Engine/Plugins/Experimental/AnimNextAnimGraph/Source/AnimNextAnimGraph/Public/TraitInterfaces/IBlendStack.h b/Engine/Plugins/Experimental/AnimNextAnimGraph/Source/AnimNextAnimGraph/Public/TraitInterfaces/IBlendStack.h
--- a/Engine/Plugins/Experimental/AnimNextAnimGraph/Source/AnimNextAnimGraph/Public/TraitInterfaces/IBlendStack.h
+++ b/Engine/Plugins/Experimental/AnimNextAnimGraph/Source/AnimNextAnimGraph/Public/TraitInterfaces/IBlendStack.h
@@ -27,6 +27,9 @@ namespace UE::AnimNext
 			 * Consider making dynamic payload if we want to implement a special blend framework.
 			 */
 
+			// Payloads applied to the AnimationGraph's variables via its data interfaces.
+			TArray<FInstancedStruct> Payloads;
+
 			// Blend-in duration for this graph request 
 			FAlphaBlendArgs BlendArgs;
 
@@ -34,11 +37,13 @@ namespace UE::AnimNext
 			TObjectPtr<const UAnimNextAnimationGraph> AnimationGraph;
 		};
 
+		typedef FGraphRequest* FGraphRequestPtr;
+
 		// Pushes a new subgraph along with blend settings defined in GraphRequest. Outputs the in-place created subgraph OutGraphInstance
-		virtual void PushGraph(FExecutionContext& Context, const TTraitBinding<IBlendStack>& Binding, const IBlendStack::FGraphRequest& GraphRequest, FAnimNextGraphInstancePtr& OutGraphInstance) const;
+		virtual void PushGraph(FExecutionContext& Context, const TTraitBinding<IBlendStack>& Binding, IBlendStack::FGraphRequest&& GraphRequest) const;
 		
 		// Gets the graph request info from the most recent PushGraph
-		virtual void GetActiveGraphRequest(FExecutionContext& Context, const TTraitBinding<IBlendStack>& Binding, IBlendStack::FGraphRequest& OutRequest) const;
+		virtual void GetActiveGraph(FExecutionContext& Context, const TTraitBinding<IBlendStack>& Binding, IBlendStack::FGraphRequestPtr& OutGraphRequest) const;
 
 #if WITH_EDITOR
 		virtual const FText& GetDisplayName() const override;
@@ -53,17 +58,17 @@ namespace UE::AnimNext
 	struct TTraitBinding<IBlendStack> : FTraitBinding
 	{
 		// @see IBlendStack::PushGraph
-		void PushGraph(FExecutionContext& Context, const IBlendStack::FGraphRequest& GraphRequest, FAnimNextGraphInstancePtr& OutGraphInstance) const
+		void PushGraph(FExecutionContext& Context, IBlendStack::FGraphRequest&& GraphRequest) const
 		{
-			GetInterface()->PushGraph(Context, *this, GraphRequest, OutGraphInstance);
+			GetInterface()->PushGraph(Context, *this, MoveTemp(GraphRequest));
 		}
 
-		// @see IBlendStack::GetActiveGraphRequest
-		void GetActiveGraphRequest(FExecutionContext& Context, IBlendStack::FGraphRequest& OutRequest) const
+		// @see IBlendStack::GetActiveGraph
+		void GetActiveGraph(FExecutionContext& Context, IBlendStack::FGraphRequestPtr& OutGraphRequest) const
 		{
-			GetInterface()->GetActiveGraphRequest(Context, *this, OutRequest);
+			GetInterface()->GetActiveGraph(Context, *this, OutGraphRequest);
 		}
-
+		
 	protected:
 		const IBlendStack* GetInterface() const { return GetInterfaceTyped<IBlendStack>(); }
 	};
diff --git a/Engine/Plugins/Experimental/AnimNextStateTree/Source/AnimNextStateTree/Private/AnimNextStateTreeContext.cpp b/Engine/Plugins/Experimental/AnimNextStateTree/Source/AnimNextStateTree/Private/AnimNextStateTreeContext.cpp
--- a/Engine/Plugins/Experimental/AnimNextStateTree/Source/AnimNextStateTree/Private/AnimNextStateTreeContext.cpp
+++ b/Engine/Plugins/Experimental/AnimNextStateTree/Source/AnimNextStateTree/Private/AnimNextStateTreeContext.cpp
@@ -14,22 +14,18 @@ bool FAnimNextStateTreeTraitContext::PushAnimationGraphOntoBlendStack(TNonNullPt
 		UE::AnimNext::TTraitBinding<UE::AnimNext::IBlendStack> BlendStackBinding;
 		if(Binding->GetInterface<UE::AnimNext::IBlendStack>(BlendStackBinding))
 		{
-			UE::AnimNext::IBlendStack::FGraphRequest CurrentRequest;			
-			BlendStackBinding.GetActiveGraphRequest(*Context, CurrentRequest);
+			UE::AnimNext::IBlendStack::FGraphRequestPtr ActiveGraphRequest;
+			BlendStackBinding.GetActiveGraph(*Context, ActiveGraphRequest);
 
-			if (CurrentRequest.AnimationGraph != nullptr)
+			if (ActiveGraphRequest && ActiveGraphRequest->AnimationGraph == InAnimationGraph)
 			{
-				if (CurrentRequest.AnimationGraph == InAnimationGraph)
-				{
-					return true;
-				}
+				return true;
 			}
 			
-			FAnimNextGraphInstancePtr GraphPtr;
-			UE::AnimNext::IBlendStack::FGraphRequest GraphRequest;
-			GraphRequest.BlendArgs = InBlendArguments;
-			GraphRequest.AnimationGraph = InAnimationGraph.Get();
-			BlendStackBinding.PushGraph(*Context, GraphRequest, GraphPtr);
+			UE::AnimNext::IBlendStack::FGraphRequest NewGraphRequest;
+			NewGraphRequest.BlendArgs = InBlendArguments;
+			NewGraphRequest.AnimationGraph = InAnimationGraph.Get();
+			BlendStackBinding.PushGraph(*Context, MoveTemp(NewGraphRequest));
 
 			return true;
 		}
