AnimNext: Move anim next graph instances to use weak semantics

This is intended to more tightly control the lifetime and ownership of animation graph instances.
Introduced module components.
Module component for animation graphs holds shared ptrs to all hosted (root) graphs.
Traits that own animation graph instances (sub graphs) hold shared ptrs to animation graphs.
All other references (including those inside RigVM work memory) are expected to be 'weak' and use FAnimNextWeakGraphInstancePtr.

#rnx

diff --git a/Engine/Plugins/Experimental/AnimNext/Source/AnimNext/Internal/AnimNextModuleImpl.h b/Engine/Plugins/Experimental/AnimNext/Source/AnimNext/Internal/AnimNextModuleImpl.h
--- a/Engine/Plugins/Experimental/AnimNext/Source/AnimNext/Internal/AnimNextModuleImpl.h
+++ b/Engine/Plugins/Experimental/AnimNext/Source/AnimNext/Internal/AnimNextModuleImpl.h
@@ -5,7 +5,7 @@
 #include "CoreMinimal.h"
 #include "IAnimNextModuleInterface.h"
 
-struct FAnimNextGraphInstancePtr;
+struct FAnimNextGraphInstance;
 struct FAnimNextVariableBindingData;
 
 namespace UE::AnimNext
@@ -29,8 +29,8 @@ public:
 	// IAnimNextModuleInterface interface
 	virtual void RegisterAnimNextAnimGraph(const IAnimNextAnimGraph& InAnimGraphImpl) override;
 	virtual void UnregisterAnimNextAnimGraph() override;
-	virtual void UpdateGraph(const FAnimNextGraphInstancePtr& GraphInstance, float DeltaTime, FTraitEventList& InputEventList, FTraitEventList& OutputEventList) override;
-	virtual void EvaluateGraph(const FAnimNextGraphInstancePtr& GraphInstance, const FReferencePose& RefPose, int32 GraphLODLevel, FAnimNextGraphLODPose& OutputPose) const override;
+	virtual void UpdateGraph(FAnimNextGraphInstance& GraphInstance, float DeltaTime, FTraitEventList& InputEventList, FTraitEventList& OutputEventList) override;
+	virtual void EvaluateGraph(FAnimNextGraphInstance& GraphInstance, const FReferencePose& RefPose, int32 GraphLODLevel, FAnimNextGraphLODPose& OutputPose) const override;
 
 #if WITH_ANIMNEXT_CONSOLE_COMMANDS
 	void ListNodeTemplates(const TArray<FString>& Args);
diff --git a/Engine/Plugins/Experimental/AnimNext/Source/AnimNext/Internal/Graph/AnimNextAnimationGraph.h b/Engine/Plugins/Experimental/AnimNext/Source/AnimNext/Internal/Graph/AnimNextAnimationGraph.h
--- a/Engine/Plugins/Experimental/AnimNext/Source/AnimNext/Internal/Graph/AnimNextAnimationGraph.h
+++ b/Engine/Plugins/Experimental/AnimNext/Source/AnimNext/Internal/Graph/AnimNextAnimationGraph.h
@@ -20,7 +20,6 @@ class UAnimNextModule;
 class UAnimGraphNode_AnimNextGraph;
 struct FAnimNode_AnimNextGraph;
 struct FRigUnit_AnimNextGraphEvaluator;
-struct FAnimNextGraphInstancePtr;
 struct FAnimNextGraphInstance;
 struct FAnimNextScheduleGraphTask;
 struct FAnimNextEditorParam;
@@ -35,6 +34,7 @@ namespace UE::AnimNext
 	struct FTestUtils;
 	struct FParametersProxy;
 	struct FPlayAnimSlotTrait;
+	struct FAnimationGraphComponent;
 }
 
 namespace UE::AnimNext::UncookedOnly
@@ -66,26 +66,17 @@ public:
 	// UObject interface
 	virtual void Serialize(FArchive& Ar) override;
 
-	// Allocates an instance of the graph
-	// @param	OutInstance			The instance to allocate data for
-	// @param	InModuleInstance	The module instance to use
-	// @param	InEntryPoint		The entry point to use. If this is NAME_None then the default entry point for this graph is used
-	void AllocateInstance(FAnimNextGraphInstancePtr& OutInstance, FAnimNextModuleInstance* InModuleInstance = nullptr, FName InEntryPoint = NAME_None) const;
-
-	// Allocates an instance of the graph with the specified parent graph instance
+	// Allocates an instance of the graph with an optional module and/or parent graph instance
+	// @param	InModuleInstance			The module instance to use
 	// @param	InParentGraphInstance		The parent graph instance to use
-	// @param	OutInstance					The instance to allocate data for
 	// @param	InEntryPoint				The entry point to use. If this is NAME_None then the default entry point for this graph is used
-	void AllocateInstance(FAnimNextGraphInstance& InParentGraphInstance, FAnimNextGraphInstancePtr& OutInstance, FName InEntryPoint = NAME_None) const;
-
-protected:
+	// @return the allocated instance. This can be invalid if the entry point was not found or the graph's root trait is invalid.
+	TSharedPtr<FAnimNextGraphInstance> AllocateInstance(FAnimNextModuleInstance* InModuleInstance = nullptr, FAnimNextGraphInstance* InParentGraphInstance = nullptr, FName InEntryPoint = NAME_None) const;
 
+private:
 	// Loads the graph data from the provided archive buffer and returns true on success, false otherwise
 	bool LoadFromArchiveBuffer(const TArray<uint8>& SharedDataArchiveBuffer);
 
-	// Allocates an instance of the graph with an optional parent graph instance
-	void AllocateInstanceImpl(FAnimNextModuleInstance* InModuleInstance, FAnimNextGraphInstance* InParentGraphInstance, FAnimNextGraphInstancePtr& OutInstance, FName InEntryPoint) const;
-
 #if WITH_EDITORONLY_DATA
 	// During graph compilation, if we have existing graph instances, we freeze them by releasing their memory before thawing them
 	// Freezing is a partial release of resources that retains the necessary information to re-create things safely
@@ -101,7 +92,6 @@ protected:
 	friend struct UE::AnimNext::UncookedOnly::FUtils;
 	friend class UE::AnimNext::Editor::FModuleEditor;
 	friend struct UE::AnimNext::FTestUtils;
-	friend FAnimNextGraphInstancePtr;
 	friend FAnimNextGraphInstance;
 	friend class UAnimGraphNode_AnimNextGraph;
 	friend UE::AnimNext::FExecutionContext;
@@ -110,7 +100,9 @@ protected:
 	friend class UE::AnimNext::Editor::FVariableCustomization;
 	friend struct UE::AnimNext::FParametersProxy;
 	friend struct UE::AnimNext::FPlayAnimSlotTrait;
+	friend struct UE::AnimNext::FAnimationGraphComponent;
 
+protected:
 #if WITH_EDITORONLY_DATA
 	mutable FCriticalSection GraphInstancesLock;
 
diff --git a/Engine/Plugins/Experimental/AnimNext/Source/AnimNext/Internal/Graph/AnimationGraphComponent.h b/Engine/Plugins/Experimental/AnimNext/Source/AnimNext/Internal/Graph/AnimationGraphComponent.h
new file mode 100644
--- /dev/null
+++ b/Engine/Plugins/Experimental/AnimNext/Source/AnimNext/Internal/Graph/AnimationGraphComponent.h
@@ -0,0 +1,27 @@
+﻿// Copyright Epic Games, Inc. All Rights Reserved.
+
+#pragma once
+
+#include "Graph/AnimNextGraphInstancePtr.h"
+#include "Module/ModuleInstanceComponent.h"
+
+namespace UE::AnimNext
+{
+
+// Module component that owns/allocates/release all animation graph instances on a module
+struct FAnimationGraphComponent : public FModuleInstanceComponent
+{
+	ANIMNEXT_DECLARE_MODULE_INSTANCE_COMPONENT(FAnimationGraphComponent);
+
+	explicit FAnimationGraphComponent(FAnimNextModuleInstance& InOwnerInstance);
+
+	FAnimNextGraphInstancePtr AllocateInstance(UAnimNextAnimationGraph* InAnimationGraph, FAnimNextGraphInstance* InParentInstance = nullptr, FName InEntryPoint = NAME_None);
+
+	void ReleaseInstance(FAnimNextGraphInstancePtr InInstance);
+
+private:
+	// All the owned graph instances for the module
+	TArray<TSharedPtr<FAnimNextGraphInstance>> GraphInstances;
+};
+
+}
diff --git a/Engine/Plugins/Experimental/AnimNext/Source/AnimNext/Internal/Module/AnimNextModuleInstance.h b/Engine/Plugins/Experimental/AnimNext/Source/AnimNext/Internal/Module/AnimNextModuleInstance.h
--- a/Engine/Plugins/Experimental/AnimNext/Source/AnimNext/Internal/Module/AnimNextModuleInstance.h
+++ b/Engine/Plugins/Experimental/AnimNext/Source/AnimNext/Internal/Module/AnimNextModuleInstance.h
@@ -11,6 +11,7 @@
 #include "TraitCore/TraitEvent.h"
 #include "TraitCore/TraitEventList.h"
 #include "Module/ModuleHandle.h"
+#include "Module/ModuleInstanceComponent.h"
 #include "AnimNextModuleInstance.generated.h"
 
 struct FAnimNextModuleInstance;
@@ -49,6 +50,11 @@ namespace UE::AnimNext::Private
 	extern void CacheAllModuleEvents();
 }
 
+namespace UE::AnimNext
+{
+	using ModuleInstanceComponentMapType = TMap<FName, TSharedPtr<FModuleInstanceComponent>>;
+}
+
 // Root memory owner of a parameterized schedule 
 USTRUCT()
 struct FAnimNextModuleInstance : public FAnimNextDataInterfaceInstance
@@ -91,6 +97,24 @@ struct FAnimNextModuleInstance : public FAnimNextDataInterfaceInstance
 
 	// Get the module that this instance represents
 	ANIMNEXT_API const UAnimNextModule* GetModule() const;
+	
+	// Returns a typed module instance component, creating it lazily the first time it is queried
+	template<class ComponentType>
+	ComponentType& GetComponent();
+
+	// Returns a typed module instance component pointer if found or nullptr otherwise
+	template<class ComponentType>
+	ComponentType* TryGetComponent();
+
+	// Returns a typed module instance component pointer if found or nullptr otherwise
+	template<class ComponentType>
+	const ComponentType* TryGetComponent() const;
+
+	// Returns a pointer to the specified component, or nullptr if not found
+	UE::AnimNext::FModuleInstanceComponent* TryGetComponent(int32 ComponentNameHash, FName ComponentName) const;
+
+	// Adds the specified component and returns a reference to it
+	UE::AnimNext::FModuleInstanceComponent& AddComponent(int32 ComponentNameHash, FName ComponentName, TSharedPtr<UE::AnimNext::FModuleInstanceComponent>&& Component);
 
 #if UE_ENABLE_DEBUG_DRAWING
 	// Get the debug draw interface
@@ -132,6 +156,9 @@ struct FAnimNextModuleInstance : public FAnimNextDataInterfaceInstance
 	TUniquePtr<UE::AnimNext::Debug::FDebugDraw> DebugDraw;
 #endif
 
+	// Components of this module
+	UE::AnimNext::ModuleInstanceComponentMapType Components; 
+
 	enum class ERunState : uint8
 	{
 		None,
@@ -173,3 +200,35 @@ struct TStructOpsTypeTraits<FAnimNextModuleInstance> : public TStructOpsTypeTrai
 		WithCopy = false
 	};
 };
+
+template<class ComponentType>
+ComponentType& FAnimNextModuleInstance::GetComponent()
+{
+	const FName ComponentName = ComponentType::StaticComponentName();
+	const int32 ComponentNameHash = GetTypeHash(ComponentName);
+
+	if (UE::AnimNext::FModuleInstanceComponent* Component = TryGetComponent(ComponentNameHash, ComponentName))
+	{
+		return *static_cast<ComponentType*>(Component);
+	}
+
+	return static_cast<ComponentType&>(AddComponent(ComponentNameHash, ComponentName, MakeShared<ComponentType>(*this)));
+}
+
+template<class ComponentType>
+ComponentType* FAnimNextModuleInstance::TryGetComponent()
+{
+	const FName ComponentName = ComponentType::StaticComponentName();
+	const int32 ComponentNameHash = GetTypeHash(ComponentName);
+
+	return static_cast<ComponentType*>(TryGetComponent(ComponentNameHash, ComponentName));
+}
+
+template<class ComponentType>
+const ComponentType* FAnimNextModuleInstance::TryGetComponent() const
+{
+	const FName ComponentName = ComponentType::StaticComponentName();
+	const int32 ComponentNameHash = GetTypeHash(ComponentName);
+
+	return static_cast<ComponentType*>(TryGetComponent(ComponentNameHash, ComponentName));
+}
diff --git a/Engine/Plugins/Experimental/AnimNext/Source/AnimNext/Internal/Module/ModuleInstanceComponent.h b/Engine/Plugins/Experimental/AnimNext/Source/AnimNext/Internal/Module/ModuleInstanceComponent.h
new file mode 100644
--- /dev/null
+++ b/Engine/Plugins/Experimental/AnimNext/Source/AnimNext/Internal/Module/ModuleInstanceComponent.h
@@ -0,0 +1,42 @@
+﻿// Copyright Epic Games, Inc. All Rights Reserved.
+
+#pragma once
+
+#include "CoreMinimal.h"
+
+struct FAnimNextModuleInstance;
+
+// Declares a module instance component and implements the necessary boilerplate
+#define ANIMNEXT_DECLARE_MODULE_INSTANCE_COMPONENT(ComponentType) \
+	static FName StaticComponentName() { return FName(TEXT(#ComponentType)); } \
+	virtual FName GetComponentName() const override { return StaticComponentName(); }
+
+namespace UE::AnimNext
+{
+
+/**
+ * FModuleInstanceComponent
+ *
+ * A module instance component is attached and owned by a module instance.
+ * It persists as long as it is needed.
+ */
+struct FModuleInstanceComponent
+{
+	explicit FModuleInstanceComponent(FAnimNextModuleInstance& InOwnerInstance) : OwnerInstance(InOwnerInstance) {}
+	virtual ~FModuleInstanceComponent() {}
+
+	static FName StaticComponentName() { return FName(TEXT("FModuleInstanceComponent")); }
+	virtual FName GetComponentName() const { return StaticComponentName(); }
+
+	// Returns the owning module instance this component lives on
+	FAnimNextModuleInstance& GetModuleInstance() { return OwnerInstance; }
+
+	// Returns the owning module instance this component lives on
+	const FAnimNextModuleInstance& GetModuleInstance() const { return OwnerInstance; }
+
+private:
+	// The owning module instance this component lives on
+	FAnimNextModuleInstance& OwnerInstance;
+};
+
+}
diff --git a/Engine/Plugins/Experimental/AnimNext/Source/AnimNext/Private/AnimNextModuleImpl.cpp b/Engine/Plugins/Experimental/AnimNext/Source/AnimNext/Private/AnimNextModuleImpl.cpp
--- a/Engine/Plugins/Experimental/AnimNext/Source/AnimNext/Private/AnimNextModuleImpl.cpp
+++ b/Engine/Plugins/Experimental/AnimNext/Source/AnimNext/Private/AnimNextModuleImpl.cpp
@@ -156,7 +156,7 @@ namespace UE::AnimNext
 		AnimGraphImpl = nullptr;
 	}
 
-	void FAnimNextModuleImpl::UpdateGraph(const FAnimNextGraphInstancePtr& GraphInstance, float DeltaTime, FTraitEventList& InputEventList, FTraitEventList& OutputEventList)
+	void FAnimNextModuleImpl::UpdateGraph(FAnimNextGraphInstance& GraphInstance, float DeltaTime, FTraitEventList& InputEventList, FTraitEventList& OutputEventList)
 	{
 		if (AnimGraphImpl != nullptr)
 		{
@@ -164,7 +164,7 @@ namespace UE::AnimNext
 		}
 	}
 
-	void FAnimNextModuleImpl::EvaluateGraph(const FAnimNextGraphInstancePtr& GraphInstance, const UE::AnimNext::FReferencePose& RefPose, int32 GraphLODLevel, FAnimNextGraphLODPose& OutputPose) const
+	void FAnimNextModuleImpl::EvaluateGraph(FAnimNextGraphInstance& GraphInstance, const UE::AnimNext::FReferencePose& RefPose, int32 GraphLODLevel, FAnimNextGraphLODPose& OutputPose) const
 	{
 		if (AnimGraphImpl != nullptr)
 		{
diff --git a/Engine/Plugins/Experimental/AnimNext/Source/AnimNext/Private/Graph/AnimNextAnimationGraph.cpp b/Engine/Plugins/Experimental/AnimNext/Source/AnimNext/Private/Graph/AnimNextAnimationGraph.cpp
--- a/Engine/Plugins/Experimental/AnimNext/Source/AnimNext/Private/Graph/AnimNextAnimationGraph.cpp
+++ b/Engine/Plugins/Experimental/AnimNext/Source/AnimNext/Private/Graph/AnimNextAnimationGraph.cpp
@@ -31,90 +31,74 @@ UAnimNextAnimationGraph::UAnimNextAnimationGraph(const FObjectInitializer& Objec
 	ExtendedExecuteContext.SetContextPublicDataStruct(FAnimNextExecuteContext::StaticStruct());
 }
 
-void UAnimNextAnimationGraph::AllocateInstance(FAnimNextGraphInstancePtr& OutInstance, FAnimNextModuleInstance* InModuleInstance, FName InEntryPoint) const
-{
-	AllocateInstanceImpl(InModuleInstance, nullptr, OutInstance, InEntryPoint);
-}
-
-void UAnimNextAnimationGraph::AllocateInstance(FAnimNextGraphInstance& InParentGraphInstance, FAnimNextGraphInstancePtr& OutInstance, FName InEntryPoint) const
-{
-	AllocateInstanceImpl(InParentGraphInstance.GetModuleInstance(), &InParentGraphInstance, OutInstance, InEntryPoint);
-}
-
-void UAnimNextAnimationGraph::AllocateInstanceImpl(FAnimNextModuleInstance* InModuleInstance, FAnimNextGraphInstance* ParentGraphInstance, FAnimNextGraphInstancePtr& Instance, FName InEntryPoint) const
+TSharedPtr<FAnimNextGraphInstance> UAnimNextAnimationGraph::AllocateInstance(FAnimNextModuleInstance* InModuleInstance, FAnimNextGraphInstance* InParentGraphInstance, FName InEntryPoint) const
 {
 	SCOPE_CYCLE_COUNTER(STAT_AnimNext_Graph_AllocateInstance);
 
-	Instance.Release();
-
 	const FName EntryPoint = (InEntryPoint == NAME_None) ? DefaultEntryPoint : InEntryPoint;
 	const FAnimNextTraitHandle ResolvedRootTraitHandle = ResolvedRootTraitHandles.FindRef(EntryPoint);
 	if (!ResolvedRootTraitHandle.IsValid())
 	{
-		return;
+		return nullptr;
 	}
 
-	{
-		TSharedPtr<FAnimNextGraphInstance> InstanceImpl = MakeShared<FAnimNextGraphInstance>();
-
-		InstanceImpl->DataInterface = this;
-		InstanceImpl->ModuleInstance = InModuleInstance;
-		InstanceImpl->ParentGraphInstance = ParentGraphInstance;
-		InstanceImpl->EntryPoint = EntryPoint;
+	TSharedPtr<FAnimNextGraphInstance> InstanceImpl = MakeShared<FAnimNextGraphInstance>();
+	InstanceImpl->DataInterface = this;
+	InstanceImpl->ModuleInstance = InModuleInstance;
+	InstanceImpl->ParentGraphInstance = InParentGraphInstance;
+	InstanceImpl->EntryPoint = EntryPoint;
 
-		// If we have a parent graph, use its root since we share the same root, otherwise if we have no parent, we are the root
-		InstanceImpl->RootGraphInstance = ParentGraphInstance != nullptr ? ParentGraphInstance->GetRootGraphInstance() : InstanceImpl.Get();
+	// If we have a parent graph, use its root since we share the same root, otherwise if we have no parent, we are the root
+	InstanceImpl->RootGraphInstance = InParentGraphInstance != nullptr ? InParentGraphInstance->GetRootGraphInstance() : InstanceImpl.Get();
 
-		InstanceImpl->Variables = VariableDefaults;
-		InstanceImpl->ExtendedExecuteContext = ExtendedExecuteContext;
-		if(VariableDefaults.GetPropertyBagStruct() != nullptr)
-		{
-			InstanceImpl->PublicVariablesState = FAnimNextGraphInstance::EPublicVariablesState::Unbound;
-
-			// Setup external variables memory ptrs manually as we dont follow the pattern of owning multiple URigVMHosts like control rig.
-			// InitializeVM() is called, but only sets up handles for the defaults in the module, not for an instance
-			const int32 NumVariables = InstanceImpl->Variables.GetNumPropertiesInBag();
-			TArray<FRigVMExternalVariableRuntimeData> ExternalVariableRuntimeData;
-			ExternalVariableRuntimeData.Reserve(NumVariables);
-			TConstArrayView<FPropertyBagPropertyDesc> Descs = InstanceImpl->Variables.GetPropertyBagStruct()->GetPropertyDescs();
-			uint8* BasePtr = InstanceImpl->Variables.GetMutableValue().GetMemory();
-			for(int32 VariableIndex = 0; VariableIndex < NumVariables; ++VariableIndex)
-			{
-				ExternalVariableRuntimeData.Emplace(Descs[VariableIndex].CachedProperty->ContainerPtrToValuePtr<uint8>(BasePtr));
-			}
-			InstanceImpl->ExtendedExecuteContext.ExternalVariableRuntimeData = MoveTemp(ExternalVariableRuntimeData);
-		}
-		else
+	InstanceImpl->Variables = VariableDefaults;
+	InstanceImpl->ExtendedExecuteContext = ExtendedExecuteContext;
+	if(VariableDefaults.GetPropertyBagStruct() != nullptr)
+	{
+		InstanceImpl->PublicVariablesState = FAnimNextGraphInstance::EPublicVariablesState::Unbound;
+
+		// Setup external variables memory ptrs manually as we dont follow the pattern of owning multiple URigVMHosts like control rig.
+		// InitializeVM() is called, but only sets up handles for the defaults in the module, not for an instance
+		const int32 NumVariables = InstanceImpl->Variables.GetNumPropertiesInBag();
+		TArray<FRigVMExternalVariableRuntimeData> ExternalVariableRuntimeData;
+		ExternalVariableRuntimeData.Reserve(NumVariables);
+		TConstArrayView<FPropertyBagPropertyDesc> Descs = InstanceImpl->Variables.GetPropertyBagStruct()->GetPropertyDescs();
+		uint8* BasePtr = InstanceImpl->Variables.GetMutableValue().GetMemory();
+		for(int32 VariableIndex = 0; VariableIndex < NumVariables; ++VariableIndex)
 		{
-			InstanceImpl->PublicVariablesState = FAnimNextGraphInstance::EPublicVariablesState::None;
+			ExternalVariableRuntimeData.Emplace(Descs[VariableIndex].CachedProperty->ContainerPtrToValuePtr<uint8>(BasePtr));
 		}
-
-		// Now initialize the 'instance', cache memory handles etc. in the context
-		VM->InitializeInstance(InstanceImpl->ExtendedExecuteContext);
-
-		// Move our implementation so that we can use the instance to allocate our root node
-		Instance.Impl = InstanceImpl;
+		InstanceImpl->ExtendedExecuteContext.ExternalVariableRuntimeData = MoveTemp(ExternalVariableRuntimeData);
+	}
+	else
+	{
+		InstanceImpl->PublicVariablesState = FAnimNextGraphInstance::EPublicVariablesState::None;
 	}
 
+	// Now initialize the 'instance', cache memory handles etc. in the context
+	VM->InitializeInstance(InstanceImpl->ExtendedExecuteContext);
+
 	{
-		UE::AnimNext::FExecutionContext Context(Instance);
-		Instance.Impl->GraphInstancePtr = Context.AllocateNodeInstance(*Instance.Impl, ResolvedRootTraitHandle);
+		UE::AnimNext::FExecutionContext Context(*InstanceImpl.Get());
+		InstanceImpl->GraphInstancePtr = Context.AllocateNodeInstance(*InstanceImpl, ResolvedRootTraitHandle);
 	}
 
-	if (!Instance.IsValid())
+	if (!InstanceImpl->IsValid())
 	{
-		// We failed to allocate our instance, clear everything
-		Instance.Release();
+		// We failed to allocate our instance, reset the ptr
+		InstanceImpl.Reset();
 	}
 
 #if WITH_EDITORONLY_DATA
-	if (Instance.IsValid())
+	if (InstanceImpl.IsValid() && InstanceImpl->IsValid())
 	{
 		FScopeLock Lock(&GraphInstancesLock);
-		check(!GraphInstances.Contains(Instance.Impl.Get()));
-		GraphInstances.Add(Instance.Impl.Get());
+		check(!GraphInstances.Contains(InstanceImpl.Get()));
+		GraphInstances.Add(InstanceImpl.Get());
 	}
 #endif
+
+	return InstanceImpl;
 }
 
 void UAnimNextAnimationGraph::Serialize(FArchive& Ar)
diff --git a/Engine/Plugins/Experimental/AnimNext/Source/AnimNext/Private/Graph/AnimNextGraphInstancePtr.cpp b/Engine/Plugins/Experimental/AnimNext/Source/AnimNext/Private/Graph/AnimNextGraphInstancePtr.cpp
--- a/Engine/Plugins/Experimental/AnimNext/Source/AnimNext/Private/Graph/AnimNextGraphInstancePtr.cpp
+++ b/Engine/Plugins/Experimental/AnimNext/Source/AnimNext/Private/Graph/AnimNextGraphInstancePtr.cpp
@@ -7,6 +7,11 @@
 
 FAnimNextGraphInstancePtr::FAnimNextGraphInstancePtr() = default;
 
+FAnimNextGraphInstancePtr::FAnimNextGraphInstancePtr(const TSharedRef<FAnimNextGraphInstance>& InInstance)
+	: Impl(InInstance)
+{
+}
+
 FAnimNextGraphInstancePtr::~FAnimNextGraphInstancePtr()
 {
 	Release();
@@ -19,86 +24,124 @@ void FAnimNextGraphInstancePtr::Release()
 
 bool FAnimNextGraphInstancePtr::IsValid() const
 {
-	return Impl && Impl->IsValid();
+	if(TSharedPtr<FAnimNextGraphInstance> PinnedImpl = Impl.Pin())
+	{
+		return PinnedImpl->IsValid();
+	}
+	return false;
 }
 
 const UAnimNextAnimationGraph* FAnimNextGraphInstancePtr::GetAnimationGraph() const
 {
-	return Impl->GetAnimationGraph();
+	if(TSharedPtr<FAnimNextGraphInstance> PinnedImpl = Impl.Pin())
+	{
+		return PinnedImpl->GetAnimationGraph();
+	}
+	return nullptr;
 }
 
 UE::AnimNext::FWeakTraitPtr FAnimNextGraphInstancePtr::GetGraphRootPtr() const
 {
-	return Impl ? Impl->GetGraphRootPtr() : UE::AnimNext::FWeakTraitPtr();
+	if(TSharedPtr<FAnimNextGraphInstance> PinnedImpl = Impl.Pin())
+	{
+		return PinnedImpl->GetGraphRootPtr();
+	}
+	return UE::AnimNext::FWeakTraitPtr();
 }
 
 FAnimNextGraphInstance* FAnimNextGraphInstancePtr::GetImpl() const
 {
-	return Impl.Get();
+	if(TSharedPtr<FAnimNextGraphInstance> PinnedImpl = Impl.Pin())
+	{
+		return PinnedImpl.Get();
+	}
+	return nullptr;
 }
 
 bool FAnimNextGraphInstancePtr::UsesAnimationGraph(const UAnimNextAnimationGraph* InAnimationGraph) const
 {
-	return Impl ? Impl->UsesAnimationGraph(InAnimationGraph) : false;
+	if(TSharedPtr<FAnimNextGraphInstance> PinnedImpl = Impl.Pin())
+	{
+		return PinnedImpl->UsesAnimationGraph(InAnimationGraph);
+	}
+	return false;
 }
 
 bool FAnimNextGraphInstancePtr::IsRoot() const
 {
-	return Impl ? Impl->IsRoot() : true;
+	if(TSharedPtr<FAnimNextGraphInstance> PinnedImpl = Impl.Pin())
+	{
+		return PinnedImpl->IsRoot();
+	}
+	return false;
 }
 
 bool FAnimNextGraphInstancePtr::HasUpdated() const
 {
-	return Impl ? Impl->HasUpdated() : false;
+	if(TSharedPtr<FAnimNextGraphInstance> PinnedImpl = Impl.Pin())
+	{
+		return PinnedImpl->HasUpdated();
+	}
+	return false;
 }
 
 void FAnimNextGraphInstancePtr::AddStructReferencedObjects(FReferenceCollector& Collector)
 {
-	if (Impl)
+	if(TSharedPtr<FAnimNextGraphInstance> PinnedImpl = Impl.Pin())
 	{
-		FAnimNextGraphInstance* ImplPtr = Impl.Get();
+		FAnimNextGraphInstance* ImplPtr = PinnedImpl.Get();
 		Collector.AddPropertyReferences(FAnimNextGraphInstance::StaticStruct(), ImplPtr);
-		Impl->AddStructReferencedObjects(Collector);
+		PinnedImpl->AddStructReferencedObjects(Collector);
 	}
 }
 
 UE::AnimNext::FGraphInstanceComponent* FAnimNextGraphInstancePtr::TryGetComponent(int32 ComponentNameHash, FName ComponentName) const
 {
-	return Impl ? Impl->TryGetComponent(ComponentNameHash, ComponentName) : nullptr;
+	if(TSharedPtr<FAnimNextGraphInstance> PinnedImpl = Impl.Pin())
+	{
+		return PinnedImpl->TryGetComponent(ComponentNameHash, ComponentName);
+	}
+	return nullptr;
 }
 
 UE::AnimNext::FGraphInstanceComponent& FAnimNextGraphInstancePtr::AddComponent(int32 ComponentNameHash, FName ComponentName, TSharedPtr<UE::AnimNext::FGraphInstanceComponent>&& Component)
 {
-	check(Impl);
-	return Impl->AddComponent(ComponentNameHash, ComponentName, MoveTemp(Component));
+	TSharedPtr<FAnimNextGraphInstance> PinnedImpl = Impl.Pin();
+	check(PinnedImpl);
+	return PinnedImpl->AddComponent(ComponentNameHash, ComponentName, MoveTemp(Component));
 }
 
 GraphInstanceComponentMapType::TConstIterator FAnimNextGraphInstancePtr::GetComponentIterator() const
 {
-	check(Impl);
-	return Impl->GetComponentIterator();
+	TSharedPtr<FAnimNextGraphInstance> PinnedImpl = Impl.Pin();
+	check(PinnedImpl);
+	return PinnedImpl->GetComponentIterator();
 }
 
 void FAnimNextGraphInstancePtr::Update() const
 {
-	check(Impl);
-	Impl->Update();
+	TSharedPtr<FAnimNextGraphInstance> PinnedImpl = Impl.Pin();
+	check(PinnedImpl);
+	PinnedImpl->Update();
 }
 
 FRigVMExtendedExecuteContext& FAnimNextGraphInstancePtr::GetExtendedExecuteContext() const
 {
-	check(Impl);
-	return Impl->GetExtendedExecuteContext();
+	TSharedPtr<FAnimNextGraphInstance> PinnedImpl = Impl.Pin();
+	check(PinnedImpl);
+	return PinnedImpl->GetExtendedExecuteContext();
 }
 
 bool FAnimNextGraphInstancePtr::RequiresPublicVariableBinding() const
 {
-	check(Impl);
-	return Impl->PublicVariablesState == FAnimNextGraphInstance::EPublicVariablesState::Unbound;
+	TSharedPtr<FAnimNextGraphInstance> PinnedImpl = Impl.Pin();
+	check(PinnedImpl);
+	return PinnedImpl->PublicVariablesState == FAnimNextGraphInstance::EPublicVariablesState::Unbound;
 }
 
 void FAnimNextGraphInstancePtr::BindPublicVariables(TConstArrayView<UE::AnimNext::IDataInterfaceHost*> InHosts) const
 {
-	check(Impl);
-	return Impl->BindPublicVariables(InHosts);
+	TSharedPtr<FAnimNextGraphInstance> PinnedImpl = Impl.Pin();
+	check(PinnedImpl);
+	return PinnedImpl->BindPublicVariables(InHosts);
 }
diff --git a/Engine/Plugins/Experimental/AnimNext/Source/AnimNext/Private/Graph/AnimNextGraphReference.h b/Engine/Plugins/Experimental/AnimNext/Source/AnimNext/Private/Graph/AnimNextGraphReference.h
new file mode 100644
--- /dev/null
+++ b/Engine/Plugins/Experimental/AnimNext/Source/AnimNext/Private/Graph/AnimNextGraphReference.h
@@ -0,0 +1,28 @@
+﻿// Copyright Epic Games, Inc. All Rights Reserved.
+
+#pragma once
+
+#include "CoreMinimal.h"
+#include "AnimNextPool.h"
+#include "AnimNextGraphReference.generated.h"
+
+struct FAnimNextGraphInstance;
+
+// Represents a reference to an instance of an AnimNext graph
+// ~Does not 'own' the graph instance itself. Graphs are owned by their host (module, sub-graph trait etc).
+USTRUCT(BlueprintType)
+struct ANIMNEXT_API FAnimNextGraphReference
+{
+	GENERATED_BODY()
+
+	FAnimNextGraphReference() = default;
+
+	FAnimNextGraphReference(UE::AnimNext::TPoolHandle<FAnimNextGraphInstance> InHandle);
+
+	// Try to resolve this reference. Returns nullptr if the reference could not be resolved. The pointer to the reference should not be persisted
+	FAnimNextGraphInstance* Get();
+
+private:
+	// Handle to the instance
+	UE::AnimNext::TPoolHandle<FAnimNextGraphInstance> Handle;
+};
\ No newline at end of file
diff --git a/Engine/Plugins/Experimental/AnimNext/Source/AnimNext/Private/Graph/AnimationGraphComponent.cpp b/Engine/Plugins/Experimental/AnimNext/Source/AnimNext/Private/Graph/AnimationGraphComponent.cpp
new file mode 100644
--- /dev/null
+++ b/Engine/Plugins/Experimental/AnimNext/Source/AnimNext/Private/Graph/AnimationGraphComponent.cpp
@@ -0,0 +1,40 @@
+﻿// Copyright Epic Games, Inc. All Rights Reserved.
+
+#include "Graph/AnimationGraphComponent.h"
+
+#include "Graph/AnimNextAnimationGraph.h"
+
+namespace UE::AnimNext
+{
+
+FAnimationGraphComponent::FAnimationGraphComponent(FAnimNextModuleInstance& InOwnerInstance)
+	: FModuleInstanceComponent(InOwnerInstance)
+{
+}
+
+FAnimNextGraphInstancePtr FAnimationGraphComponent::AllocateInstance(UAnimNextAnimationGraph* InAnimationGraph, FAnimNextGraphInstance* InParentInstance, FName InEntryPoint)
+{
+	check(InAnimationGraph);
+
+	TSharedPtr<FAnimNextGraphInstance> NewInstance = InAnimationGraph->AllocateInstance(&GetModuleInstance(), InParentInstance, InEntryPoint);
+	if(!NewInstance.IsValid())
+	{
+		return FAnimNextGraphInstancePtr();
+	}
+
+	GraphInstances.Add(NewInstance);
+	return FAnimNextGraphInstancePtr(NewInstance.ToSharedRef());
+}
+
+void FAnimationGraphComponent::ReleaseInstance(FAnimNextGraphInstancePtr InInstance)
+{
+	TSharedPtr<FAnimNextGraphInstance> PinnedInstance = InInstance.Impl.Pin();
+	if(PinnedInstance.IsValid())
+	{
+		check(GraphInstances.Contains(PinnedInstance));			// Should not be releasing this instance if it is not owned by this module
+		GraphInstances.Remove(PinnedInstance);
+		check(PinnedInstance.GetSharedReferenceCount() == 1);	// This should be the final reference, all others should be weak
+	}
+}
+
+}
diff --git a/Engine/Plugins/Experimental/AnimNext/Source/AnimNext/Private/Graph/RigUnit_AnimNextRunAnimationGraph.cpp b/Engine/Plugins/Experimental/AnimNext/Source/AnimNext/Private/Graph/RigUnit_AnimNextRunAnimationGraph.cpp
--- a/Engine/Plugins/Experimental/AnimNext/Source/AnimNext/Private/Graph/RigUnit_AnimNextRunAnimationGraph.cpp
+++ b/Engine/Plugins/Experimental/AnimNext/Source/AnimNext/Private/Graph/RigUnit_AnimNextRunAnimationGraph.cpp
@@ -5,6 +5,7 @@
 #include "Units/RigUnitContext.h"
 #include "AnimNextStats.h"
 #include "IAnimNextModuleInterface.h"
+#include "Graph/AnimationGraphComponent.h"
 #include "Module/AnimNextModuleInstance.h"
 #include "Graph/AnimNextAnimationGraph.h"
 #include "Graph/RigVMTrait_AnimNextPublicVariables.h"
@@ -28,21 +29,24 @@ FRigUnit_AnimNextRunAnimationGraph_Execute()
 		return;
 	}
 
+	FAnimationGraphComponent& AnimationGraphComponent = ModuleInstance.GetComponent<FAnimationGraphComponent>();
+
 	if(Graph == nullptr)
 	{
-		Instance.Release();
+		AnimationGraphComponent.ReleaseInstance(Instance);
 		return;
 	}
 
-	// Lazily (re-)allocate graph instance if required
-	if(!Instance.IsValid() || !Instance.UsesAnimationGraph(Graph))
+	// Release the instance if the graph has changed
+	if(Instance.IsValid() && !Instance.UsesAnimationGraph(Graph))
 	{
-		Graph->AllocateInstance(Instance, &ModuleInstance);
+		AnimationGraphComponent.ReleaseInstance(Instance);
+	}
 
-		if(!Instance.IsValid())
-		{
-			return;
-		}
+	// Lazily (re-)allocate graph instance if required
+	if(!Instance.IsValid())
+	{
+		Instance = AnimationGraphComponent.AllocateInstance(Graph);
 	}
 
 	if(!Instance.IsValid())
@@ -104,8 +108,8 @@ FRigUnit_AnimNextRunAnimationGraph_Execute()
 	MemStack.Alloc(size_t(FPageAllocator::SmallPageSize) + 1, 16);
 
 	IAnimNextModuleInterface& AnimNextModule = IAnimNextModuleInterface::Get();
-	AnimNextModule.UpdateGraph(Instance, ExecuteContext.GetDeltaTime(), InputEventList, OutputEventList);
-	AnimNextModule.EvaluateGraph(Instance, RefPose, LOD, Result);
+	AnimNextModule.UpdateGraph(*Instance.GetImpl(), ExecuteContext.GetDeltaTime(), InputEventList, OutputEventList);
+	AnimNextModule.EvaluateGraph(*Instance.GetImpl(), RefPose, LOD, Result);
 
 	// We might have appended new input/output events, append them
 	{
diff --git a/Engine/Plugins/Experimental/AnimNext/Source/AnimNext/Private/Graph/RigUnit_AnimNextWriteSkeletalMeshComponentPose.h b/Engine/Plugins/Experimental/AnimNext/Source/AnimNext/Private/Graph/RigUnit_AnimNextWriteSkeletalMeshComponentPose.h
--- a/Engine/Plugins/Experimental/AnimNext/Source/AnimNext/Private/Graph/RigUnit_AnimNextWriteSkeletalMeshComponentPose.h
+++ b/Engine/Plugins/Experimental/AnimNext/Source/AnimNext/Private/Graph/RigUnit_AnimNextWriteSkeletalMeshComponentPose.h
@@ -5,7 +5,6 @@
 #include "AnimNextExecuteContext.h"
 #include "Graph/AnimNext_LODPose.h"
 #include "Graph/RigUnit_AnimNextBase.h"
-#include "Graph/AnimNextGraphInstancePtr.h"
 
 #include "RigUnit_AnimNextWriteSkeletalMeshComponentPose.generated.h"
 
diff --git a/Engine/Plugins/Experimental/AnimNext/Source/AnimNext/Private/Graph/RigUnit_MakeReferencePoseFromSkeletalMeshComponent.h b/Engine/Plugins/Experimental/AnimNext/Source/AnimNext/Private/Graph/RigUnit_MakeReferencePoseFromSkeletalMeshComponent.h
--- a/Engine/Plugins/Experimental/AnimNext/Source/AnimNext/Private/Graph/RigUnit_MakeReferencePoseFromSkeletalMeshComponent.h
+++ b/Engine/Plugins/Experimental/AnimNext/Source/AnimNext/Private/Graph/RigUnit_MakeReferencePoseFromSkeletalMeshComponent.h
@@ -5,7 +5,6 @@
 #include "AnimNextExecuteContext.h"
 #include "Graph/AnimNext_LODPose.h"
 #include "Graph/RigUnit_AnimNextBase.h"
-#include "Graph/AnimNextGraphInstancePtr.h"
 
 #include "RigUnit_MakeReferencePoseFromSkeletalMeshComponent.generated.h"
 
diff --git a/Engine/Plugins/Experimental/AnimNext/Source/AnimNext/Private/Module/AnimNextModuleInstance.cpp b/Engine/Plugins/Experimental/AnimNext/Source/AnimNext/Private/Module/AnimNextModuleInstance.cpp
--- a/Engine/Plugins/Experimental/AnimNext/Source/AnimNext/Private/Module/AnimNextModuleInstance.cpp
+++ b/Engine/Plugins/Experimental/AnimNext/Source/AnimNext/Private/Module/AnimNextModuleInstance.cpp
@@ -383,6 +383,21 @@ void FAnimNextModuleInstance::OnModuleCompiled()
 }
 #endif
 
+UE::AnimNext::FModuleInstanceComponent* FAnimNextModuleInstance::TryGetComponent(int32 ComponentNameHash, FName ComponentName) const
+{
+	if (const TSharedPtr<UE::AnimNext::FModuleInstanceComponent>* Component = Components.FindByHash(ComponentNameHash, ComponentName))
+	{
+		return Component->Get();
+	}
+
+	return nullptr;
+}
+
+UE::AnimNext::FModuleInstanceComponent& FAnimNextModuleInstance::AddComponent(int32 ComponentNameHash, FName ComponentName, TSharedPtr<UE::AnimNext::FModuleInstanceComponent>&& Component)
+{
+	return *Components.AddByHash(ComponentNameHash, ComponentName, MoveTemp(Component)).Get();
+}
+
 #if UE_ENABLE_DEBUG_DRAWING
 FRigVMDrawInterface& FAnimNextModuleInstance::GetDebugDrawInterface()
 {
diff --git a/Engine/Plugins/Experimental/AnimNext/Source/AnimNext/Public/Graph/AnimNextGraphInstance.h b/Engine/Plugins/Experimental/AnimNext/Source/AnimNext/Public/Graph/AnimNextGraphInstance.h
--- a/Engine/Plugins/Experimental/AnimNext/Source/AnimNext/Public/Graph/AnimNextGraphInstance.h
+++ b/Engine/Plugins/Experimental/AnimNext/Source/AnimNext/Public/Graph/AnimNextGraphInstance.h
@@ -30,6 +30,8 @@ namespace UE::AnimNext
 	struct FGraphInstanceComponent;
 	struct FLatentPropertyHandle;
 	struct FTraitStackBinding;
+	struct FPlayAnimSlotTrait;
+	struct FBlendStackCoreTrait;
 }
 
 using GraphInstanceComponentMapType = TMap<FName, TSharedPtr<UE::AnimNext::FGraphInstanceComponent>>;
@@ -192,6 +194,8 @@ private:
 	friend UE::AnimNext::FExecutionContext;
 	friend FAnimNextGraphInstancePtr;
 	friend UE::AnimNext::FTraitStackBinding;
+	friend UE::AnimNext::FPlayAnimSlotTrait;
+	friend UE::AnimNext::FBlendStackCoreTrait;
 };
 
 template<>
diff --git a/Engine/Plugins/Experimental/AnimNext/Source/AnimNext/Public/Graph/AnimNextGraphInstancePtr.h b/Engine/Plugins/Experimental/AnimNext/Source/AnimNext/Public/Graph/AnimNextGraphInstancePtr.h
--- a/Engine/Plugins/Experimental/AnimNext/Source/AnimNext/Public/Graph/AnimNextGraphInstancePtr.h
+++ b/Engine/Plugins/Experimental/AnimNext/Source/AnimNext/Public/Graph/AnimNextGraphInstancePtr.h
@@ -23,6 +23,7 @@ namespace UE::AnimNext
 	struct FExecutionContext;
 	struct FGraphInstanceComponent;
 	class IDataInterfaceHost;
+	struct FAnimationGraphComponent;
 }
 
 using GraphInstanceComponentMapType = TMap<FName, TSharedPtr<UE::AnimNext::FGraphInstanceComponent>>;
@@ -41,7 +42,7 @@ struct ANIMNEXT_API FAnimNextGraphInstancePtr
 	// If the graph instance is allocated, we release it during destruction
 	~FAnimNextGraphInstancePtr();
 
-	// Releases the graph instance and frees all corresponding memory
+	// Releases the weak ptr to the graph instance
 	void Release();
 
 	// Returns true if we have a live graph instance, false otherwise
@@ -53,7 +54,7 @@ struct ANIMNEXT_API FAnimNextGraphInstancePtr
 	// Returns a weak handle to the root trait instance
 	UE::AnimNext::FWeakTraitPtr GetGraphRootPtr() const;
 
-	// Returns the graph instance implementation
+	// Returns the graph instance implementation. If the instance is invalid (i.e. its owner has released it) this returns nullptr
 	FAnimNextGraphInstance* GetImpl() const;
 
 	// Check to see if this instance data matches the provided animation graph
@@ -96,16 +97,21 @@ struct ANIMNEXT_API FAnimNextGraphInstancePtr
 	void BindPublicVariables(TConstArrayView<UE::AnimNext::IDataInterfaceHost*> InHosts) const;
 
 private:
+	// Private constructor for use by owners
+	FAnimNextGraphInstancePtr(const TSharedRef<FAnimNextGraphInstance>& InInstance);
+
 	// Returns a pointer to the specified component, or nullptr if not found
 	UE::AnimNext::FGraphInstanceComponent* TryGetComponent(int32 ComponentNameHash, FName ComponentName) const;
 
 	// Adds the specified component and returns a reference to it
 	UE::AnimNext::FGraphInstanceComponent& AddComponent(int32 ComponentNameHash, FName ComponentName, TSharedPtr<UE::AnimNext::FGraphInstanceComponent>&& Component);
 
-	// Indirection to hide implementation details and to fix the graph instance into a single memory location
-	TSharedPtr<FAnimNextGraphInstance> Impl;
+	// Weak ptr used to ensure references dont persist in work memory when passing this around in RivGM graphs
+	// Shared ptr owned ultimately by the module/sub-graph trait
+	TWeakPtr<FAnimNextGraphInstance> Impl;
 
-	friend UAnimNextAnimationGraph;			// The graph is the one that allocates instances
+	friend UAnimNextAnimationGraph;				// The graph is the one that allocates instances
+	friend UE::AnimNext::FAnimationGraphComponent; // The module component owns the instances
 	friend FRigUnit_AnimNextGraphEvaluator;		// We evaluate the instance
 	friend FRigUnit_AnimNextRunAnimationGraph;	// We evaluate the instance
 	friend UE::AnimNext::FExecutionContext;
diff --git a/Engine/Plugins/Experimental/AnimNext/Source/AnimNext/Public/IAnimNextModuleInterface.h b/Engine/Plugins/Experimental/AnimNext/Source/AnimNext/Public/IAnimNextModuleInterface.h
--- a/Engine/Plugins/Experimental/AnimNext/Source/AnimNext/Public/IAnimNextModuleInterface.h
+++ b/Engine/Plugins/Experimental/AnimNext/Source/AnimNext/Public/IAnimNextModuleInterface.h
@@ -6,7 +6,7 @@
 #include "Modules/ModuleInterface.h"
 #include "ReferencePose.h"
 
-struct FAnimNextGraphInstancePtr;
+struct FAnimNextGraphInstance;
 struct FAnimNextVariableBindingData;
 struct FAnimNextGraphLODPose;
 
@@ -25,8 +25,8 @@ class IAnimNextAnimGraph
 public:
 	virtual ~IAnimNextAnimGraph() = default;
 
-	virtual void UpdateGraph(const FAnimNextGraphInstancePtr& GraphInstance, float DeltaTime, FTraitEventList& InputEventList, FTraitEventList& OutputEventList) const = 0;
-	virtual void EvaluateGraph(const FAnimNextGraphInstancePtr& GraphInstance, const FReferencePose& RefPose, int32 GraphLODLevel, FAnimNextGraphLODPose& OutputPose) const = 0;
+	virtual void UpdateGraph(FAnimNextGraphInstance& GraphInstance, float DeltaTime, FTraitEventList& InputEventList, FTraitEventList& OutputEventList) const = 0;
+	virtual void EvaluateGraph(FAnimNextGraphInstance& GraphInstance, const FReferencePose& RefPose, int32 GraphLODLevel, FAnimNextGraphLODPose& OutputPose) const = 0;
 };
 
 class IAnimNextModuleInterface : public IModuleInterface
@@ -42,8 +42,8 @@ public:
 	static ANIMNEXT_API IAnimNextModuleInterface& Get();
 	virtual void RegisterAnimNextAnimGraph(const IAnimNextAnimGraph& InAnimGraphImpl) = 0;
 	virtual void UnregisterAnimNextAnimGraph() = 0;
-	virtual void UpdateGraph(const FAnimNextGraphInstancePtr& GraphInstance, float DeltaTime, FTraitEventList& InputEventList, FTraitEventList& OutputEventList) = 0;
-	virtual void EvaluateGraph(const FAnimNextGraphInstancePtr& GraphInstance, const FReferencePose& RefPose, int32 GraphLODLevel, FAnimNextGraphLODPose& OutputPose) const = 0;
+	virtual void UpdateGraph(FAnimNextGraphInstance& GraphInstance, float DeltaTime, FTraitEventList& InputEventList, FTraitEventList& OutputEventList) = 0;
+	virtual void EvaluateGraph(FAnimNextGraphInstance& GraphInstance, const FReferencePose& RefPose, int32 GraphLODLevel, FAnimNextGraphLODPose& OutputPose) const = 0;
 };
 
 }
diff --git a/Engine/Plugins/Experimental/AnimNext/Source/AnimNext/Public/TraitCore/ExecutionContext.h b/Engine/Plugins/Experimental/AnimNext/Source/AnimNext/Public/TraitCore/ExecutionContext.h
--- a/Engine/Plugins/Experimental/AnimNext/Source/AnimNext/Public/TraitCore/ExecutionContext.h
+++ b/Engine/Plugins/Experimental/AnimNext/Source/AnimNext/Public/TraitCore/ExecutionContext.h
@@ -49,8 +49,8 @@ namespace UE::AnimNext
 		// Creates an unbound execution context
 		FExecutionContext();
 
-		// Creates an execution context and binds it to the specified graph instance
-		explicit FExecutionContext(FAnimNextGraphInstancePtr& InGraphInstance);
+		// Creates an execution context and binds it to the specified weak graph instance
+		explicit FExecutionContext(FAnimNextGraphInstancePtr& InWeakGraphInstance);
 
 		// Creates an execution context and binds it to the specified graph instance
 		explicit FExecutionContext(FAnimNextGraphInstance& InGraphInstance);
@@ -64,7 +64,7 @@ namespace UE::AnimNext
 		// In order to be used, the execution context must be bound to a valid root graph instance
 
 		// Binds the execution context to the specified graph instance if it differs from the currently bound instance
-		void BindTo(const FAnimNextGraphInstancePtr& InGraphInstance);
+		void BindTo(const FAnimNextGraphInstancePtr& InWeakGraphInstance);
 
 		// Binds the execution context to the specified graph instance if it differs from the currently bound instance
 		void BindTo(FAnimNextGraphInstance& InGraphInstance);
@@ -77,7 +77,7 @@ namespace UE::AnimNext
 
 		// Returns whether or not this execution context is bound to the specified graph instance
 		// Returns true if the root graphs match
-		bool IsBoundTo(const FAnimNextGraphInstancePtr& InGraphInstance) const;
+		bool IsBoundTo(const FAnimNextGraphInstancePtr& InWeakGraphInstance) const;
 
 		// Returns whether or not this execution context is bound to the specified graph instance
 		// Returns true if the root graphs match
diff --git a/Engine/Plugins/Experimental/AnimNext/Source/AnimNextTestSuite/Private/AnimNextTraitBaseTest.cpp b/Engine/Plugins/Experimental/AnimNext/Source/AnimNextTestSuite/Private/AnimNextTraitBaseTest.cpp
--- a/Engine/Plugins/Experimental/AnimNext/Source/AnimNextTestSuite/Private/AnimNextTraitBaseTest.cpp
+++ b/Engine/Plugins/Experimental/AnimNext/Source/AnimNextTestSuite/Private/AnimNextTraitBaseTest.cpp
@@ -584,17 +584,16 @@ bool FAnimationAnimNextRuntimeTest_NodeLifetime::RunTest(const FString& InParame
 		// Read our graph
 		FTestUtils::LoadFromArchiveBuffer(*AnimationGraph, NodeHandles, GraphSharedDataArchiveBuffer);
 
-		FAnimNextGraphInstancePtr GraphInstance;
-		AnimationGraph->AllocateInstance(GraphInstance);
+		TSharedPtr<FAnimNextGraphInstance> GraphInstance = AnimationGraph->AllocateInstance();
 
-		FExecutionContext Context(GraphInstance);
+		FExecutionContext Context(*GraphInstance.Get());
 
 		// Validate handle bookkeeping
 		{
 			FAnimNextTraitHandle TraitHandle00(NodeHandles[0], 0);	// Point to first node, first base trait
 
 			// Allocate a node
-			FTraitPtr TraitPtr00 = Context.AllocateNodeInstance(*GraphInstance.GetImpl(), TraitHandle00);
+			FTraitPtr TraitPtr00 = Context.AllocateNodeInstance(*GraphInstance.Get(), TraitHandle00);
 			AddErrorIfFalse(TraitPtr00.IsValid(), "FAnimationAnimNextRuntimeTest_NodeLifetime -> Failed to allocate a node instance");
 			AddErrorIfFalse(TraitPtr00.GetTraitIndex() == 0, "FAnimationAnimNextRuntimeTest_NodeLifetime -> Allocated trait pointer should point to root trait");
 			AddErrorIfFalse(!TraitPtr00.IsWeak(), "FAnimationAnimNextRuntimeTest_NodeLifetime -> Allocated trait pointer should not be weak, we have no parent");
@@ -645,7 +644,7 @@ bool FAnimationAnimNextRuntimeTest_NodeLifetime::RunTest(const FString& InParame
 			FAnimNextTraitHandle TraitHandle10(NodeHandles[1], 0);		// Point to second node, first base trait
 
 			// Allocate our first node
-			FTraitPtr TraitPtr00 = Context.AllocateNodeInstance(*GraphInstance.GetImpl(), TraitHandle00);
+			FTraitPtr TraitPtr00 = Context.AllocateNodeInstance(*GraphInstance.Get(), TraitHandle00);
 			AddErrorIfFalse(TraitPtr00.IsValid(), "FAnimationAnimNextRuntimeTest_NodeLifetime -> Failed to allocate a node instance");
 
 			// Allocate a new node, using the first as a parent
@@ -681,7 +680,7 @@ bool FAnimationAnimNextRuntimeTest_NodeLifetime::RunTest(const FString& InParame
 				FAnimNextTraitHandle TraitHandle00(NodeHandles[0], 0);		// Point to first node, first base trait
 
 				// Allocate our node instance
-				FTraitPtr TraitPtr00 = Context.AllocateNodeInstance(*GraphInstance.GetImpl(), TraitHandle00);
+				FTraitPtr TraitPtr00 = Context.AllocateNodeInstance(*GraphInstance.Get(), TraitHandle00);
 				AddErrorIfFalse(TraitPtr00.IsValid(), "FAnimationAnimNextRuntimeTest_NodeLifetime -> Failed to allocate a node instance");
 
 				// Validate instance constructors
@@ -784,16 +783,15 @@ bool FAnimationAnimNextRuntimeTest_TraitStackBinding::RunTest(const FString& InP
 		// Read our graph
 		FTestUtils::LoadFromArchiveBuffer(*AnimationGraph, NodeHandles, GraphSharedDataArchiveBuffer);
 
-		FAnimNextGraphInstancePtr GraphInstance;
-		AnimationGraph->AllocateInstance(GraphInstance);
+		TSharedPtr<FAnimNextGraphInstance> GraphInstance = AnimationGraph->AllocateInstance();
 
-		FExecutionContext Context(GraphInstance);
+		FExecutionContext Context(*GraphInstance.Get());
 
 		// Validate from the first base trait
 		{
 			FAnimNextTraitHandle TraitHandle0(NodeHandles[0], 0);		// Point to first node, first base trait
 
-			FTraitPtr TraitPtr0 = Context.AllocateNodeInstance(*GraphInstance.GetImpl(), TraitHandle0);
+			FTraitPtr TraitPtr0 = Context.AllocateNodeInstance(*GraphInstance.Get(), TraitHandle0);
 			AddErrorIfFalse(TraitPtr0.IsValid(), "FAnimationAnimNextRuntimeTest_TraitStackBinding -> Failed to allocate a node instance");
 
 			FTraitStackBinding Stack0;
@@ -1004,7 +1002,7 @@ bool FAnimationAnimNextRuntimeTest_TraitStackBinding::RunTest(const FString& InP
 
 			FAnimNextTraitHandle TraitHandle3(NodeHandles[0], 3);		// Point to first node, second base trait
 
-			FTraitPtr TraitPtr3 = Context.AllocateNodeInstance(*GraphInstance.GetImpl(), TraitHandle3);
+			FTraitPtr TraitPtr3 = Context.AllocateNodeInstance(*GraphInstance.Get(), TraitHandle3);
 			AddErrorIfFalse(TraitPtr3.IsValid(), "FAnimationAnimNextRuntimeTest_TraitStackBinding -> Failed to allocate a node instance");
 
 			FTraitStackBinding Stack3;
@@ -1361,20 +1359,19 @@ bool FAnimationAnimNextRuntimeTest_TraitSerialization::RunTest(const FString& In
 			// Read our graph
 			FTestUtils::LoadFromArchiveBuffer(*AnimationGraph, NodeHandles, GraphSharedDataArchiveBuffer);
 
-			FAnimNextGraphInstancePtr GraphInstance;
-			AnimationGraph->AllocateInstance(GraphInstance);
+			TSharedPtr<FAnimNextGraphInstance> GraphInstance = AnimationGraph->AllocateInstance();
 
-			FExecutionContext Context(GraphInstance);
+			FExecutionContext Context(*GraphInstance.Get());
 
 			// Validate trait serialization
 			{
 				FAnimNextTraitHandle TraitHandle0(NodeHandles[0], 0);		// Point to first node, first base trait
 				FAnimNextTraitHandle TraitHandle1(NodeHandles[1], 0);		// Point to second node, first base trait
 
-				FTraitPtr TraitPtr0 = Context.AllocateNodeInstance(*GraphInstance.GetImpl(), TraitHandle0);
+				FTraitPtr TraitPtr0 = Context.AllocateNodeInstance(*GraphInstance.Get(), TraitHandle0);
 				AddErrorIfFalse(TraitPtr0.IsValid(), "FAnimationAnimNextRuntimeTest_TraitSerialization -> Failed to allocate a node instance");
 
-				FTraitPtr TraitPtr1 = Context.AllocateNodeInstance(*GraphInstance.GetImpl(), TraitHandle1);
+				FTraitPtr TraitPtr1 = Context.AllocateNodeInstance(*GraphInstance.Get(), TraitHandle1);
 				AddErrorIfFalse(TraitPtr1.IsValid(), "FAnimationAnimNextRuntimeTest_TraitSerialization -> Failed to allocate a node instance");
 
 				FTraitStackBinding StackNode0;
diff --git a/Engine/Plugins/Experimental/AnimNext/Source/AnimNextTestSuite/Private/AnimNextTraitEventTest.cpp b/Engine/Plugins/Experimental/AnimNext/Source/AnimNextTestSuite/Private/AnimNextTraitEventTest.cpp
--- a/Engine/Plugins/Experimental/AnimNext/Source/AnimNextTestSuite/Private/AnimNextTraitEventTest.cpp
+++ b/Engine/Plugins/Experimental/AnimNext/Source/AnimNextTestSuite/Private/AnimNextTraitEventTest.cpp
@@ -293,15 +293,14 @@ bool FAnimationAnimNextRuntimeTest_TraitEventRaising::RunTest(const FString& InP
 		// Read our graph
 		FTestUtils::LoadFromArchiveBuffer(*AnimationGraph, NodeHandles, GraphSharedDataArchiveBuffer);
 
-		FAnimNextGraphInstancePtr GraphInstance;
-		AnimationGraph->AllocateInstance(GraphInstance);
+		TSharedPtr<FAnimNextGraphInstance> GraphInstance = AnimationGraph->AllocateInstance();
 
-		FExecutionContext Context(GraphInstance);
+		FExecutionContext Context(*GraphInstance.Get());
 
 		{
 			FAnimNextTraitHandle TraitHandle0(NodeHandles[0], 0);	// Point to node base trait
 
-			FTraitPtr TraitPtr0 = Context.AllocateNodeInstance(*GraphInstance.GetImpl(), TraitHandle0);
+			FTraitPtr TraitPtr0 = Context.AllocateNodeInstance(*GraphInstance.Get(), TraitHandle0);
 			AddErrorIfFalse(TraitPtr0.IsValid(), "FAnimationAnimNextRuntimeTest_TraitEventRaising -> Failed to allocate a node instance");
 
 			// Node has 4 traits: AB base, AB add, A add, B add
diff --git a/Engine/Plugins/Experimental/AnimNextAnimGraph/Source/AnimNextAnimGraph/Internal/Graph/AnimGraph/AnimNode_AnimNextGraph.h b/Engine/Plugins/Experimental/AnimNextAnimGraph/Source/AnimNextAnimGraph/Internal/Graph/AnimGraph/AnimNode_AnimNextGraph.h
--- a/Engine/Plugins/Experimental/AnimNextAnimGraph/Source/AnimNextAnimGraph/Internal/Graph/AnimGraph/AnimNode_AnimNextGraph.h
+++ b/Engine/Plugins/Experimental/AnimNextAnimGraph/Source/AnimNextAnimGraph/Internal/Graph/AnimGraph/AnimNode_AnimNextGraph.h
@@ -4,8 +4,7 @@
 
 #include "Animation/AnimNode_CustomProperty.h"
 #include "Graph/AnimNextAnimationGraph.h"
-#include "Graph/AnimNextGraphInstancePtr.h"
-#include "TraitCore/TraitPtr.h"
+#include "Graph/AnimNextGraphInstance.h"
 #include "AnimNode_AnimNextGraph.generated.h"
 
 /**
@@ -44,8 +43,7 @@ private:
 	TObjectPtr<UAnimNextAnimationGraph> AnimationGraph;
 
 	// Our graph instance, we own it
-	UPROPERTY()
-	FAnimNextGraphInstancePtr GraphInstance;
+	TSharedPtr<FAnimNextGraphInstance> GraphInstance;
 
 	/*
 	 * Max LOD that this node is allowed to run
@@ -62,6 +60,7 @@ protected:
 public:
 
 	void PostSerialize(const FArchive& Ar);
+	void AddStructReferencedObjects(class FReferenceCollector& Collector);
 
 	friend class UAnimGraphNode_AnimNextGraph;
 };
@@ -73,5 +72,6 @@ struct TStructOpsTypeTraits<FAnimNode_AnimNextGraph> : public TStructOpsTypeTrai
 	{
 		WithCopy = false,
 		WithPostSerialize = true,
+		WithAddStructReferencedObjects = true,
 	};
 };
\ No newline at end of file
diff --git a/Engine/Plugins/Experimental/AnimNextAnimGraph/Source/AnimNextAnimGraph/Internal/Traits/BlendStackTrait.h b/Engine/Plugins/Experimental/AnimNextAnimGraph/Source/AnimNextAnimGraph/Internal/Traits/BlendStackTrait.h
--- a/Engine/Plugins/Experimental/AnimNextAnimGraph/Source/AnimNextAnimGraph/Internal/Traits/BlendStackTrait.h
+++ b/Engine/Plugins/Experimental/AnimNextAnimGraph/Source/AnimNextAnimGraph/Internal/Traits/BlendStackTrait.h
@@ -101,7 +101,7 @@ namespace UE::AnimNext
 			IBlendStack::FGraphRequest Request;
 
 			// The graph instance
-			FAnimNextGraphInstancePtr Instance;
+			TSharedPtr<FAnimNextGraphInstance> Instance;
 
 			// Our child handle
 			// Weak handle to the graph instance's root
diff --git a/Engine/Plugins/Experimental/AnimNextAnimGraph/Source/AnimNextAnimGraph/Private/Graph/AnimGraph/AnimNode_AnimNextGraph.cpp b/Engine/Plugins/Experimental/AnimNextAnimGraph/Source/AnimNextAnimGraph/Private/Graph/AnimGraph/AnimNode_AnimNextGraph.cpp
--- a/Engine/Plugins/Experimental/AnimNextAnimGraph/Source/AnimNextAnimGraph/Private/Graph/AnimGraph/AnimNode_AnimNextGraph.cpp
+++ b/Engine/Plugins/Experimental/AnimNextAnimGraph/Source/AnimNextAnimGraph/Private/Graph/AnimGraph/AnimNode_AnimNextGraph.cpp
@@ -13,6 +13,7 @@
 #include "TraitInterfaces/IEvaluate.h"
 #include "TraitInterfaces/IUpdate.h"
 #include "EvaluationVM/EvaluationVM.h"
+#include "Graph/AnimNextGraphInstance.h"
 
 #include UE_INLINE_GENERATED_CPP_BY_NAME(AnimNode_AnimNextGraph)
 
@@ -61,7 +62,7 @@ void FAnimNode_AnimNextGraph::Update_AnyThread(const FAnimationUpdateContext& Co
 
 		UE::AnimNext::FTraitEventList InputEventList;
 		UE::AnimNext::FTraitEventList OutputEventList;
-		UE::AnimNext::UpdateGraph(GraphInstance, Context.GetDeltaTime(), InputEventList, OutputEventList);
+		UE::AnimNext::UpdateGraph(*GraphInstance.Get(), Context.GetDeltaTime(), InputEventList, OutputEventList);
 	}
 
 	FAnimNode_CustomProperty::Update_AnyThread(Context);
@@ -77,9 +78,16 @@ void FAnimNode_AnimNextGraph::Initialize_AnyThread(const FAnimationInitializeCon
 
 	SourceLink.Initialize(Context);
 
-	if (!GraphInstance.IsValid() && AnimationGraph)
+	// Release the instance if the graph has changed
+	if(GraphInstance.IsValid() && !GraphInstance->UsesAnimationGraph(AnimationGraph))
 	{
-		AnimationGraph->AllocateInstance(GraphInstance);
+		GraphInstance.Reset();
+	}
+
+	// Lazily (re-)allocate graph instance if required
+	if(!GraphInstance.IsValid())
+	{
+		GraphInstance = AnimationGraph->AllocateInstance();
 	}
 
 	FAnimNode_CustomProperty::Initialize_AnyThread(Context);
@@ -115,7 +123,7 @@ void FAnimNode_AnimNextGraph::Evaluate_AnyThread(FPoseContext& Output)
 		ResultPose.LODPose = FLODPoseHeap(RefPose, LODLevel, true, Output.ExpectsAdditivePose());
 
 		{
-			const FEvaluationProgram EvaluationProgram = UE::AnimNext::EvaluateGraph(GraphInstance);
+			const FEvaluationProgram EvaluationProgram = UE::AnimNext::EvaluateGraph(*GraphInstance.Get());
 
 			FEvaluationVM EvaluationVM(EEvaluationFlags::All, RefPose, LODLevel);
 			bool bHasValidOutput = false;
@@ -175,6 +183,15 @@ void FAnimNode_AnimNextGraph::PostSerialize(const FArchive& Ar)
 	}
 }
 
+void FAnimNode_AnimNextGraph::AddStructReferencedObjects(FReferenceCollector& Collector)
+{
+	if(GraphInstance.IsValid())
+	{
+		FAnimNextGraphInstance* ImplPtr = GraphInstance.Get();
+		Collector.AddPropertyReferencesWithStructARO(FAnimNextGraphInstance::StaticStruct(), ImplPtr);
+	}
+}
+
 void FAnimNode_AnimNextGraph::InitializeProperties(const UObject* InSourceInstance, UClass* InTargetClass)
 {
 	// Build property lists
diff --git a/Engine/Plugins/Experimental/AnimNextAnimGraph/Source/AnimNextAnimGraph/Private/Module.cpp b/Engine/Plugins/Experimental/AnimNextAnimGraph/Source/AnimNextAnimGraph/Private/Module.cpp
--- a/Engine/Plugins/Experimental/AnimNextAnimGraph/Source/AnimNextAnimGraph/Private/Module.cpp
+++ b/Engine/Plugins/Experimental/AnimNextAnimGraph/Source/AnimNextAnimGraph/Private/Module.cpp
@@ -44,12 +44,12 @@ public:
 		IAnimNextModuleInterface::Get().UnregisterAnimNextAnimGraph();
 	}
 
-	virtual void UpdateGraph(const FAnimNextGraphInstancePtr& GraphInstance, float DeltaTime, UE::AnimNext::FTraitEventList& InputEventList, FTraitEventList& OutputEventList) const override
+	virtual void UpdateGraph(FAnimNextGraphInstance& GraphInstance, float DeltaTime, UE::AnimNext::FTraitEventList& InputEventList, FTraitEventList& OutputEventList) const override
 	{
 		UE::AnimNext::UpdateGraph(GraphInstance, DeltaTime, InputEventList, OutputEventList);
 	}
 
-	virtual void EvaluateGraph(const FAnimNextGraphInstancePtr& GraphInstance, const UE::AnimNext::FReferencePose& RefPose, int32 GraphLODLevel, FAnimNextGraphLODPose& OutputPose) const override
+	virtual void EvaluateGraph(FAnimNextGraphInstance& GraphInstance, const UE::AnimNext::FReferencePose& RefPose, int32 GraphLODLevel, FAnimNextGraphLODPose& OutputPose) const override
 	{
 		const FEvaluationProgram EvaluationProgram = UE::AnimNext::EvaluateGraph(GraphInstance);
 
diff --git a/Engine/Plugins/Experimental/AnimNextAnimGraph/Source/AnimNextAnimGraph/Private/PlayAnim/PlayAnimSlotTrait.cpp b/Engine/Plugins/Experimental/AnimNextAnimGraph/Source/AnimNextAnimGraph/Private/PlayAnim/PlayAnimSlotTrait.cpp
--- a/Engine/Plugins/Experimental/AnimNextAnimGraph/Source/AnimNextAnimGraph/Private/PlayAnim/PlayAnimSlotTrait.cpp
+++ b/Engine/Plugins/Experimental/AnimNextAnimGraph/Source/AnimNextAnimGraph/Private/PlayAnim/PlayAnimSlotTrait.cpp
@@ -382,12 +382,12 @@ namespace UE::AnimNext
 			if (SlotRequest.State == EPlayAnimRequestState::Active)
 			{
 				const FName EntryPoint = SlotRequest.AnimationGraph->DefaultEntryPoint;
-				SlotRequest.AnimationGraph->AllocateInstance(Binding.GetTraitPtr().GetNodeInstance()->GetOwner(), SlotRequest.GraphInstance, EntryPoint);
-				SlotRequest.ChildPtr = SlotRequest.GraphInstance.GetGraphRootPtr();
+				SlotRequest.GraphInstance = SlotRequest.AnimationGraph->AllocateInstance(nullptr, &Binding.GetTraitPtr().GetNodeInstance()->GetOwner(), EntryPoint);
+				SlotRequest.ChildPtr = SlotRequest.GraphInstance.IsValid() ? SlotRequest.GraphInstance->GetGraphRootPtr() : FTraitPtr();
 
 				// Note: args are mutable here as bindings allow writes!
 				FDataInterfaceStructAdapter VariableBinding(SlotRequest.AnimationGraph, SlotRequest.Request->GetMutableArgs().Payload);
-				SlotRequest.GraphInstance.BindPublicVariables({ &VariableBinding });
+				SlotRequest.GraphInstance->BindPublicVariables({ &VariableBinding });
 
 				// TODO: Validate that our child implements the ITimeline interface
 
@@ -420,7 +420,7 @@ namespace UE::AnimNext
 
 			if (SlotRequest.State == EPlayAnimRequestState::Active)
 			{
-				SlotRequest.GraphInstance.Release();
+				SlotRequest.GraphInstance.Reset();
 
 				{
 					// Queue our status update
@@ -526,6 +526,8 @@ namespace UE::AnimNext
 			}
 
 			Collector.AddReferencedObject(SlotRequest.AnimationGraph);
+			FAnimNextGraphInstance* ImplPtr = SlotRequest.GraphInstance.Get();
+			Collector.AddPropertyReferencesWithStructARO(FAnimNextGraphInstance::StaticStruct(), ImplPtr);
 		}
 	}
 }
diff --git a/Engine/Plugins/Experimental/AnimNextAnimGraph/Source/AnimNextAnimGraph/Private/PlayAnim/PlayAnimSlotTrait.h b/Engine/Plugins/Experimental/AnimNextAnimGraph/Source/AnimNextAnimGraph/Private/PlayAnim/PlayAnimSlotTrait.h
--- a/Engine/Plugins/Experimental/AnimNextAnimGraph/Source/AnimNextAnimGraph/Private/PlayAnim/PlayAnimSlotTrait.h
+++ b/Engine/Plugins/Experimental/AnimNextAnimGraph/Source/AnimNextAnimGraph/Private/PlayAnim/PlayAnimSlotTrait.h
@@ -12,7 +12,7 @@
 #include "TraitInterfaces/ISmoothBlend.h"
 #include "TraitInterfaces/IUpdate.h"
 #include "Graph/AnimNextAnimationGraph.h"
-#include "Graph/AnimNextGraphInstancePtr.h"
+#include "Graph/AnimNextGraphInstance.h"
 #include "PlayAnim/PlayAnimEvents.h"
 #include "Chooser.h"
 
@@ -111,7 +111,7 @@ namespace UE::AnimNext
 			TObjectPtr<const UAnimNextAnimationGraph> AnimationGraph;
 
 			// The graph instance
-			FAnimNextGraphInstancePtr GraphInstance;
+			TSharedPtr<FAnimNextGraphInstance> GraphInstance;
 
 			// Our child handle
 			// If we use the source input, this is a strong handle to it, otherwise we are a weak handle to the graph instance's root
diff --git a/Engine/Plugins/Experimental/AnimNextAnimGraph/Source/AnimNextAnimGraph/Private/TraitInterfaces/IEvaluate.cpp b/Engine/Plugins/Experimental/AnimNextAnimGraph/Source/AnimNextAnimGraph/Private/TraitInterfaces/IEvaluate.cpp
--- a/Engine/Plugins/Experimental/AnimNextAnimGraph/Source/AnimNextAnimGraph/Private/TraitInterfaces/IEvaluate.cpp
+++ b/Engine/Plugins/Experimental/AnimNextAnimGraph/Source/AnimNextAnimGraph/Private/TraitInterfaces/IEvaluate.cpp
@@ -5,6 +5,7 @@
 #include "TraitCore/ExecutionContext.h"
 #include "TraitInterfaces/IHierarchy.h"
 #include "AnimNextAnimGraphStats.h"
+#include "Graph/AnimNextGraphInstance.h"
 
 DEFINE_STAT(STAT_AnimNext_EvaluateGraph);
 
@@ -88,7 +89,7 @@ namespace UE::AnimNext
 		}
 	};
 
-	FEvaluationProgram EvaluateGraph(const FAnimNextGraphInstancePtr& GraphInstance)
+	FEvaluationProgram EvaluateGraph(const FAnimNextGraphInstance& GraphInstance)
 	{
 		return EvaluateGraph(GraphInstance.GetGraphRootPtr());
 	}
diff --git a/Engine/Plugins/Experimental/AnimNextAnimGraph/Source/AnimNextAnimGraph/Private/TraitInterfaces/IUpdate.cpp b/Engine/Plugins/Experimental/AnimNextAnimGraph/Source/AnimNextAnimGraph/Private/TraitInterfaces/IUpdate.cpp
--- a/Engine/Plugins/Experimental/AnimNextAnimGraph/Source/AnimNextAnimGraph/Private/TraitInterfaces/IUpdate.cpp
+++ b/Engine/Plugins/Experimental/AnimNextAnimGraph/Source/AnimNextAnimGraph/Private/TraitInterfaces/IUpdate.cpp
@@ -5,6 +5,7 @@
 #include "TraitInterfaces/IHierarchy.h"
 #include "Graph/GraphInstanceComponent.h"
 #include "AnimNextAnimGraphStats.h"
+#include "Graph/AnimNextGraphInstance.h"
 
 DEFINE_STAT(STAT_AnimNext_UpdateGraph);
 
@@ -472,7 +473,7 @@ namespace UE::AnimNext
 	// perform as much useful work as possible while waiting for memory, hiding its slow latency by fully
 	// leveraging out-of-order CPU execution.
 
-	void UpdateGraph(const FAnimNextGraphInstancePtr& GraphInstance, float DeltaTime, FTraitEventList& InputEventList, FTraitEventList& OutputEventList)
+	void UpdateGraph(FAnimNextGraphInstance& GraphInstance, float DeltaTime, FTraitEventList& InputEventList, FTraitEventList& OutputEventList)
 	{
 		SCOPE_CYCLE_COUNTER(STAT_AnimNext_UpdateGraph);
 		
diff --git a/Engine/Plugins/Experimental/AnimNextAnimGraph/Source/AnimNextAnimGraph/Private/Traits/BlendStackTrait.cpp b/Engine/Plugins/Experimental/AnimNextAnimGraph/Source/AnimNextAnimGraph/Private/Traits/BlendStackTrait.cpp
--- a/Engine/Plugins/Experimental/AnimNextAnimGraph/Source/AnimNextAnimGraph/Private/Traits/BlendStackTrait.cpp
+++ b/Engine/Plugins/Experimental/AnimNextAnimGraph/Source/AnimNextAnimGraph/Private/Traits/BlendStackTrait.cpp
@@ -61,7 +61,7 @@ namespace UE::AnimNext
 
 	void FBlendStackCoreTrait::FGraphState::Terminate()
 	{
-		Instance.Release();
+		Instance.Reset();
 		ChildPtr.Reset();
 		bNewlyCreated = false;
 		State = FBlendStackCoreTrait::EGraphState::Inactive;
@@ -151,8 +151,8 @@ namespace UE::AnimNext
 
 			//@TODO: Remove or implement entry points once we decide if we still need them.
 			static const FName EntryPoint = NAME_None;
-			Graph.Request.AnimationGraph->AllocateInstance(Binding.GetTraitPtr().GetNodeInstance()->GetOwner(), Graph.Instance, EntryPoint);
-			Graph.ChildPtr = Graph.Instance.GetGraphRootPtr();
+			Graph.Instance = Graph.Request.AnimationGraph->AllocateInstance(nullptr, &Binding.GetTraitPtr().GetNodeInstance()->GetOwner(), EntryPoint);
+			Graph.ChildPtr = Graph.Instance->GetGraphRootPtr();
 		}
 	}
 
@@ -264,6 +264,8 @@ namespace UE::AnimNext
 		for (FBlendStackCoreTrait::FGraphState& Graph : InstanceData->ChildGraphs)
 		{
 			Collector.AddReferencedObject(Graph.Request.AnimationGraph);
+			FAnimNextGraphInstance* ImplPtr = Graph.Instance.Get();
+			Collector.AddPropertyReferencesWithStructARO(FAnimNextGraphInstance::StaticStruct(), ImplPtr);
 		}
 	}
 
@@ -304,7 +306,7 @@ namespace UE::AnimNext
 			}
 			
 			// binding 
-			Graph.Instance.BindPublicVariables(PayloadVariableBindings);
+			Graph.Instance->BindPublicVariables(PayloadVariableBindings);
 			
 			// uninitializing the FDataInterfaceStructAdapter(s)
 			for (int32 PayloadIndex = 0; PayloadIndex < NumPayloads; ++PayloadIndex)
diff --git a/Engine/Plugins/Experimental/AnimNextAnimGraph/Source/AnimNextAnimGraph/Private/Traits/SubGraphHost.cpp b/Engine/Plugins/Experimental/AnimNextAnimGraph/Source/AnimNextAnimGraph/Private/Traits/SubGraphHost.cpp
--- a/Engine/Plugins/Experimental/AnimNextAnimGraph/Source/AnimNextAnimGraph/Private/Traits/SubGraphHost.cpp
+++ b/Engine/Plugins/Experimental/AnimNextAnimGraph/Source/AnimNextAnimGraph/Private/Traits/SubGraphHost.cpp
@@ -57,7 +57,7 @@ namespace UE::AnimNext
 			else
 			{
 				// Even if the slot is inactive, we queue an empty handle
-				Children.Add(SubGraphEntry.GraphInstance.GetGraphRootPtr());
+				Children.Add(SubGraphEntry.GraphInstance.IsValid() ? SubGraphEntry.GraphInstance->GetGraphRootPtr() : FTraitPtr());
 			}
 		}
 	}
@@ -150,6 +150,11 @@ namespace UE::AnimNext
 		for (int32 SubGraphIndex = 0; SubGraphIndex < NumSubGraphs; ++SubGraphIndex)
 		{
 			const FSubGraphSlot& SubGraphSlot = InstanceData->SubGraphSlots[SubGraphIndex];
+			if(SubGraphSlot.AnimationGraph == nullptr || !SubGraphSlot.GraphInstance.IsValid())
+			{
+				continue;
+			}
+
 			const float BlendWeight = DiscreteBlendTrait.GetBlendWeight(Context, SubGraphIndex);
 
 			FTraitUpdateState SubGraphTraitState = TraitState
@@ -157,7 +162,7 @@ namespace UE::AnimNext
 				.AsBlendingOut(SubGraphIndex != InstanceData->CurrentlyActiveSubGraphIndex)
 				.AsNewlyRelevant(!SubGraphSlot.bWasRelevant);
 
-			TraversalQueue.Push(SubGraphSlot.GraphInstance.GetGraphRootPtr(), SubGraphTraitState);
+			TraversalQueue.Push(SubGraphSlot.GraphInstance->GetGraphRootPtr(), SubGraphTraitState);
 		}
 	}
 
@@ -210,7 +215,7 @@ namespace UE::AnimNext
 
 			if (SubGraphEntry.State == ESlotState::ActiveWithGraph)
 			{
-				SubGraphEntry.AnimationGraph->AllocateInstance(Binding.GetTraitPtr().GetNodeInstance()->GetOwner(), SubGraphEntry.GraphInstance, SubGraphEntry.EntryPoint);
+				SubGraphEntry.GraphInstance = SubGraphEntry.AnimationGraph->AllocateInstance(nullptr, &Binding.GetTraitPtr().GetNodeInstance()->GetOwner(), SubGraphEntry.EntryPoint);
 			}
 		}
 	}
@@ -226,7 +231,7 @@ namespace UE::AnimNext
 
 			if (SubGraphEntry.State == ESlotState::ActiveWithGraph)
 			{
-				SubGraphEntry.GraphInstance.Release();
+				SubGraphEntry.GraphInstance.Reset();
 			}
 
 			SubGraphEntry.State = ESlotState::Inactive;
@@ -242,7 +247,9 @@ namespace UE::AnimNext
 
 		for (FSubGraphSlot& SubGraphEntry : InstanceData->SubGraphSlots)
 		{
-			Collector.AddPropertyReferencesWithStructARO(FAnimNextGraphInstancePtr::StaticStruct(), &SubGraphEntry.GraphInstance);
+			Collector.AddReferencedObject(SubGraphEntry.AnimationGraph);
+			FAnimNextGraphInstance* ImplPtr = SubGraphEntry.GraphInstance.Get();
+			Collector.AddPropertyReferencesWithStructARO(FAnimNextGraphInstance::StaticStruct(), ImplPtr);
 		}
 	}
 }
diff --git a/Engine/Plugins/Experimental/AnimNextAnimGraph/Source/AnimNextAnimGraph/Private/Traits/SubGraphHost.h b/Engine/Plugins/Experimental/AnimNextAnimGraph/Source/AnimNextAnimGraph/Private/Traits/SubGraphHost.h
--- a/Engine/Plugins/Experimental/AnimNextAnimGraph/Source/AnimNextAnimGraph/Private/Traits/SubGraphHost.h
+++ b/Engine/Plugins/Experimental/AnimNextAnimGraph/Source/AnimNextAnimGraph/Private/Traits/SubGraphHost.h
@@ -10,7 +10,7 @@
 #include "TraitInterfaces/IHierarchy.h"
 #include "TraitInterfaces/IUpdate.h"
 #include "Graph/AnimNextAnimationGraph.h"
-#include "Graph/AnimNextGraphInstancePtr.h"
+#include "Graph/AnimNextGraphInstance.h"
 
 #include "SubGraphHost.generated.h"
 
@@ -64,7 +64,7 @@ namespace UE::AnimNext
 			TObjectPtr<const UAnimNextAnimationGraph> AnimationGraph;
 
 			// The graph instance
-			FAnimNextGraphInstancePtr GraphInstance;
+			TSharedPtr<FAnimNextGraphInstance> GraphInstance;
 
 			// The entry point to use
 			FName EntryPoint;
diff --git a/Engine/Plugins/Experimental/AnimNextAnimGraph/Source/AnimNextAnimGraph/Public/TraitInterfaces/IEvaluate.h b/Engine/Plugins/Experimental/AnimNextAnimGraph/Source/AnimNextAnimGraph/Public/TraitInterfaces/IEvaluate.h
--- a/Engine/Plugins/Experimental/AnimNextAnimGraph/Source/AnimNextAnimGraph/Public/TraitInterfaces/IEvaluate.h
+++ b/Engine/Plugins/Experimental/AnimNextAnimGraph/Source/AnimNextAnimGraph/Public/TraitInterfaces/IEvaluate.h
@@ -97,7 +97,7 @@ namespace UE::AnimNext
 	 *
 	 * @see IEvaluate::PreEvaluate, IEvaluate::PostEvaluate, IHierarchy::GetChildren
 	 */
-	[[nodiscard]] ANIMNEXTANIMGRAPH_API FEvaluationProgram EvaluateGraph(const FAnimNextGraphInstancePtr& GraphInstance);
+	[[nodiscard]] ANIMNEXTANIMGRAPH_API FEvaluationProgram EvaluateGraph(const FAnimNextGraphInstance& GraphInstance);
 
 	/**
 	 * Evaluates a sub-graph starting at its root and produces an evaluation program.
diff --git a/Engine/Plugins/Experimental/AnimNextAnimGraph/Source/AnimNextAnimGraph/Public/TraitInterfaces/IUpdate.h b/Engine/Plugins/Experimental/AnimNextAnimGraph/Source/AnimNextAnimGraph/Public/TraitInterfaces/IUpdate.h
--- a/Engine/Plugins/Experimental/AnimNextAnimGraph/Source/AnimNextAnimGraph/Public/TraitInterfaces/IUpdate.h
+++ b/Engine/Plugins/Experimental/AnimNextAnimGraph/Source/AnimNextAnimGraph/Public/TraitInterfaces/IUpdate.h
@@ -185,7 +185,7 @@ namespace UE::AnimNext
 		// The root node doesn't have a parent but we need a bookkeeping list regardless
 		Private::FUpdateEventBookkeepingList* RootParentBookkeepingEntryList = nullptr;
 
-		friend ANIMNEXTANIMGRAPH_API void UpdateGraph(const FAnimNextGraphInstancePtr& GraphInstance, float DeltaTime, FTraitEventList& InputEventList, FTraitEventList& OutputEventList);
+		friend ANIMNEXTANIMGRAPH_API void UpdateGraph(FAnimNextGraphInstance& GraphInstance, float DeltaTime, FTraitEventList& InputEventList, FTraitEventList& OutputEventList);
 		friend FUpdateTraversalQueue;
 	};
 
@@ -213,7 +213,7 @@ namespace UE::AnimNext
 		// emptied and pushed onto the update stack
 		Private::FUpdateEntry* QueuedUpdateStackHead = nullptr;
 
-		friend ANIMNEXTANIMGRAPH_API void UpdateGraph(const FAnimNextGraphInstancePtr& GraphInstance, float DeltaTime, FTraitEventList& InputEventList, FTraitEventList& OutputEventList);
+		friend ANIMNEXTANIMGRAPH_API void UpdateGraph(FAnimNextGraphInstance& GraphInstance, float DeltaTime, FTraitEventList& InputEventList, FTraitEventList& OutputEventList);
 		friend FUpdateTraversalContext;
 	};
 
@@ -327,5 +327,5 @@ namespace UE::AnimNext
 	 *
 	 * @see IUpdate::PreUpdate, IUpdate::PostUpdate, IHierarchy::GetChildren
 	 */
-	ANIMNEXTANIMGRAPH_API void UpdateGraph(const FAnimNextGraphInstancePtr& GraphInstance, float DeltaTime, FTraitEventList& InputEventList, FTraitEventList& OutputEventList);
+	ANIMNEXTANIMGRAPH_API void UpdateGraph(FAnimNextGraphInstance& GraphInstance, float DeltaTime, FTraitEventList& InputEventList, FTraitEventList& OutputEventList);
 }
diff --git a/Engine/Plugins/Experimental/AnimNextAnimGraph/Source/AnimNextAnimGraphTestSuite/Private/AnimNextAnimGraphTraitEventTest.cpp b/Engine/Plugins/Experimental/AnimNextAnimGraph/Source/AnimNextAnimGraphTestSuite/Private/AnimNextAnimGraphTraitEventTest.cpp
--- a/Engine/Plugins/Experimental/AnimNextAnimGraph/Source/AnimNextAnimGraphTestSuite/Private/AnimNextAnimGraphTraitEventTest.cpp
+++ b/Engine/Plugins/Experimental/AnimNextAnimGraph/Source/AnimNextAnimGraphTestSuite/Private/AnimNextAnimGraphTraitEventTest.cpp
@@ -298,8 +298,7 @@ bool FAnimationAnimNextRuntimeTest_GraphTraitEvent::RunTest(const FString& InPar
 		// Read our graph
 		FTestUtils::LoadFromArchiveBuffer(*AnimationGraph, NodeHandles, GraphSharedDataArchiveBuffer);
 
-		FAnimNextGraphInstancePtr GraphInstance;
-		AnimationGraph->AllocateInstance(GraphInstance);
+		TSharedPtr<FAnimNextGraphInstance> GraphInstance = AnimationGraph->AllocateInstance();
 
 		{
 			// Raise EventA and EventB on graph, every node sees them
@@ -312,7 +311,7 @@ bool FAnimationAnimNextRuntimeTest_GraphTraitEvent::RunTest(const FString& InPar
 				InputEventList.Push(EventB);
 
 				FTraitEventList OutputEventList;
-				UpdateGraph(GraphInstance, 0.0333f, InputEventList, OutputEventList);
+				UpdateGraph(*GraphInstance.Get(), 0.0333f, InputEventList, OutputEventList);
 
 				UE_RETURN_ON_ERROR(EventA->VisitedTraits.Num() == 4, "FAnimationAnimNextRuntimeTest_GraphTraitEvent -> Unexpected number of traits visited");
 				AddErrorIfFalse(EventA->VisitedTraits[0] == NodeTemplateTraitListB[0], "FAnimationAnimNextRuntimeTest_GraphTraitEvent -> Unexpected trait visited");
@@ -336,7 +335,7 @@ bool FAnimationAnimNextRuntimeTest_GraphTraitEvent::RunTest(const FString& InPar
 				InputEventList.Push(EventA);
 
 				FTraitEventList OutputEventList;
-				UpdateGraph(GraphInstance, 0.0333f, InputEventList, OutputEventList);
+				UpdateGraph(*GraphInstance.Get(), 0.0333f, InputEventList, OutputEventList);
 
 				UE_RETURN_ON_ERROR(EventA->VisitedTraits.Num() == 2, "FAnimationAnimNextRuntimeTest_GraphTraitEvent -> Unexpected number of traits visited");
 				AddErrorIfFalse(EventA->VisitedTraits[0] == NodeTemplateTraitListB[0], "FAnimationAnimNextRuntimeTest_GraphTraitEvent -> Unexpected trait visited");
@@ -354,7 +353,7 @@ bool FAnimationAnimNextRuntimeTest_GraphTraitEvent::RunTest(const FString& InPar
 				InputEventList.Push(EventB);
 
 				FTraitEventList OutputEventList;
-				UpdateGraph(GraphInstance, 0.0333f, InputEventList, OutputEventList);
+				UpdateGraph(*GraphInstance.Get(), 0.0333f, InputEventList, OutputEventList);
 
 				UE_RETURN_ON_ERROR(EventA->VisitedTraits.Num() == 2, "FAnimationAnimNextRuntimeTest_GraphTraitEvent -> Unexpected number of traits visited");
 				AddErrorIfFalse(EventA->VisitedTraits[0] == NodeTemplateTraitListB[0], "FAnimationAnimNextRuntimeTest_GraphTraitEvent -> Unexpected trait visited");
@@ -376,7 +375,7 @@ bool FAnimationAnimNextRuntimeTest_GraphTraitEvent::RunTest(const FString& InPar
 				InputEventList.Push(EventB);
 
 				FTraitEventList OutputEventList;
-				UpdateGraph(GraphInstance, 0.0333f, InputEventList, OutputEventList);
+				UpdateGraph(*GraphInstance.Get(), 0.0333f, InputEventList, OutputEventList);
 
 				UE_RETURN_ON_ERROR(EventB->VisitedTraits.Num() == 4, "FAnimationAnimNextRuntimeTest_GraphTraitEvent -> Unexpected number of traits visited");
 				AddErrorIfFalse(EventB->VisitedTraits[0] == NodeTemplateTraitListB[0], "FAnimationAnimNextRuntimeTest_GraphTraitEvent -> Unexpected trait visited");
@@ -401,7 +400,7 @@ bool FAnimationAnimNextRuntimeTest_GraphTraitEvent::RunTest(const FString& InPar
 				InputEventList.Push(EventB);
 
 				FTraitEventList OutputEventList;
-				UpdateGraph(GraphInstance, 0.0333f, InputEventList, OutputEventList);
+				UpdateGraph(*GraphInstance.Get(), 0.0333f, InputEventList, OutputEventList);
 
 				UE_RETURN_ON_ERROR(EventB->VisitedTraits.Num() == 4, "FAnimationAnimNextRuntimeTest_GraphTraitEvent -> Unexpected number of traits visited");
 				AddErrorIfFalse(EventB->VisitedTraits[0] == NodeTemplateTraitListB[0], "FAnimationAnimNextRuntimeTest_GraphTraitEvent -> Unexpected trait visited");
diff --git a/Engine/Plugins/Experimental/AnimNextAnimGraph/Source/AnimNextAnimGraphTestSuite/Private/AnimNextAnimGraphTraitGraphTest.cpp b/Engine/Plugins/Experimental/AnimNextAnimGraph/Source/AnimNextAnimGraphTestSuite/Private/AnimNextAnimGraphTraitGraphTest.cpp
--- a/Engine/Plugins/Experimental/AnimNextAnimGraph/Source/AnimNextAnimGraphTestSuite/Private/AnimNextAnimGraphTraitGraphTest.cpp
+++ b/Engine/Plugins/Experimental/AnimNextAnimGraph/Source/AnimNextAnimGraphTestSuite/Private/AnimNextAnimGraphTraitGraphTest.cpp
@@ -454,8 +454,7 @@ bool FAnimationAnimNextRuntimeTest_GraphExecute::RunTest(const FString& InParame
 		Controller->SetPinDefaultValue(DecoratorPin->GetSubPins()[3]->GetPinPath(), TEXT("78"));
 		Controller->SetPinDefaultValue(DecoratorPin->GetSubPins()[4]->GetPinPath(), TEXT("142.33"));
 
-		FAnimNextGraphInstancePtr GraphInstance;
-		AnimationGraph->AllocateInstance(GraphInstance);
+		TSharedPtr<FAnimNextGraphInstance> GraphInstance = AnimationGraph->AllocateInstance();
 
 		TArray<FString> Messages;
 		FRigVMRuntimeSettings RuntimeSettings;
@@ -463,13 +462,13 @@ bool FAnimationAnimNextRuntimeTest_GraphExecute::RunTest(const FString& InParame
 		{
 			Messages.Add(Message);
 		});
-		GraphInstance.GetExtendedExecuteContext().SetRuntimeSettings(RuntimeSettings);
+		GraphInstance->GetExtendedExecuteContext().SetRuntimeSettings(RuntimeSettings);
 
 		{
 			UE::AnimNext::FTraitEventList InputEventList;
 			UE::AnimNext::FTraitEventList OutputEventList;
-			UE::AnimNext::UpdateGraph(GraphInstance, 1.0f / 30.0f, InputEventList, OutputEventList);
-			(void)UE::AnimNext::EvaluateGraph(GraphInstance);
+			UE::AnimNext::UpdateGraph(*GraphInstance.Get(), 1.0f / 30.0f, InputEventList, OutputEventList);
+			(void)UE::AnimNext::EvaluateGraph(*GraphInstance.Get());
 		}
 
 		AddErrorIfFalse(Messages.Num() == 7, "FAnimationAnimNextRuntimeTest_GraphExecute - > Unexpected message count");
@@ -589,8 +588,7 @@ bool FAnimationAnimNextRuntimeTest_GraphExecuteLatent::RunTest(const FString& In
 				DecoratorPin->GetSubPins()[6]);	// SomeOtherLatentInt32
 		}
 
-		FAnimNextGraphInstancePtr GraphInstance;
-		AnimationGraph->AllocateInstance(GraphInstance);
+		TSharedPtr<FAnimNextGraphInstance> GraphInstance = AnimationGraph->AllocateInstance();
 
 		TArray<FString> Messages;
 		FRigVMRuntimeSettings RuntimeSettings;
@@ -598,13 +596,13 @@ bool FAnimationAnimNextRuntimeTest_GraphExecuteLatent::RunTest(const FString& In
 		{
 			Messages.Add(Message);
 		});
-		GraphInstance.GetExtendedExecuteContext().SetRuntimeSettings(RuntimeSettings);
+		GraphInstance->GetExtendedExecuteContext().SetRuntimeSettings(RuntimeSettings);
 		
 		{
 			UE::AnimNext::FTraitEventList InputEventList;
 			UE::AnimNext::FTraitEventList OutputEventList;
-			UE::AnimNext::UpdateGraph(GraphInstance, 1.0f / 30.0f, InputEventList, OutputEventList);
-			(void)UE::AnimNext::EvaluateGraph(GraphInstance);
+			UE::AnimNext::UpdateGraph(*GraphInstance.Get(), 1.0f / 30.0f, InputEventList, OutputEventList);
+			(void)UE::AnimNext::EvaluateGraph(*GraphInstance.Get());
 		}
 
 		AddErrorIfFalse(Messages.Num() == 7, "FAnimationAnimNextRuntimeTest_GraphExecuteLatent - > Unexpected message count");
@@ -664,141 +662,140 @@ bool FAnimationAnimNextRuntimeTest_Variables::RunTest(const FString& InParameter
 		AddPublicVariable(AnimationGraph, "Class", UAnimNextDataInterface::StaticClass());
 		AddPublicVariable(AnimationGraph, "SoftClass", FSoftClassPath(UAnimNextDataInterface::StaticClass()));
 
-		FAnimNextGraphInstancePtr GraphInstance;
-		AnimationGraph->AllocateInstance(GraphInstance);
+		TSharedPtr<FAnimNextGraphInstance> GraphInstance = AnimationGraph->AllocateInstance();
 
 		// Bool/Integers + conversions
 		{
 			// Gets
 			bool TestBool = false;
-			UE_RETURN_ON_ERROR(GraphInstance.GetImpl()->GetVariable("Bool", TestBool) == EPropertyBagResult::Success, "FAnimationAnimNextRuntimeTest_Variables -> GetVariable failed");
+			UE_RETURN_ON_ERROR(GraphInstance->GetVariable("Bool", TestBool) == EPropertyBagResult::Success, "FAnimationAnimNextRuntimeTest_Variables -> GetVariable failed");
 			UE_RETURN_ON_ERROR(TestBool == true, "FAnimationAnimNextRuntimeTest_Variables -> Variable value did not match");
 
 			uint8 TestByte = 0;
-			UE_RETURN_ON_ERROR(GraphInstance.GetImpl()->GetVariable("Byte", TestByte) == EPropertyBagResult::Success, "FAnimationAnimNextRuntimeTest_Variables -> GetVariable failed");
+			UE_RETURN_ON_ERROR(GraphInstance->GetVariable("Byte", TestByte) == EPropertyBagResult::Success, "FAnimationAnimNextRuntimeTest_Variables -> GetVariable failed");
 			UE_RETURN_ON_ERROR(TestByte == 42, "FAnimationAnimNextRuntimeTest_Variables -> Variable value did not match");
 
 			int32 TestInt32 = 0;
-			UE_RETURN_ON_ERROR(GraphInstance.GetImpl()->GetVariable("Int32", TestInt32) == EPropertyBagResult::Success, "FAnimationAnimNextRuntimeTest_Variables -> GetVariable failed");
+			UE_RETURN_ON_ERROR(GraphInstance->GetVariable("Int32", TestInt32) == EPropertyBagResult::Success, "FAnimationAnimNextRuntimeTest_Variables -> GetVariable failed");
 			UE_RETURN_ON_ERROR(TestInt32 == -4679222, "FAnimationAnimNextRuntimeTest_Variables -> Variable value did not match");
 
 			uint32 TestUInt32 = 0;
-			UE_RETURN_ON_ERROR(GraphInstance.GetImpl()->GetVariable("UInt32", TestUInt32) == EPropertyBagResult::Success, "FAnimationAnimNextRuntimeTest_Variables -> GetVariable failed");
+			UE_RETURN_ON_ERROR(GraphInstance->GetVariable("UInt32", TestUInt32) == EPropertyBagResult::Success, "FAnimationAnimNextRuntimeTest_Variables -> GetVariable failed");
 			UE_RETURN_ON_ERROR(TestUInt32 == 3415919103, "FAnimationAnimNextRuntimeTest_Variables -> Variable value did not match");
 
 			int64 TestInt64 = 0;
-			UE_RETURN_ON_ERROR(GraphInstance.GetImpl()->GetVariable("Int64", TestInt64) == EPropertyBagResult::Success, "FAnimationAnimNextRuntimeTest_Variables -> GetVariable failed");
+			UE_RETURN_ON_ERROR(GraphInstance->GetVariable("Int64", TestInt64) == EPropertyBagResult::Success, "FAnimationAnimNextRuntimeTest_Variables -> GetVariable failed");
 			UE_RETURN_ON_ERROR(TestInt64 == -3415919105, "FAnimationAnimNextRuntimeTest_Variables -> Variable value did not match");
 
 			uint64 TestUInt64 = 0;
-			UE_RETURN_ON_ERROR(GraphInstance.GetImpl()->GetVariable("UInt64", TestUInt64) == EPropertyBagResult::Success, "FAnimationAnimNextRuntimeTest_Variables -> GetVariable failed");
+			UE_RETURN_ON_ERROR(GraphInstance->GetVariable("UInt64", TestUInt64) == EPropertyBagResult::Success, "FAnimationAnimNextRuntimeTest_Variables -> GetVariable failed");
 			UE_RETURN_ON_ERROR(TestUInt64 == 34159191067, "FAnimationAnimNextRuntimeTest_Variables -> Variable value did not match");
 
 			// Conversions
 
 			// Bool
 			TestByte = 0;
-			UE_RETURN_ON_ERROR(GraphInstance.GetImpl()->GetVariable("Bool", TestByte) == EPropertyBagResult::Success, "FAnimationAnimNextRuntimeTest_Variables -> GetVariable failed");
+			UE_RETURN_ON_ERROR(GraphInstance->GetVariable("Bool", TestByte) == EPropertyBagResult::Success, "FAnimationAnimNextRuntimeTest_Variables -> GetVariable failed");
 			UE_RETURN_ON_ERROR(TestByte == 1, "FAnimationAnimNextRuntimeTest_Variables -> Variable value did not match");
 
 			TestInt32 = 0;
-			UE_RETURN_ON_ERROR(GraphInstance.GetImpl()->GetVariable("Bool", TestInt32) == EPropertyBagResult::Success, "FAnimationAnimNextRuntimeTest_Variables -> GetVariable failed");
+			UE_RETURN_ON_ERROR(GraphInstance->GetVariable("Bool", TestInt32) == EPropertyBagResult::Success, "FAnimationAnimNextRuntimeTest_Variables -> GetVariable failed");
 			UE_RETURN_ON_ERROR(TestInt32 == 1, "FAnimationAnimNextRuntimeTest_Variables -> Variable value did not match");
 
 			TestUInt32 = 0;
-			UE_RETURN_ON_ERROR(GraphInstance.GetImpl()->GetVariable("Bool", TestUInt32) == EPropertyBagResult::Success, "FAnimationAnimNextRuntimeTest_Variables -> GetVariable failed");
+			UE_RETURN_ON_ERROR(GraphInstance->GetVariable("Bool", TestUInt32) == EPropertyBagResult::Success, "FAnimationAnimNextRuntimeTest_Variables -> GetVariable failed");
 			UE_RETURN_ON_ERROR(TestUInt32 == 1, "FAnimationAnimNextRuntimeTest_Variables -> Variable value did not match");
 
 			TestInt64 = 0;
-			UE_RETURN_ON_ERROR(GraphInstance.GetImpl()->GetVariable("Bool", TestInt64) == EPropertyBagResult::Success, "FAnimationAnimNextRuntimeTest_Variables -> GetVariable failed");
+			UE_RETURN_ON_ERROR(GraphInstance->GetVariable("Bool", TestInt64) == EPropertyBagResult::Success, "FAnimationAnimNextRuntimeTest_Variables -> GetVariable failed");
 			UE_RETURN_ON_ERROR(TestInt64 == 1, "FAnimationAnimNextRuntimeTest_Variables -> Variable value did not match");
 
 			TestUInt64 = 0;
-			UE_RETURN_ON_ERROR(GraphInstance.GetImpl()->GetVariable("Bool", TestUInt64) == EPropertyBagResult::Success, "FAnimationAnimNextRuntimeTest_Variables -> GetVariable failed");
+			UE_RETURN_ON_ERROR(GraphInstance->GetVariable("Bool", TestUInt64) == EPropertyBagResult::Success, "FAnimationAnimNextRuntimeTest_Variables -> GetVariable failed");
 			UE_RETURN_ON_ERROR(TestUInt64 == 1, "FAnimationAnimNextRuntimeTest_Variables -> Variable value did not match");
 
 			// Byte
 			TestBool = false;
-			UE_RETURN_ON_ERROR(GraphInstance.GetImpl()->GetVariable("Byte", TestBool) == EPropertyBagResult::Success, "FAnimationAnimNextRuntimeTest_Variables -> GetVariable failed");
+			UE_RETURN_ON_ERROR(GraphInstance->GetVariable("Byte", TestBool) == EPropertyBagResult::Success, "FAnimationAnimNextRuntimeTest_Variables -> GetVariable failed");
 			UE_RETURN_ON_ERROR(TestBool == true, "FAnimationAnimNextRuntimeTest_Variables -> Variable value did not match");
 
 			TestInt32 = 0;
-			UE_RETURN_ON_ERROR(GraphInstance.GetImpl()->GetVariable("Byte", TestInt32) == EPropertyBagResult::Success, "FAnimationAnimNextRuntimeTest_Variables -> GetVariable failed");
+			UE_RETURN_ON_ERROR(GraphInstance->GetVariable("Byte", TestInt32) == EPropertyBagResult::Success, "FAnimationAnimNextRuntimeTest_Variables -> GetVariable failed");
 			UE_RETURN_ON_ERROR(TestInt32 == 42, "FAnimationAnimNextRuntimeTest_Variables -> Variable value did not match");
 
 			TestUInt32 = 0;
-			UE_RETURN_ON_ERROR(GraphInstance.GetImpl()->GetVariable("Byte", TestUInt32) == EPropertyBagResult::Success, "FAnimationAnimNextRuntimeTest_Variables -> GetVariable failed");
+			UE_RETURN_ON_ERROR(GraphInstance->GetVariable("Byte", TestUInt32) == EPropertyBagResult::Success, "FAnimationAnimNextRuntimeTest_Variables -> GetVariable failed");
 			UE_RETURN_ON_ERROR(TestUInt32 == 42, "FAnimationAnimNextRuntimeTest_Variables -> Variable value did not match");
 
 			TestInt64 = 0;
-			UE_RETURN_ON_ERROR(GraphInstance.GetImpl()->GetVariable("Byte", TestInt64) == EPropertyBagResult::Success, "FAnimationAnimNextRuntimeTest_Variables -> GetVariable failed");
+			UE_RETURN_ON_ERROR(GraphInstance->GetVariable("Byte", TestInt64) == EPropertyBagResult::Success, "FAnimationAnimNextRuntimeTest_Variables -> GetVariable failed");
 			UE_RETURN_ON_ERROR(TestInt64 == 42, "FAnimationAnimNextRuntimeTest_Variables -> Variable value did not match");
 
 			TestUInt64 = 0;
-			UE_RETURN_ON_ERROR(GraphInstance.GetImpl()->GetVariable("Byte", TestUInt64) == EPropertyBagResult::Success, "FAnimationAnimNextRuntimeTest_Variables -> GetVariable failed");
+			UE_RETURN_ON_ERROR(GraphInstance->GetVariable("Byte", TestUInt64) == EPropertyBagResult::Success, "FAnimationAnimNextRuntimeTest_Variables -> GetVariable failed");
 			UE_RETURN_ON_ERROR(TestUInt64 == 42, "FAnimationAnimNextRuntimeTest_Variables -> Variable value did not match");
 
 			// Int32
 			TestBool = false;
-			UE_RETURN_ON_ERROR(GraphInstance.GetImpl()->GetVariable("Int32", TestBool) == EPropertyBagResult::Success, "FAnimationAnimNextRuntimeTest_Variables -> GetVariable failed");
+			UE_RETURN_ON_ERROR(GraphInstance->GetVariable("Int32", TestBool) == EPropertyBagResult::Success, "FAnimationAnimNextRuntimeTest_Variables -> GetVariable failed");
 			UE_RETURN_ON_ERROR(TestBool == true, "FAnimationAnimNextRuntimeTest_Variables -> Variable value did not match");
 
 			TestByte = 0;
-			UE_RETURN_ON_ERROR(GraphInstance.GetImpl()->GetVariable("Int32", TestByte) == EPropertyBagResult::Success, "FAnimationAnimNextRuntimeTest_Variables -> GetVariable failed");
+			UE_RETURN_ON_ERROR(GraphInstance->GetVariable("Int32", TestByte) == EPropertyBagResult::Success, "FAnimationAnimNextRuntimeTest_Variables -> GetVariable failed");
 			UE_RETURN_ON_ERROR(TestByte == 202/*(uint8)-4679222*/, "FAnimationAnimNextRuntimeTest_Variables -> Variable value did not match");
 
 			TestUInt32 = 0;
-			UE_RETURN_ON_ERROR(GraphInstance.GetImpl()->GetVariable("Int32", TestUInt32) == EPropertyBagResult::Success, "FAnimationAnimNextRuntimeTest_Variables -> GetVariable failed");
+			UE_RETURN_ON_ERROR(GraphInstance->GetVariable("Int32", TestUInt32) == EPropertyBagResult::Success, "FAnimationAnimNextRuntimeTest_Variables -> GetVariable failed");
 			UE_RETURN_ON_ERROR(TestUInt32 == 4290288074/*(uint32)-4679222*/, "FAnimationAnimNextRuntimeTest_Variables -> Variable value did not match");
 
 			TestInt64 = 0;
-			UE_RETURN_ON_ERROR(GraphInstance.GetImpl()->GetVariable("Int32", TestInt64) == EPropertyBagResult::Success, "FAnimationAnimNextRuntimeTest_Variables -> GetVariable failed");
+			UE_RETURN_ON_ERROR(GraphInstance->GetVariable("Int32", TestInt64) == EPropertyBagResult::Success, "FAnimationAnimNextRuntimeTest_Variables -> GetVariable failed");
 			UE_RETURN_ON_ERROR(TestInt64 == -4679222, "FAnimationAnimNextRuntimeTest_Variables -> Variable value did not match");
 
 			TestUInt64 = 0;
-			UE_RETURN_ON_ERROR(GraphInstance.GetImpl()->GetVariable("Int32", TestUInt64) == EPropertyBagResult::Success, "FAnimationAnimNextRuntimeTest_Variables -> GetVariable failed");
+			UE_RETURN_ON_ERROR(GraphInstance->GetVariable("Int32", TestUInt64) == EPropertyBagResult::Success, "FAnimationAnimNextRuntimeTest_Variables -> GetVariable failed");
 			UE_RETURN_ON_ERROR(TestUInt64 == 4290288074/*(uint64)-4679222*/, "FAnimationAnimNextRuntimeTest_Variables -> Variable value did not match");
 
 			// Uint32
 			TestBool = false;
-			UE_RETURN_ON_ERROR(GraphInstance.GetImpl()->GetVariable("UInt32", TestBool) == EPropertyBagResult::Success, "FAnimationAnimNextRuntimeTest_Variables -> GetVariable failed");
+			UE_RETURN_ON_ERROR(GraphInstance->GetVariable("UInt32", TestBool) == EPropertyBagResult::Success, "FAnimationAnimNextRuntimeTest_Variables -> GetVariable failed");
 			UE_RETURN_ON_ERROR(TestBool == true, "FAnimationAnimNextRuntimeTest_Variables -> Variable value did not match");
 
 			TestByte = 0;
-			UE_RETURN_ON_ERROR(GraphInstance.GetImpl()->GetVariable("UInt32", TestByte) == EPropertyBagResult::Success, "FAnimationAnimNextRuntimeTest_Variables -> GetVariable failed");
+			UE_RETURN_ON_ERROR(GraphInstance->GetVariable("UInt32", TestByte) == EPropertyBagResult::Success, "FAnimationAnimNextRuntimeTest_Variables -> GetVariable failed");
 			UE_RETURN_ON_ERROR(TestByte == 255/*(uint8)3415919103*/, "FAnimationAnimNextRuntimeTest_Variables -> Variable value did not match");
 
 			TestInt32 = 0;
-			UE_RETURN_ON_ERROR(GraphInstance.GetImpl()->GetVariable("UInt32", TestInt32) == EPropertyBagResult::Success, "FAnimationAnimNextRuntimeTest_Variables -> GetVariable failed");
+			UE_RETURN_ON_ERROR(GraphInstance->GetVariable("UInt32", TestInt32) == EPropertyBagResult::Success, "FAnimationAnimNextRuntimeTest_Variables -> GetVariable failed");
 			UE_RETURN_ON_ERROR(TestInt32 == -879048193/*(int32)3415919103*/, "FAnimationAnimNextRuntimeTest_Variables -> Variable value did not match");
 
 			TestInt64 = 0;
-			UE_RETURN_ON_ERROR(GraphInstance.GetImpl()->GetVariable("UInt32", TestInt64) == EPropertyBagResult::Success, "FAnimationAnimNextRuntimeTest_Variables -> GetVariable failed");
+			UE_RETURN_ON_ERROR(GraphInstance->GetVariable("UInt32", TestInt64) == EPropertyBagResult::Success, "FAnimationAnimNextRuntimeTest_Variables -> GetVariable failed");
 			UE_RETURN_ON_ERROR(TestInt64 == 3415919103, "FAnimationAnimNextRuntimeTest_Variables -> Variable value did not match");
 
 			TestUInt64 = 0;
-			UE_RETURN_ON_ERROR(GraphInstance.GetImpl()->GetVariable("UInt32", TestUInt64) == EPropertyBagResult::Success, "FAnimationAnimNextRuntimeTest_Variables -> GetVariable failed");
+			UE_RETURN_ON_ERROR(GraphInstance->GetVariable("UInt32", TestUInt64) == EPropertyBagResult::Success, "FAnimationAnimNextRuntimeTest_Variables -> GetVariable failed");
 			UE_RETURN_ON_ERROR(TestUInt64 == 3415919103, "FAnimationAnimNextRuntimeTest_Variables -> Variable value did not match");
 
 			// Int64
 			TestBool = false;
-			UE_RETURN_ON_ERROR(GraphInstance.GetImpl()->GetVariable("Int64", TestBool) == EPropertyBagResult::Success, "FAnimationAnimNextRuntimeTest_Variables -> GetVariable failed");
+			UE_RETURN_ON_ERROR(GraphInstance->GetVariable("Int64", TestBool) == EPropertyBagResult::Success, "FAnimationAnimNextRuntimeTest_Variables -> GetVariable failed");
 			UE_RETURN_ON_ERROR(TestBool == true, "FAnimationAnimNextRuntimeTest_Variables -> Variable value did not match");
 
 			TestByte = 0;
-			UE_RETURN_ON_ERROR(GraphInstance.GetImpl()->GetVariable("Int64", TestByte) == EPropertyBagResult::Success, "FAnimationAnimNextRuntimeTest_Variables -> GetVariable failed");
+			UE_RETURN_ON_ERROR(GraphInstance->GetVariable("Int64", TestByte) == EPropertyBagResult::Success, "FAnimationAnimNextRuntimeTest_Variables -> GetVariable failed");
 			UE_RETURN_ON_ERROR(TestByte == 255/*(uint8)-3415919105*/, "FAnimationAnimNextRuntimeTest_Variables -> Variable value did not match");
 
 			TestInt32 = 0;
-			UE_RETURN_ON_ERROR(GraphInstance.GetImpl()->GetVariable("Int64", TestInt32) == EPropertyBagResult::Success, "FAnimationAnimNextRuntimeTest_Variables -> GetVariable failed");
+			UE_RETURN_ON_ERROR(GraphInstance->GetVariable("Int64", TestInt32) == EPropertyBagResult::Success, "FAnimationAnimNextRuntimeTest_Variables -> GetVariable failed");
 			UE_RETURN_ON_ERROR(TestInt32 == 879048191/*(int32)-3415919105*/, "FAnimationAnimNextRuntimeTest_Variables -> Variable value did not match");
 
 			TestUInt32 = 0;
-			UE_RETURN_ON_ERROR(GraphInstance.GetImpl()->GetVariable("Int64", TestUInt32) == EPropertyBagResult::Success, "FAnimationAnimNextRuntimeTest_Variables -> GetVariable failed");
+			UE_RETURN_ON_ERROR(GraphInstance->GetVariable("Int64", TestUInt32) == EPropertyBagResult::Success, "FAnimationAnimNextRuntimeTest_Variables -> GetVariable failed");
 			UE_RETURN_ON_ERROR(TestUInt32 == 879048191/*(uint32)-3415919105*/, "FAnimationAnimNextRuntimeTest_Variables -> Variable value did not match");
 
 			TestUInt64 = 0;
-			UE_RETURN_ON_ERROR(GraphInstance.GetImpl()->GetVariable("Int64", TestUInt64) == EPropertyBagResult::Success, "FAnimationAnimNextRuntimeTest_Variables -> GetVariable failed");
+			UE_RETURN_ON_ERROR(GraphInstance->GetVariable("Int64", TestUInt64) == EPropertyBagResult::Success, "FAnimationAnimNextRuntimeTest_Variables -> GetVariable failed");
 			UE_RETURN_ON_ERROR(TestUInt64 == 18446744070293632511ull/*(uint64)-3415919105*/, "FAnimationAnimNextRuntimeTest_Variables -> Variable value did not match");
 		}
 
@@ -806,108 +803,108 @@ bool FAnimationAnimNextRuntimeTest_Variables::RunTest(const FString& InParameter
 		{
 			// Gets
 			float TestFloat = 0.0f;
-			UE_RETURN_ON_ERROR(GraphInstance.GetImpl()->GetVariable("Float", TestFloat) == EPropertyBagResult::Success, "FAnimationAnimNextRuntimeTest_Variables -> GetVariable failed");
+			UE_RETURN_ON_ERROR(GraphInstance->GetVariable("Float", TestFloat) == EPropertyBagResult::Success, "FAnimationAnimNextRuntimeTest_Variables -> GetVariable failed");
 			UE_RETURN_ON_ERROR(TestFloat == 1.0f, "FAnimationAnimNextRuntimeTest_Variables -> Variable value did not match");
 
 			double TestDouble = 0.0;
-			UE_RETURN_ON_ERROR(GraphInstance.GetImpl()->GetVariable("Double", TestDouble) == EPropertyBagResult::Success, "FAnimationAnimNextRuntimeTest_Variables -> GetVariable failed");
+			UE_RETURN_ON_ERROR(GraphInstance->GetVariable("Double", TestDouble) == EPropertyBagResult::Success, "FAnimationAnimNextRuntimeTest_Variables -> GetVariable failed");
 			UE_RETURN_ON_ERROR(TestDouble == 1.0, "FAnimationAnimNextRuntimeTest_Variables -> Variable value did not match");
 
 			// Conversions
 			TestDouble = 0.0f;
-			UE_RETURN_ON_ERROR(GraphInstance.GetImpl()->GetVariable("Float", TestDouble) == EPropertyBagResult::Success, "FAnimationAnimNextRuntimeTest_Variables -> GetVariable failed");
+			UE_RETURN_ON_ERROR(GraphInstance->GetVariable("Float", TestDouble) == EPropertyBagResult::Success, "FAnimationAnimNextRuntimeTest_Variables -> GetVariable failed");
 			UE_RETURN_ON_ERROR(TestDouble == 1.0f, "FAnimationAnimNextRuntimeTest_Variables -> Variable value did not match");
 
 			TestFloat = 0.0;
-			UE_RETURN_ON_ERROR(GraphInstance.GetImpl()->GetVariable("Double", TestFloat) == EPropertyBagResult::Success, "FAnimationAnimNextRuntimeTest_Variables -> GetVariable failed");
+			UE_RETURN_ON_ERROR(GraphInstance->GetVariable("Double", TestFloat) == EPropertyBagResult::Success, "FAnimationAnimNextRuntimeTest_Variables -> GetVariable failed");
 			UE_RETURN_ON_ERROR(TestFloat == 1.0, "FAnimationAnimNextRuntimeTest_Variables -> Variable value did not match");
 		}
 		
 		FName TestFName = NAME_None;
-		UE_RETURN_ON_ERROR(GraphInstance.GetImpl()->GetVariable("Name", TestFName) == EPropertyBagResult::Success, "FAnimationAnimNextRuntimeTest_Variables -> GetVariable failed");
+		UE_RETURN_ON_ERROR(GraphInstance->GetVariable("Name", TestFName) == EPropertyBagResult::Success, "FAnimationAnimNextRuntimeTest_Variables -> GetVariable failed");
 		UE_RETURN_ON_ERROR(TestFName == "Test", "FAnimationAnimNextRuntimeTest_Variables -> Variable value did not match");
 
 		FString TestString;
-		UE_RETURN_ON_ERROR(GraphInstance.GetImpl()->GetVariable("String", TestString) == EPropertyBagResult::Success, "FAnimationAnimNextRuntimeTest_Variables -> GetVariable failed");
+		UE_RETURN_ON_ERROR(GraphInstance->GetVariable("String", TestString) == EPropertyBagResult::Success, "FAnimationAnimNextRuntimeTest_Variables -> GetVariable failed");
 		UE_RETURN_ON_ERROR(TestString == "Test", "FAnimationAnimNextRuntimeTest_Variables -> Variable value did not match");
 
 		FText TestText;
-		UE_RETURN_ON_ERROR(GraphInstance.GetImpl()->GetVariable("Text", TestText) == EPropertyBagResult::Success, "FAnimationAnimNextRuntimeTest_Variables -> GetVariable failed");
+		UE_RETURN_ON_ERROR(GraphInstance->GetVariable("Text", TestText) == EPropertyBagResult::Success, "FAnimationAnimNextRuntimeTest_Variables -> GetVariable failed");
 		UE_RETURN_ON_ERROR(TestText.EqualTo(NSLOCTEXT("Tests", "Test", "Test")), "FAnimationAnimNextRuntimeTest_Variables -> Variable value did not match");
 
 		// Enums
 		{
 			EPropertyBagPropertyType TestEnum = EPropertyBagPropertyType::None;
-			UE_RETURN_ON_ERROR(GraphInstance.GetImpl()->GetVariable("Enum", TestEnum) == EPropertyBagResult::Success, "FAnimationAnimNextRuntimeTest_Variables -> GetVariable failed");
+			UE_RETURN_ON_ERROR(GraphInstance->GetVariable("Enum", TestEnum) == EPropertyBagResult::Success, "FAnimationAnimNextRuntimeTest_Variables -> GetVariable failed");
 			UE_RETURN_ON_ERROR(TestEnum == EPropertyBagPropertyType::Double, "FAnimationAnimNextRuntimeTest_Variables -> Variable value did not match");
 
 			// Test mismatched enum fails
 			EPropertyBagContainerType MismatchedEnum = EPropertyBagContainerType::None;
-			UE_RETURN_ON_ERROR(GraphInstance.GetImpl()->GetVariable("Enum", MismatchedEnum) != EPropertyBagResult::Success, "FAnimationAnimNextRuntimeTest_Variables -> GetVariable succeeded");
+			UE_RETURN_ON_ERROR(GraphInstance->GetVariable("Enum", MismatchedEnum) != EPropertyBagResult::Success, "FAnimationAnimNextRuntimeTest_Variables -> GetVariable succeeded");
 		}
 
 		// Structs
 		{
 			FVector TestStruct = FVector::ZeroVector;
-			UE_RETURN_ON_ERROR(GraphInstance.GetImpl()->GetVariable("Struct", TestStruct) == EPropertyBagResult::Success, "FAnimationAnimNextRuntimeTest_Variables -> GetVariable failed");
+			UE_RETURN_ON_ERROR(GraphInstance->GetVariable("Struct", TestStruct) == EPropertyBagResult::Success, "FAnimationAnimNextRuntimeTest_Variables -> GetVariable failed");
 			UE_RETURN_ON_ERROR(TestStruct == FVector::OneVector, "FAnimationAnimNextRuntimeTest_Variables -> Variable value did not match");
 
 			// Test conversion from base -> derived fails 
 			FTestDerivedVector TestDerived;
-			UE_RETURN_ON_ERROR(GraphInstance.GetImpl()->GetVariable("Struct", TestDerived) != EPropertyBagResult::Success, "FAnimationAnimNextRuntimeTest_Variables -> GetVariable succeeded");
+			UE_RETURN_ON_ERROR(GraphInstance->GetVariable("Struct", TestDerived) != EPropertyBagResult::Success, "FAnimationAnimNextRuntimeTest_Variables -> GetVariable succeeded");
 
 			// Test conversion from derived -> base
-			UE_RETURN_ON_ERROR(GraphInstance.GetImpl()->GetVariable("DerivedStruct", TestStruct) == EPropertyBagResult::Success, "FAnimationAnimNextRuntimeTest_Variables -> GetVariable failed");
+			UE_RETURN_ON_ERROR(GraphInstance->GetVariable("DerivedStruct", TestStruct) == EPropertyBagResult::Success, "FAnimationAnimNextRuntimeTest_Variables -> GetVariable failed");
 			UE_RETURN_ON_ERROR(TestStruct == FVector::OneVector, "FAnimationAnimNextRuntimeTest_Variables -> Variable value did not match");
 		}
 
 		// Objects
 		{
 			UAnimNextDataInterface* TestObject = nullptr;
-			UE_RETURN_ON_ERROR(GraphInstance.GetImpl()->GetVariable("Object", TestObject) == EPropertyBagResult::Success, "FAnimationAnimNextRuntimeTest_Variables -> GetVariable failed");
+			UE_RETURN_ON_ERROR(GraphInstance->GetVariable("Object", TestObject) == EPropertyBagResult::Success, "FAnimationAnimNextRuntimeTest_Variables -> GetVariable failed");
 			UE_RETURN_ON_ERROR(TestObject == UAnimNextDataInterface::StaticClass()->GetDefaultObject<UAnimNextDataInterface>(), "FAnimationAnimNextRuntimeTest_Variables -> Variable value did not match");
 
 			// Test unrelated object fails
 			UAnimSequence* TestAnimSequence = nullptr;
-			UE_RETURN_ON_ERROR(GraphInstance.GetImpl()->GetVariable("Object", TestAnimSequence) != EPropertyBagResult::Success, "FAnimationAnimNextRuntimeTest_Variables -> GetVariable succeeded");
+			UE_RETURN_ON_ERROR(GraphInstance->GetVariable("Object", TestAnimSequence) != EPropertyBagResult::Success, "FAnimationAnimNextRuntimeTest_Variables -> GetVariable succeeded");
 
 			// Test Derived -> Base succeeds
 			UObject* BaseObject = nullptr;
-			UE_RETURN_ON_ERROR(GraphInstance.GetImpl()->GetVariable("Object", BaseObject) == EPropertyBagResult::Success, "FAnimationAnimNextRuntimeTest_Variables -> GetVariable failed");
+			UE_RETURN_ON_ERROR(GraphInstance->GetVariable("Object", BaseObject) == EPropertyBagResult::Success, "FAnimationAnimNextRuntimeTest_Variables -> GetVariable failed");
 			UE_RETURN_ON_ERROR(BaseObject == UAnimNextDataInterface::StaticClass()->GetDefaultObject<UAnimNextDataInterface>(), "FAnimationAnimNextRuntimeTest_Variables -> Variable value did not match");
 
 			// Test Base -> Derived fails
 			UAnimNextAnimationGraph* DerivedObject = nullptr;
-			UE_RETURN_ON_ERROR(GraphInstance.GetImpl()->GetVariable("Object", DerivedObject) != EPropertyBagResult::Success, "FAnimationAnimNextRuntimeTest_Variables -> GetVariable succeeded");
+			UE_RETURN_ON_ERROR(GraphInstance->GetVariable("Object", DerivedObject) != EPropertyBagResult::Success, "FAnimationAnimNextRuntimeTest_Variables -> GetVariable succeeded");
 		}
 
 		FSoftObjectPath TestSoftObjectPath;
-		UE_RETURN_ON_ERROR(GraphInstance.GetImpl()->GetVariable("SoftObject", TestSoftObjectPath) == EPropertyBagResult::Success, "FAnimationAnimNextRuntimeTest_Variables -> GetVariable failed");
+		UE_RETURN_ON_ERROR(GraphInstance->GetVariable("SoftObject", TestSoftObjectPath) == EPropertyBagResult::Success, "FAnimationAnimNextRuntimeTest_Variables -> GetVariable failed");
 		UE_RETURN_ON_ERROR(TestSoftObjectPath == FSoftObjectPath(UAnimNextDataInterface::StaticClass()->GetDefaultObject<UAnimNextDataInterface>()), "FAnimationAnimNextRuntimeTest_Variables -> Variable value did not match");
 
 		// Classes
 		{
 			UClass* TestClass = nullptr;
-			UE_RETURN_ON_ERROR(GraphInstance.GetImpl()->GetVariable("Class", TestClass) == EPropertyBagResult::Success, "FAnimationAnimNextRuntimeTest_Variables -> GetVariable failed");
+			UE_RETURN_ON_ERROR(GraphInstance->GetVariable("Class", TestClass) == EPropertyBagResult::Success, "FAnimationAnimNextRuntimeTest_Variables -> GetVariable failed");
 			UE_RETURN_ON_ERROR(TestClass == UAnimNextDataInterface::StaticClass(), "FAnimationAnimNextRuntimeTest_Variables -> Variable value did not match");
 
 
 			// Test unrelated class fails
 			TSubclassOf<UAnimSequence> TestUnrelatedSubclassOf;
-			UE_RETURN_ON_ERROR(GraphInstance.GetImpl()->GetVariable("Class", TestUnrelatedSubclassOf) != EPropertyBagResult::Success, "FAnimationAnimNextRuntimeTest_Variables -> GetVariable succeeded");
+			UE_RETURN_ON_ERROR(GraphInstance->GetVariable("Class", TestUnrelatedSubclassOf) != EPropertyBagResult::Success, "FAnimationAnimNextRuntimeTest_Variables -> GetVariable succeeded");
 
 			// Test Derived -> Base succeeds
 			TSubclassOf<UObject> BaseSubclassOf;
-			UE_RETURN_ON_ERROR(GraphInstance.GetImpl()->GetVariable("Class", BaseSubclassOf) == EPropertyBagResult::Success, "FAnimationAnimNextRuntimeTest_Variables -> GetVariable failed");
+			UE_RETURN_ON_ERROR(GraphInstance->GetVariable("Class", BaseSubclassOf) == EPropertyBagResult::Success, "FAnimationAnimNextRuntimeTest_Variables -> GetVariable failed");
 			UE_RETURN_ON_ERROR(BaseSubclassOf == UAnimNextDataInterface::StaticClass(), "FAnimationAnimNextRuntimeTest_Variables -> Variable value did not match");
 
 			// Test Base -> Derived fails
 			TSubclassOf<UAnimNextAnimationGraph> DerivedSubclassOf;
-			UE_RETURN_ON_ERROR(GraphInstance.GetImpl()->GetVariable("Class", DerivedSubclassOf) != EPropertyBagResult::Success, "FAnimationAnimNextRuntimeTest_Variables -> GetVariable succeeded");
+			UE_RETURN_ON_ERROR(GraphInstance->GetVariable("Class", DerivedSubclassOf) != EPropertyBagResult::Success, "FAnimationAnimNextRuntimeTest_Variables -> GetVariable succeeded");
 		}
 
 		FSoftClassPath TestSoftClassPath;
-		UE_RETURN_ON_ERROR(GraphInstance.GetImpl()->GetVariable("SoftClass", TestSoftClassPath) == EPropertyBagResult::Success, "FAnimationAnimNextRuntimeTest_Variables -> GetVariable failed");
+		UE_RETURN_ON_ERROR(GraphInstance->GetVariable("SoftClass", TestSoftClassPath) == EPropertyBagResult::Success, "FAnimationAnimNextRuntimeTest_Variables -> GetVariable failed");
 		UE_RETURN_ON_ERROR(TestSoftClassPath == FSoftClassPath(UAnimNextDataInterface::StaticClass()), "FAnimationAnimNextRuntimeTest_Variables -> Variable value did not match");
 	}
 
diff --git a/Engine/Plugins/Experimental/AnimNextAnimGraph/Source/AnimNextAnimGraphTestSuite/Private/AnimNextAnimGraphTraitInterfacesTest.cpp b/Engine/Plugins/Experimental/AnimNextAnimGraph/Source/AnimNextAnimGraphTestSuite/Private/AnimNextAnimGraphTraitInterfacesTest.cpp
--- a/Engine/Plugins/Experimental/AnimNextAnimGraph/Source/AnimNextAnimGraphTestSuite/Private/AnimNextAnimGraphTraitInterfacesTest.cpp
+++ b/Engine/Plugins/Experimental/AnimNextAnimGraph/Source/AnimNextAnimGraphTestSuite/Private/AnimNextAnimGraphTraitInterfacesTest.cpp
@@ -593,17 +593,16 @@ bool FAnimationAnimNextRuntimeTest_IHierarchy::RunTest(const FString& InParamete
 		// Read our graph
 		FTestUtils::LoadFromArchiveBuffer(*AnimationGraph, NodeHandles, GraphSharedDataArchiveBuffer);
 
-		FAnimNextGraphInstancePtr GraphInstance;
-		AnimationGraph->AllocateInstance(GraphInstance);
+		TSharedPtr<FAnimNextGraphInstance> GraphInstance = AnimationGraph->AllocateInstance();
 
-		FExecutionContext Context(GraphInstance);
+		FExecutionContext Context(*GraphInstance.Get());
 
 		{
 			FMemMark Mark(FMemStack::Get());
 
 			FAnimNextTraitHandle RootHandle(NodeHandles[3]);	// Point to NodeD, first base trait
 
-			FTraitPtr NodeDPtr = Context.AllocateNodeInstance(*GraphInstance.GetImpl(), RootHandle);
+			FTraitPtr NodeDPtr = Context.AllocateNodeInstance(*GraphInstance.Get(), RootHandle);
 			AddErrorIfFalse(NodeDPtr.IsValid(), "FAnimationAnimNextRuntimeTest_IHierarchy -> Failed to allocate root node instance");
 
 			FTraitStackBinding StackNodeD;
@@ -792,10 +791,9 @@ bool FAnimationAnimNextRuntimeTest_IUpdate::RunTest(const FString& InParameters)
 		// Read our graph
 		FTestUtils::LoadFromArchiveBuffer(*AnimationGraph, NodeHandles, GraphSharedDataArchiveBuffer);
 
-		FAnimNextGraphInstancePtr GraphInstance;
-		AnimationGraph->AllocateInstance(GraphInstance);
+		TSharedPtr<FAnimNextGraphInstance> GraphInstance = AnimationGraph->AllocateInstance();
 
-		FExecutionContext Context(GraphInstance);
+		FExecutionContext Context(*GraphInstance.Get());
 
 		{
 			TArray<FTraitUID> UpdatedTraits;
@@ -804,7 +802,7 @@ bool FAnimationAnimNextRuntimeTest_IUpdate::RunTest(const FString& InParameters)
 			// Call pre/post update on our graph
 			FTraitEventList InputEventList;
 			FTraitEventList OutputEventList;
-			UpdateGraph(GraphInstance, 0.0333f, InputEventList, OutputEventList);
+			UpdateGraph(*GraphInstance.Get(), 0.0333f, InputEventList, OutputEventList);
 
 			AddErrorIfFalse(UpdatedTraits.Num() == 6, "FAnimationAnimNextRuntimeTest_IUpdate -> Expected 6 nodes to have been visited during the update traversal");
 			AddErrorIfFalse(UpdatedTraits[0] == FTraitWithChildren::TraitUID, "FAnimationAnimNextRuntimeTest_IUpdate -> Unexpected update order");		// NodeC
@@ -935,15 +933,14 @@ bool FAnimationAnimNextRuntimeTest_IEvaluate::RunTest(const FString& InParameter
 		// Read our graph
 		FTestUtils::LoadFromArchiveBuffer(*AnimationGraph, NodeHandles, GraphSharedDataArchiveBuffer);
 
-		FAnimNextGraphInstancePtr GraphInstance;
-		AnimationGraph->AllocateInstance(GraphInstance);
+		TSharedPtr<FAnimNextGraphInstance> GraphInstance = AnimationGraph->AllocateInstance();
 
 		{
 			TArray<FTraitUID> EvaluatedTraits;
 			Private::EvaluatedTraits = &EvaluatedTraits;
 
 			// Call pre/post evaluate on our graph
-			(void)EvaluateGraph(GraphInstance);
+			(void)EvaluateGraph(*GraphInstance.Get());
 
 			AddErrorIfFalse(EvaluatedTraits.Num() == 6, "FAnimationAnimNextRuntimeTest_IEvaluate -> Expected 6 nodes to have been visited during the evaluate traversal");
 			AddErrorIfFalse(EvaluatedTraits[0] == FTraitWithChildren::TraitUID, "FAnimationAnimNextRuntimeTest_IEvaluate -> Unexpected evaluate order");		// NodeC
@@ -1072,10 +1069,9 @@ bool FAnimationAnimNextRuntimeTest_IScopedInterface::RunTest(const FString& InPa
 		// Read our graph
 		FTestUtils::LoadFromArchiveBuffer(*AnimationGraph, NodeHandles, GraphSharedDataArchiveBuffer);
 
-		FAnimNextGraphInstancePtr GraphInstance;
-		AnimationGraph->AllocateInstance(GraphInstance);
+		TSharedPtr<FAnimNextGraphInstance> GraphInstance = AnimationGraph->AllocateInstance();
 
-		FExecutionContext Context(GraphInstance);
+		FExecutionContext Context(*GraphInstance.Get());
 
 		{
 			TArray<bool> IsTagInScope;
@@ -1085,7 +1081,7 @@ bool FAnimationAnimNextRuntimeTest_IScopedInterface::RunTest(const FString& InPa
 			Private::AutoPopTag = true;
 			FTraitEventList InputEventList;
 			FTraitEventList OutputEventList;
-			UpdateGraph(GraphInstance, 0.0333f, InputEventList, OutputEventList);
+			UpdateGraph(*GraphInstance.Get(), 0.0333f, InputEventList, OutputEventList);
 
 			AddErrorIfFalse(IsTagInScope.Num() == 8, "FAnimationAnimNextRuntimeTest_IScopedInterface -> Unexpected number of entries");
 			AddErrorIfFalse(IsTagInScope[0] == false, "FAnimationAnimNextRuntimeTest_IScopedInterface -> Unexpected scoped tag state");		// NodeC::PreUpdate (template 0)
@@ -1101,7 +1097,7 @@ bool FAnimationAnimNextRuntimeTest_IScopedInterface::RunTest(const FString& InPa
 			Private::AutoPopTag = false;
 
 			// Call pre/post update on our graph
-			UpdateGraph(GraphInstance, 0.0333f, InputEventList, OutputEventList);
+			UpdateGraph(*GraphInstance.Get(), 0.0333f, InputEventList, OutputEventList);
 
 			AddErrorIfFalse(IsTagInScope.Num() == 8, "FAnimationAnimNextRuntimeTest_IScopedInterface -> Unexpected number of entries");
 			AddErrorIfFalse(IsTagInScope[0] == false, "FAnimationAnimNextRuntimeTest_IScopedInterface -> Unexpected scoped tag state");		// NodeC::PreUpdate (template 0)
