AnimNext: Write custom attributes to the right container + hold a delta time record for root motion extraction

diff --git a/Engine/Plugins/Experimental/AnimNext/Source/AnimNext/Private/EvaluationVM/Tasks/PushAnimSequenceKeyframe.cpp b/Engine/Plugins/Experimental/AnimNext/Source/AnimNext/Private/EvaluationVM/Tasks/PushAnimSequenceKeyframe.cpp
--- a/Engine/Plugins/Experimental/AnimNext/Source/AnimNext/Private/EvaluationVM/Tasks/PushAnimSequenceKeyframe.cpp
+++ b/Engine/Plugins/Experimental/AnimNext/Source/AnimNext/Private/EvaluationVM/Tasks/PushAnimSequenceKeyframe.cpp
@@ -36,10 +36,7 @@ void FAnimNextAnimSequenceKeyframeTask::Execute(UE::AnimNext::FEvaluationVM& VM)
 	{
 		const bool bIsAdditive = AnimSequencePtr->IsValidAdditive();
 
-		FDeltaTimeRecord DeltaTimeRecord;	// Not needed
 		const bool bExtractRootMotion = bExtractTrajectory;
-		const bool bLooping = false;		// Not needed
-
 		const FAnimExtractContext ExtractionContext(SampleTime, bExtractRootMotion, DeltaTimeRecord, bLooping);
 
 		FKeyframeState Keyframe = VM.MakeUninitializedKeyframe(bIsAdditive);
diff --git a/Engine/Plugins/Experimental/AnimNext/Source/AnimNext/Private/Graph/RigUnit_AnimNextWriteSkeletalMeshComponentPose.cpp b/Engine/Plugins/Experimental/AnimNext/Source/AnimNext/Private/Graph/RigUnit_AnimNextWriteSkeletalMeshComponentPose.cpp
--- a/Engine/Plugins/Experimental/AnimNext/Source/AnimNext/Private/Graph/RigUnit_AnimNextWriteSkeletalMeshComponentPose.cpp
+++ b/Engine/Plugins/Experimental/AnimNext/Source/AnimNext/Private/Graph/RigUnit_AnimNextWriteSkeletalMeshComponentPose.cpp
@@ -49,7 +49,7 @@ FRigUnit_AnimNextWriteSkeletalMeshComponentPose_Execute()
 	SkeletalMeshComponent->AnimCurves.CopyFrom(Pose.Curves);
 
 	// Attributes require remapping since the indices are LOD indices and we want mesh indices
-	FGenerationTools::RemapAttributes(Pose.LODPose, Pose.Attributes, SkeletalMeshComponent->CustomAttributes);
+	FGenerationTools::RemapAttributes(Pose.LODPose, Pose.Attributes, SkeletalMeshComponent->GetEditableCustomAttributes());
 
 	// Convert and dispatch to renderer
 	UE::Anim::FSkinnedMeshComponentExtensions::CompleteAndDispatch(
diff --git a/Engine/Plugins/Experimental/AnimNext/Source/AnimNext/Public/EvaluationVM/Tasks/PushAnimSequenceKeyframe.h b/Engine/Plugins/Experimental/AnimNext/Source/AnimNext/Public/EvaluationVM/Tasks/PushAnimSequenceKeyframe.h
--- a/Engine/Plugins/Experimental/AnimNext/Source/AnimNext/Public/EvaluationVM/Tasks/PushAnimSequenceKeyframe.h
+++ b/Engine/Plugins/Experimental/AnimNext/Source/AnimNext/Public/EvaluationVM/Tasks/PushAnimSequenceKeyframe.h
@@ -4,6 +4,7 @@
 
 #include "CoreMinimal.h"
 #include "EvaluationVM/EvaluationTask.h"
+#include "Animation/AnimationAsset.h"
 
 #include "PushAnimSequenceKeyframe.generated.h"
 
@@ -31,6 +32,9 @@ struct ANIMNEXT_API FAnimNextAnimSequenceKeyframeTask : public FAnimNextEvaluati
 	UPROPERTY()
 	TWeakObjectPtr<UAnimSequence> AnimSequence;
 
+	/** Delta time range required for root motion extraction **/
+	FDeltaTimeRecord DeltaTimeRecord;
+
 	// The point in time within the animation sequence at which we sample the keyframe.
 	// If negative, the sample time hasn't been provided and we should use the keyframe index.
 	UPROPERTY()
@@ -49,4 +53,7 @@ struct ANIMNEXT_API FAnimNextAnimSequenceKeyframeTask : public FAnimNextEvaluati
 	// Whether to extract trajectory or not
 	UPROPERTY()
 	bool bExtractTrajectory = false;
+
+	UPROPERTY()
+	bool bLooping = false;
 };
diff --git a/Engine/Plugins/Experimental/AnimNextAnimGraph/Source/AnimNextAnimGraph/Private/Traits/SequencePlayer.cpp b/Engine/Plugins/Experimental/AnimNextAnimGraph/Source/AnimNextAnimGraph/Private/Traits/SequencePlayer.cpp
--- a/Engine/Plugins/Experimental/AnimNextAnimGraph/Source/AnimNextAnimGraph/Private/Traits/SequencePlayer.cpp
+++ b/Engine/Plugins/Experimental/AnimNextAnimGraph/Source/AnimNextAnimGraph/Private/Traits/SequencePlayer.cpp
@@ -41,7 +41,12 @@ namespace UE::AnimNext
 		const FInstanceData* InstanceData = Binding.GetInstanceData<FInstanceData>();
 		const bool bInterpolate = true;
 
-		FAnimNextAnimSequenceKeyframeTask Task = FAnimNextAnimSequenceKeyframeTask::MakeFromSampleTime(InstanceData->AnimSequence, InstanceData->InternalTimeAccumulator, bInterpolate);
+		FAnimNextAnimSequenceKeyframeTask Task;
+		Task.AnimSequence = InstanceData->AnimSequence;
+		Task.SampleTime = InstanceData->InternalTimeAccumulator;
+		Task.bInterpolate = bInterpolate;
+		Task.DeltaTimeRecord = InstanceData->DeltaTimeRecord;
+		Task.bLooping = SharedData->GetbLoop(Binding);
 		Task.bExtractTrajectory = true;	/*Output.AnimInstanceProxy->ShouldExtractRootMotion()*/
 
 		Context.AppendTask(Task);
@@ -103,7 +108,10 @@ namespace UE::AnimNext
 			const bool bIsLooping = SharedData->GetbLoop(Binding);
 			const float SequenceLength = AnimSeq->GetPlayLength();
 
-			FAnimationRuntime::AdvanceTime(bIsLooping, DeltaTime * PlayRate, InstanceData->InternalTimeAccumulator, SequenceLength);
+			const float PreviousTime = InstanceData->InternalTimeAccumulator;
+			const float MoveDelta = DeltaTime * PlayRate;
+			InstanceData->DeltaTimeRecord.Set(PreviousTime, MoveDelta);
+			FAnimationRuntime::AdvanceTime(bIsLooping, MoveDelta, InstanceData->InternalTimeAccumulator, SequenceLength);
 
 			return FTimelineProgress(InstanceData->InternalTimeAccumulator, SequenceLength);
 		}
@@ -119,7 +127,10 @@ namespace UE::AnimNext
 		{
 			const float SequenceLength = AnimSeq->GetPlayLength();
 
+			const float PreviousTime = InstanceData->InternalTimeAccumulator;
 			InstanceData->InternalTimeAccumulator = FMath::Clamp(ProgressRatio, 0.0f, 1.0f) * SequenceLength;
+			const float MoveDelta = InstanceData->InternalTimeAccumulator - PreviousTime;
+			InstanceData->DeltaTimeRecord.Set(PreviousTime, MoveDelta);
 		}
 	}
 
diff --git a/Engine/Plugins/Experimental/AnimNextAnimGraph/Source/AnimNextAnimGraph/Private/Traits/SequencePlayer.h b/Engine/Plugins/Experimental/AnimNextAnimGraph/Source/AnimNextAnimGraph/Private/Traits/SequencePlayer.h
--- a/Engine/Plugins/Experimental/AnimNextAnimGraph/Source/AnimNextAnimGraph/Private/Traits/SequencePlayer.h
+++ b/Engine/Plugins/Experimental/AnimNextAnimGraph/Source/AnimNextAnimGraph/Private/Traits/SequencePlayer.h
@@ -62,6 +62,8 @@ namespace UE::AnimNext
 			// Cached value of the anim sequence we are playing
 			TObjectPtr<UAnimSequence> AnimSequence;
 
+			/** Delta time range required for root motion extraction **/
+			FDeltaTimeRecord DeltaTimeRecord;
 			float InternalTimeAccumulator = 0.0f;
 
 			void Construct(const FExecutionContext& Context, const FTraitBinding& Binding);
