Simplification following addition of Quat::GetShortestArcWith No functional changes

diff --git a/Engine/Plugins/Animation/MotionWarping/Source/MotionWarping/Private/AttributeBasedRootMotionComponent.cpp b/Engine/Plugins/Animation/MotionWarping/Source/MotionWarping/Private/AttributeBasedRootMotionComponent.cpp
--- a/Engine/Plugins/Animation/MotionWarping/Source/MotionWarping/Private/AttributeBasedRootMotionComponent.cpp
+++ b/Engine/Plugins/Animation/MotionWarping/Source/MotionWarping/Private/AttributeBasedRootMotionComponent.cpp
@@ -81,8 +81,7 @@ void UAttributeBasedRootMotionComponent::TickComponent(float DeltaTime, enum ELe
 				{
 					TranslationVelocity = RootMotionAttribute->Value.GetTranslation() / DeltaTime;
 
-					FQuat RootMotionRotation = RootMotionAttribute->Value.GetRotation();
-					RootMotionRotation.EnforceShortestArcWith(FQuat::Identity);
+					const FQuat RootMotionRotation = RootMotionAttribute->Value.GetRotation().GetShortestArcWith(FQuat::Identity);
 					RotationVelocity = RootMotionRotation.ToRotationVector() / DeltaTime;
 
 					if (Mode == EAttributeBasedRootMotionMode::ApplyDelta)
diff --git a/Engine/Plugins/Animation/PoseSearch/Source/Runtime/Private/PoseSearchTrajectoryLibrary.cpp b/Engine/Plugins/Animation/PoseSearch/Source/Runtime/Private/PoseSearchTrajectoryLibrary.cpp
--- a/Engine/Plugins/Animation/PoseSearch/Source/Runtime/Private/PoseSearchTrajectoryLibrary.cpp
+++ b/Engine/Plugins/Animation/PoseSearch/Source/Runtime/Private/PoseSearchTrajectoryLibrary.cpp
@@ -620,9 +620,7 @@ void UPoseSearchTrajectoryLibrary::GetTrajectoryAngularVelocity(const FPoseSearc
 	FPoseSearchQueryTrajectorySample Sample1 = InTrajectory.GetSampleAtTime(Time1, bExtrapolate);
 	FPoseSearchQueryTrajectorySample Sample2 = InTrajectory.GetSampleAtTime(Time2, bExtrapolate);
 
-	FQuat DeltaRotation = (Sample2.Facing * Sample1.Facing.Inverse());
-	DeltaRotation.EnforceShortestArcWith(FQuat::Identity);
-
+	const FQuat DeltaRotation = (Sample2.Facing * Sample1.Facing.Inverse()).GetShortestArcWith(FQuat::Identity);
 	const FVector AngularVelocityInRadians = DeltaRotation.ToRotationVector() / (Time2 - Time1);
 
 	OutAngularVelocity = FVector(
diff --git a/Engine/Plugins/Experimental/AnimNextAnimGraph/Source/AnimNextAnimGraph/Private/EvaluationVM/Tasks/DeadBlending.cpp b/Engine/Plugins/Experimental/AnimNextAnimGraph/Source/AnimNextAnimGraph/Private/EvaluationVM/Tasks/DeadBlending.cpp
--- a/Engine/Plugins/Experimental/AnimNextAnimGraph/Source/AnimNextAnimGraph/Private/EvaluationVM/Tasks/DeadBlending.cpp
+++ b/Engine/Plugins/Experimental/AnimNextAnimGraph/Source/AnimNextAnimGraph/Private/EvaluationVM/Tasks/DeadBlending.cpp
@@ -257,9 +257,7 @@ namespace UE::AnimNext::DeadBlending::Private
 
 			const FVector TranslationDiff = SrcTranslationCurr - SrcTranslationPrev;
 
-			FQuat RotationDiff = SrcRotationCurr * SrcRotationPrev.Inverse();
-			RotationDiff.EnforceShortestArcWith(FQuat::Identity);
-
+			const FQuat RotationDiff = (SrcRotationCurr * SrcRotationPrev.Inverse()).GetShortestArcWith(FQuat::Identity);
 			const FVector ScaleDiff = Private::VectorDivMax(SrcScaleCurr, SrcScalePrev);
 
 			SourceBoneTranslationVelocities[LODBoneIdx] = (FVector3f)(TranslationDiff / DeltaTime).GetClampedToMaxSize(Parameters.MaximumTranslationVelocity);
@@ -274,9 +272,7 @@ namespace UE::AnimNext::DeadBlending::Private
 
 			const FVector TranslationSrcDstDiff = DstTranslation - SrcTranslationCurr;
 
-			FQuat RotationSrcDstDiff = DstRotation * SrcRotationCurr.Inverse();
-			RotationSrcDstDiff.EnforceShortestArcWith(FQuat::Identity);
-
+			const FQuat RotationSrcDstDiff = DstRotation * SrcRotationCurr.Inverse().GetShortestArcWith(FQuat::Identity);
 			const FVector ScaleSrcDstDiff = Private::VectorDivMax(DstScale, SrcScaleCurr);
 
 			SourceBoneTranslationDecayHalfLives[LODBoneIdx] = Private::ComputeDecayHalfLifeFromDiffAndVelocity(
diff --git a/Engine/Plugins/Experimental/LearningAgents/Source/LearningAgents/Private/LearningAgentsActions.cpp b/Engine/Plugins/Experimental/LearningAgents/Source/LearningAgents/Private/LearningAgentsActions.cpp
--- a/Engine/Plugins/Experimental/LearningAgents/Source/LearningAgents/Private/LearningAgentsActions.cpp
+++ b/Engine/Plugins/Experimental/LearningAgents/Source/LearningAgents/Private/LearningAgentsActions.cpp
@@ -1567,8 +1567,7 @@ FLearningAgentsActionObjectElement ULearningAgentsActions::MakeRotationAction(UL
 
 FLearningAgentsActionObjectElement ULearningAgentsActions::MakeRotationActionFromQuat(ULearningAgentsActionObject* Object, const FQuat Rotation, const FQuat RelativeRotation, const FName Tag)
 {
-	FQuat LocalRotation = RelativeRotation.Inverse() * Rotation;
-	LocalRotation.EnforceShortestArcWith(FQuat::Identity);
+	FQuat LocalRotation = (RelativeRotation.Inverse() * Rotation).GetShortestArcWith(FQuat::Identity);
 	const FVector RotationVector = LocalRotation.ToRotationVector();
 
 	return MakeContinuousActionFromArrayView(Object, {
diff --git a/Engine/Plugins/Experimental/PhysicsControl/Source/PhysicsControl/Private/PhysicsControlComponentImpl.cpp b/Engine/Plugins/Experimental/PhysicsControl/Source/PhysicsControl/Private/PhysicsControlComponentImpl.cpp
--- a/Engine/Plugins/Experimental/PhysicsControl/Source/PhysicsControl/Private/PhysicsControlComponentImpl.cpp
+++ b/Engine/Plugins/Experimental/PhysicsControl/Source/PhysicsControl/Private/PhysicsControlComponentImpl.cpp
@@ -467,9 +467,8 @@ void UPhysicsControlComponent::CalculateControlTargetData(
 		if (SkeletalDeltaTime * Record.PhysicsControl.ControlData.AngularTargetVelocityMultiplier != 0)
 		{
 			const FQuat Q = OutSkeletalTargetTM.GetRotation();
-			FQuat PrevQ = Record.PreviousSkeletalTargetTM.GetRotation();
-			PrevQ.EnforceShortestArcWith(Q);
-			const FQuat DeltaQ = Q * PrevQ.Inverse();
+			const FQuat PrevQ = Record.PreviousSkeletalTargetTM.GetRotation();
+			const FQuat DeltaQ = (Q * PrevQ.Inverse()).GetShortestArcWith(FQuat::Identity);
 			OutTargetAngularVelocity = DeltaQ.ToRotationVector() * (
 				Record.PhysicsControl.ControlData.AngularTargetVelocityMultiplier / SkeletalDeltaTime);
 		}
diff --git a/Engine/Plugins/Experimental/PhysicsControl/Source/PhysicsControl/Private/PhysicsControlPoseData.cpp b/Engine/Plugins/Experimental/PhysicsControl/Source/PhysicsControl/Private/PhysicsControlPoseData.cpp
--- a/Engine/Plugins/Experimental/PhysicsControl/Source/PhysicsControl/Private/PhysicsControlPoseData.cpp
+++ b/Engine/Plugins/Experimental/PhysicsControl/Source/PhysicsControl/Private/PhysicsControlPoseData.cpp
@@ -96,8 +96,8 @@ FVector FBoneData::CalculateAngularVelocity(float Dt) const
 	}
 
 	// Note that quats multiply in the opposite order to TMs, and must be in the same hemisphere.
-	const FQuat DeltaQ = CurrentTM.GetRotation() * PreviousTM.GetRotation().Inverse();
-	return (DeltaQ.GetShortestArcWith(FQuat::Identity).ToRotationVector()) / Dt;
+	const FQuat DeltaQ = (CurrentTM.GetRotation() * PreviousTM.GetRotation().Inverse()).GetShortestArcWith(FQuat::Identity);
+	return DeltaQ.ToRotationVector() / Dt;
 }
 
 } // namespace PhysicsControl
diff --git a/Engine/Source/Runtime/Engine/Private/Animation/AnimNode_DeadBlending.cpp b/Engine/Source/Runtime/Engine/Private/Animation/AnimNode_DeadBlending.cpp
--- a/Engine/Source/Runtime/Engine/Private/Animation/AnimNode_DeadBlending.cpp
+++ b/Engine/Source/Runtime/Engine/Private/Animation/AnimNode_DeadBlending.cpp
@@ -369,9 +369,7 @@ void FAnimNode_DeadBlending::InitFrom(
 			const FVector SrcScalePrev = SrcPosePrev.BoneScales[PrevPoseBoneIndex];
 
 			const FVector TranslationDiff = SrcTranslationCurr - SrcTranslationPrev;
-
-			FQuat RotationDiff = SrcRotationCurr * SrcRotationPrev.Inverse();
-			RotationDiff.EnforceShortestArcWith(FQuat::Identity);
+			const FQuat RotationDiff = (SrcRotationCurr * SrcRotationPrev.Inverse()).GetShortestArcWith(FQuat::Identity);
 
 			const FVector ScaleDiffLinear = SrcScaleCurr - SrcScalePrev;
 			const FVector ScaleDiffExponential = UE::Anim::DeadBlending::Private::VectorDivMax(SrcScaleCurr, SrcScalePrev);
