AnimNext: Factories

Added to new trait interfaces: IGraphFactory and IPayloadFactory. Both use fallbacks to project-defined (with engine defaults) mappings between asset classes (e.g. UAnimSequence) and anim next graphs. The mapping is managed by UAnimNextAnimGraphSettings.

The idea is that anything that wants to implement the hosting of some subgraph can use the IGraphFactory to not be limited to only handling UAnimNextAnimationGraph assets, but can use it to make UAnimNextAnimationGraphs to instantiate from some other object key. The object key can be optionally injected as a variable into the facotry graph via the user mapping.

IGraphFactory has two implementations, FMakeGraphFromAssetTrait (which just implements the fallback to UAnimNextAnimGraphSettings) and FMakeGraphFromChooserTrait (which allows a chooser to decide the graph to instantiate).
IPayloadFactory does not have any implementations yet (so fallback is always used), but the intent is that we can implement this in future to supply per-variable overrides via traits that expose values on pins.
Updated the following traits to use the new interfaces: FBlendStackTrait, FBlendStackRequesterTrait, FMotionMatchingTrait, FPlayAnimSlotTrait.

Added the concept of a 'native interface' to all AnimNext RigVM assets. This is a structure that is intended to be used for native C++ communication with AnimNext asset instances via struct-adapters. This is compiled into all 'implemented interface' definitions.

Updated StateTree to use the new interfaces. Updated state tree's task to allow any object that is supported by the IGraphFactory interface. Also updated the UI to allow selecting and parameterizing the instanced graph via non-native interfaces.

Added one plugin-level asset to implement an anim sequence player: /AnimNextAnimGraph/FactoryGraphs/AG_AnimSequencePlayer.

Removed AnimNextAnimGraph and AnimNextStateTree's dependency on Chooser module.

#virtualized

diff --git a/Engine/Plugins/Experimental/AnimNext/Source/AnimNext/Internal/DataInterface/AnimNextDataInterface.h b/Engine/Plugins/Experimental/AnimNext/Source/AnimNext/Internal/DataInterface/AnimNextDataInterface.h
--- a/Engine/Plugins/Experimental/AnimNext/Source/AnimNext/Internal/DataInterface/AnimNextDataInterface.h
+++ b/Engine/Plugins/Experimental/AnimNext/Source/AnimNext/Internal/DataInterface/AnimNextDataInterface.h
@@ -7,6 +7,7 @@
 #include "AnimNextDataInterface.generated.h"
 
 class UAnimNextDataInterfaceFactory;
+class UAnimNextDataInterface;
 
 namespace UE::AnimNext::Tests
 {
@@ -24,9 +25,14 @@ struct FAnimNextImplementedDataInterface
 {
 	GENERATED_BODY()
 
+	// Reference to the implemented data interface asset
 	UPROPERTY()
 	TObjectPtr<const UAnimNextDataInterface> DataInterface;
 
+	// Native struct used for communication with the data interface
+	UPROPERTY()
+	TObjectPtr<const UScriptStruct> NativeInterface;
+
 	// Index of the first variable that implements the interface
 	UPROPERTY()
 	int32 VariableIndex = INDEX_NONE;
diff --git a/Engine/Plugins/Experimental/AnimNext/Source/AnimNext/Private/DataInterface/DataInterfaceStructAdapter.cpp b/Engine/Plugins/Experimental/AnimNext/Source/AnimNext/Private/DataInterface/DataInterfaceStructAdapter.cpp
--- a/Engine/Plugins/Experimental/AnimNext/Source/AnimNext/Private/DataInterface/DataInterfaceStructAdapter.cpp
+++ b/Engine/Plugins/Experimental/AnimNext/Source/AnimNext/Private/DataInterface/DataInterfaceStructAdapter.cpp
@@ -14,7 +14,13 @@ const UAnimNextDataInterface* FDataInterfaceStructAdapter::GetDataInterface() co
 
 uint8* FDataInterfaceStructAdapter::GetMemoryForVariable(int32 InVariableIndex, FName InVariableName, const FProperty* InVariableProperty) const
 {
-	const FProperty* Property = StructView.GetScriptStruct()->FindPropertyByName(InVariableName);
+	const UScriptStruct* Struct = StructView.GetScriptStruct();
+	if(Struct == nullptr)
+	{
+		return nullptr;
+	}
+
+	const FProperty* Property = Struct->FindPropertyByName(InVariableName);
 	if(Property == nullptr)
 	{
 		return nullptr;
diff --git a/Engine/Plugins/Experimental/AnimNext/Source/AnimNext/Private/Graph/AnimNextGraphInstance.cpp b/Engine/Plugins/Experimental/AnimNext/Source/AnimNext/Private/Graph/AnimNextGraphInstance.cpp
--- a/Engine/Plugins/Experimental/AnimNext/Source/AnimNext/Private/Graph/AnimNextGraphInstance.cpp
+++ b/Engine/Plugins/Experimental/AnimNext/Source/AnimNext/Private/Graph/AnimNextGraphInstance.cpp
@@ -4,6 +4,7 @@
 
 #include "AnimNextStats.h"
 #include "DataInterface/AnimNextDataInterfaceHost.h"
+#include "DataInterface/DataInterfaceStructAdapter.h"
 #include "Graph/AnimNextAnimationGraph.h"
 #include "TraitCore/ExecutionContext.h"
 #include "Graph/GC_GraphInstanceComponent.h"
@@ -285,6 +286,43 @@ bool FAnimNextGraphInstance::BindToHostHelper(const HostType& InHost, bool bInAu
 	return bPublicVariablesBound;
 }
 
+void FAnimNextGraphInstance::BindPublicVariables(TArrayView<FStructView> InHostStructs)
+{
+	using namespace UE::AnimNext;
+
+	const UAnimNextAnimationGraph* AnimationGraph = GetAnimationGraph();
+	if (AnimationGraph == nullptr)
+	{
+		return;
+	}
+
+	const int32 NumStructs = InHostStructs.Num();
+	if(NumStructs > 0)
+	{
+		// stack allocating uninitialized arrays to contain FDataInterfaceStructAdapter(s)
+		TArrayView<FDataInterfaceStructAdapter> Adapters((FDataInterfaceStructAdapter*)FMemory_Alloca(NumStructs * sizeof(FDataInterfaceStructAdapter)), NumStructs);
+		TArrayView<IDataInterfaceHost*> PayloadVariableBindings((IDataInterfaceHost**)FMemory_Alloca(NumStructs * sizeof(IDataInterfaceHost*)), NumStructs);
+
+		// initializing the FDataInterfaceStructAdapter(s) with their respective HostStruct
+		for (int32 StructIndex = 0; StructIndex < NumStructs; ++StructIndex)
+		{
+			FStructView& StructView = InHostStructs[StructIndex];
+			PayloadVariableBindings[StructIndex] = new (&Adapters[StructIndex]) FDataInterfaceStructAdapter(AnimationGraph, StructView);
+		}
+
+		// Inner binding 
+		BindPublicVariables(PayloadVariableBindings);
+
+		// uninitializing the FDataInterfaceStructAdapter(s)
+		for (int32 StructIndex = 0; StructIndex < NumStructs; ++StructIndex)
+		{
+			Adapters[StructIndex].~FDataInterfaceStructAdapter();
+		}
+
+		// abandoning no longer initialized stack allocated arrays to their demise
+	}
+}
+
 void FAnimNextGraphInstance::BindPublicVariables(TConstArrayView<UE::AnimNext::IDataInterfaceHost*> InHosts)
 {
 	using namespace UE::AnimNext;
diff --git a/Engine/Plugins/Experimental/AnimNext/Source/AnimNext/Public/AnimNextDataInterfacePayload.h b/Engine/Plugins/Experimental/AnimNext/Source/AnimNext/Public/AnimNextDataInterfacePayload.h
new file mode 100644
--- /dev/null
+++ b/Engine/Plugins/Experimental/AnimNext/Source/AnimNext/Public/AnimNextDataInterfacePayload.h
@@ -0,0 +1,194 @@
+﻿// Copyright Epic Games, Inc. All Rights Reserved.
+
+#pragma once
+
+#include "StructUtils/PropertyBag.h"
+#include "AnimNextDataInterfacePayload.generated.h"
+
+class UAnimNextAnimGraphSettings;
+
+// Combined native & non-native payloads, for supplying to a data interface instance
+USTRUCT()
+struct FAnimNextDataInterfacePayload
+{
+	GENERATED_BODY()
+
+	FAnimNextDataInterfacePayload() = default;
+
+	FAnimNextDataInterfacePayload(const FAnimNextDataInterfacePayload& InOther)
+	{
+		OwnedPayload = InOther.OwnedPayload;
+		OwnedNativePayloads = InOther.OwnedNativePayloads;
+		NonOwnedPayloads = InOther.NonOwnedPayloads;
+		bCombinedPayloadsDirty = true;
+	}
+
+	FAnimNextDataInterfacePayload& operator=(const FAnimNextDataInterfacePayload& InOther)
+	{
+		OwnedPayload = InOther.OwnedPayload;
+		OwnedNativePayloads = InOther.OwnedNativePayloads;
+		NonOwnedPayloads = InOther.NonOwnedPayloads;
+		bCombinedPayloadsDirty = true;
+		return *this;
+	}
+
+	FAnimNextDataInterfacePayload(FAnimNextDataInterfacePayload&& InOther)
+	{
+		OwnedPayload = MoveTemp(InOther.OwnedPayload);
+		OwnedNativePayloads = MoveTemp(InOther.OwnedNativePayloads);
+		NonOwnedPayloads = MoveTemp(InOther.NonOwnedPayloads);
+		bCombinedPayloadsDirty = true;
+
+		InOther.CombinedPayloads.Empty();
+		InOther.bCombinedPayloadsDirty = false;
+	}
+
+	FAnimNextDataInterfacePayload& operator=(FAnimNextDataInterfacePayload&& InOther)
+	{
+		OwnedPayload = MoveTemp(InOther.OwnedPayload);
+		OwnedNativePayloads = MoveTemp(InOther.OwnedNativePayloads);
+		NonOwnedPayloads = MoveTemp(InOther.NonOwnedPayloads);
+		bCombinedPayloadsDirty = true;
+
+		InOther.CombinedPayloads.Empty();
+		InOther.bCombinedPayloadsDirty = false;
+		return *this;
+	}
+
+	// Get all payloads
+	TArrayView<FStructView> Get()
+	{
+		UpdateCombinedPayloads();
+		return CombinedPayloads;
+	}
+
+	// Set payload
+	void Set(FInstancedPropertyBag&& InPayload)
+	{
+		OwnedPayload = MoveTemp(InPayload);
+		bCombinedPayloadsDirty = true;
+	}
+
+	// Set payload
+	void Set(const FInstancedPropertyBag& InPayload)
+	{
+		OwnedPayload = InPayload;
+		bCombinedPayloadsDirty = true;
+	}
+
+	// Set native payloads
+	void SetNative(TArray<FInstancedStruct>&& InNativePayloads)
+	{
+		OwnedNativePayloads = MoveTemp(InNativePayloads);
+		bCombinedPayloadsDirty = true;
+	}
+
+	// Append native payloads
+	void AppendNative(const TArray<FInstancedStruct>& InNativePayloads)
+	{
+		OwnedNativePayloads.Append(InNativePayloads);
+		bCombinedPayloadsDirty = true;
+	}
+
+	// Add a native payload
+	void AddNative(FInstancedStruct&& InNativePayload)
+	{
+		OwnedNativePayloads.Add(MoveTemp(InNativePayload));
+		bCombinedPayloadsDirty = true;
+	}
+
+	// Add a native payload
+	void AddNative(const FInstancedStruct& InNativePayload)
+	{
+		OwnedNativePayloads.Add(InNativePayload);
+		bCombinedPayloadsDirty = true;
+	}
+
+	// Add non-owned payload
+	void AddStructView(const FStructView& InPayload)
+	{
+		NonOwnedPayloads.Add(InPayload);
+		bCombinedPayloadsDirty = true;
+	}
+
+	// Append non-owned payloads
+	void AppendStructView(const TArray<FStructView>& InPayloads)
+	{
+		NonOwnedPayloads.Append(InPayloads);
+		bCombinedPayloadsDirty = true;
+	}
+
+	const FInstancedPropertyBag& GetPayload() const
+	{
+		return OwnedPayload;
+	}
+
+	TConstArrayView<FInstancedStruct> GetNativePayloads() const
+	{
+		return OwnedNativePayloads;
+	}
+
+	TConstArrayView<FStructView> GetNonOwnedPayloads() const
+	{
+		return NonOwnedPayloads;
+	}
+	
+#if WITH_EDITOR
+	// Details customization support
+	static FName GetOwnedPayloadPropertyName()
+	{
+		return GET_MEMBER_NAME_CHECKED(FAnimNextDataInterfacePayload, OwnedPayload);
+	}
+
+	// Details customization support
+	static FName GetOwnedNativePayloadsPropertyName()
+	{
+		return GET_MEMBER_NAME_CHECKED(FAnimNextDataInterfacePayload, OwnedNativePayloads);
+	}
+#endif
+
+private:
+	// Lazily refresh combined payloads
+	void UpdateCombinedPayloads()
+	{
+		if(bCombinedPayloadsDirty)
+		{
+			CombinedPayloads.Reset();
+			CombinedPayloads.Reserve(OwnedNativePayloads.Num() + NonOwnedPayloads.Num() + 1);
+			if(OwnedPayload.IsValid())
+			{
+				CombinedPayloads.Add(OwnedPayload.GetMutableValue());
+			}
+			for(FInstancedStruct& InstancedStruct : OwnedNativePayloads)
+			{
+				CombinedPayloads.Add(InstancedStruct);
+			}
+			for(FStructView& StructView : NonOwnedPayloads)
+			{
+				CombinedPayloads.Add(StructView);
+			}
+
+			bCombinedPayloadsDirty = false;
+		}
+	}
+
+private:
+	friend UAnimNextAnimGraphSettings;
+
+	// Single payload used for data-driven variable bindings
+	UPROPERTY(EditAnywhere, Category = Payload)
+	FInstancedPropertyBag OwnedPayload;
+
+	// Multiple payloads used for native bindings
+	UPROPERTY(EditAnywhere, Category = Payload)
+	TArray<FInstancedStruct> OwnedNativePayloads;
+
+	// Externally-owned payloads used to avoid copying constant data
+	TArray<FStructView> NonOwnedPayloads;
+
+	// Combined view of all payloads used for binding
+	TArray<FStructView> CombinedPayloads;
+
+	// dirty flag for lazy rebuilds of CombinedPayloads
+	bool bCombinedPayloadsDirty = false;
+};
\ No newline at end of file
diff --git a/Engine/Plugins/Experimental/AnimNext/Source/AnimNext/Public/DataInterface/AnimNextNativeDataInterface.h b/Engine/Plugins/Experimental/AnimNext/Source/AnimNext/Public/DataInterface/AnimNextNativeDataInterface.h
new file mode 100644
--- /dev/null
+++ b/Engine/Plugins/Experimental/AnimNext/Source/AnimNext/Public/DataInterface/AnimNextNativeDataInterface.h
@@ -0,0 +1,33 @@
+﻿// Copyright Epic Games, Inc. All Rights Reserved.
+
+#pragma once
+
+#include "CoreMinimal.h"
+#include "AnimNextNativeDataInterface.generated.h"
+
+class UAnimNextDataInterface;
+
+// Base structure used to implement a native interface into AnimNext assets
+USTRUCT()
+struct ANIMNEXT_API FAnimNextNativeDataInterface
+{
+	GENERATED_BODY()
+
+	FAnimNextNativeDataInterface() = default;
+	virtual ~FAnimNextNativeDataInterface() = default;
+
+	// Context passed to Initialize()
+	struct FBindToFactoryObjectContext
+	{
+		// If this native interface was created as a result of a factory, then this is the object that the factor used as a source. Can be nullptr.
+		const UObject* FactoryObject = nullptr;
+
+		// The data interface that this native interface was constructed for
+		const UAnimNextDataInterface* DataInterface = nullptr;
+	};
+
+	// Initialize this native interface during factory object creation.
+	// Called once on worker threads to set up this native interface's properties after it has been default-constructed.
+	// The data interface instance that it will manage has not yet been constructed at this point.
+	virtual void BindToFactoryObject(const FBindToFactoryObjectContext& InContext) {}
+};
\ No newline at end of file
diff --git a/Engine/Plugins/Experimental/AnimNext/Source/AnimNext/Public/Graph/AnimNextGraphInstance.h b/Engine/Plugins/Experimental/AnimNext/Source/AnimNext/Public/Graph/AnimNextGraphInstance.h
--- a/Engine/Plugins/Experimental/AnimNext/Source/AnimNext/Public/Graph/AnimNextGraphInstance.h
+++ b/Engine/Plugins/Experimental/AnimNext/Source/AnimNext/Public/Graph/AnimNextGraphInstance.h
@@ -141,6 +141,10 @@ private:
 	template<typename HostType>
 	bool BindToHostHelper(const HostType& InHost, bool bInAutoBind);
 
+	// Bind the variables in the supplied traits in scope to their respective public variables, represented by structs held on a host
+	// Note: struct views are mutable here as bindings allow writes!
+	void BindPublicVariables(TArrayView<FStructView> InHostStructs);
+
 	// Bind the variables in the supplied traits in scope to their respective public variables, so they point at host memory
 	void BindPublicVariables(TConstArrayView<UE::AnimNext::IDataInterfaceHost*> InHosts);
 
diff --git a/Engine/Plugins/Experimental/AnimNext/Source/AnimNextUncookedOnly/Internal/AnimNextRigVMAssetEditorData.h b/Engine/Plugins/Experimental/AnimNext/Source/AnimNextUncookedOnly/Internal/AnimNextRigVMAssetEditorData.h
--- a/Engine/Plugins/Experimental/AnimNext/Source/AnimNextUncookedOnly/Internal/AnimNextRigVMAssetEditorData.h
+++ b/Engine/Plugins/Experimental/AnimNext/Source/AnimNextUncookedOnly/Internal/AnimNextRigVMAssetEditorData.h
@@ -420,6 +420,10 @@ protected:
 	UPROPERTY()
 	FRigVMGraphFunctionStore GraphFunctionStore;
 
+	// The native C++ struct that is used to communicate with this asset
+	UPROPERTY(EditAnywhere, Category = "Native", meta = (ShowDisplayNames, MetaStruct = "/Script/AnimNext.AnimNextNativeDataInterface"))
+	TObjectPtr<const UScriptStruct> NativeInterface;
+
 	UPROPERTY(EditAnywhere, Category = "User Interface")
 	FRigVMEdGraphDisplaySettings RigGraphDisplaySettings;
 
diff --git a/Engine/Plugins/Experimental/AnimNext/Source/AnimNextUncookedOnly/Private/UncookedOnlyUtils.cpp b/Engine/Plugins/Experimental/AnimNext/Source/AnimNextUncookedOnly/Private/UncookedOnlyUtils.cpp
--- a/Engine/Plugins/Experimental/AnimNext/Source/AnimNextUncookedOnly/Private/UncookedOnlyUtils.cpp
+++ b/Engine/Plugins/Experimental/AnimNext/Source/AnimNextUncookedOnly/Private/UncookedOnlyUtils.cpp
@@ -70,11 +70,12 @@ void FUtils::CompileVariables(UAnimNextRigVMAsset* InAsset)
 
 	struct FStructEntryInfo
 	{
-		const UAnimNextDataInterface* FromDataInterface;
+		const UAnimNextDataInterface* FromDataInterface = nullptr;
+		const UScriptStruct* NativeInterface = nullptr;
 		FName Name;
 		FAnimNextParamType Type;
-		EAnimNextExportAccessSpecifier AccessSpecifier;
-		bool bAutoBindDataInterfaceToHost;
+		EAnimNextExportAccessSpecifier AccessSpecifier = EAnimNextExportAccessSpecifier::Private;
+		bool bAutoBindDataInterfaceToHost = false;
 		TConstArrayView<uint8> Value;
 	};
 
@@ -148,9 +149,13 @@ void FUtils::CompileVariables(UAnimNextRigVMAsset* InAsset)
 		}
 		else
 		{
+			const UAnimNextDataInterface_EditorData* FromInterfaceEditorData = GetEditorData<const UAnimNextDataInterface_EditorData>(InFromInterface);
+			check(FromInterfaceEditorData != nullptr);
+
 			int32 Index = StructEntryInfos.Add(
 			{
 				InFromInterface,
+				FromInterfaceEditorData->NativeInterface,
 				Name,
 				FAnimNextParamType(Type.GetValueType(), Type.GetContainerType(), Type.GetValueTypeObject()),
 				AccessSpecifier,
@@ -253,6 +258,7 @@ void FUtils::CompileVariables(UAnimNextRigVMAsset* InAsset)
 			{
 				FAnimNextImplementedDataInterface& NewImplementedDataInterface = DataInterface->ImplementedInterfaces.AddDefaulted_GetRef();
 				NewImplementedDataInterface.DataInterface = StructEntryInfo->FromDataInterface;
+				NewImplementedDataInterface.NativeInterface = StructEntryInfo->NativeInterface;
 				NewImplementedDataInterface.VariableIndex = StructEntryInfo.GetIndex();
 				NewImplementedDataInterface.NumVariables = 1;
 				NewImplementedDataInterface.bAutoBindToHost = StructEntryInfo->bAutoBindDataInterfaceToHost; 
diff --git a/Engine/Plugins/Experimental/AnimNextAnimGraph/AnimNextAnimGraph.uplugin b/Engine/Plugins/Experimental/AnimNextAnimGraph/AnimNextAnimGraph.uplugin
--- a/Engine/Plugins/Experimental/AnimNextAnimGraph/AnimNextAnimGraph.uplugin
+++ b/Engine/Plugins/Experimental/AnimNextAnimGraph/AnimNextAnimGraph.uplugin
@@ -11,7 +11,7 @@
 	"MarketplaceURL" : "",
 	"SupportURL" : "",
 	"EnabledByDefault" : false,
-	"CanContainContent" : false,
+	"CanContainContent" : true,
 	"IsBetaVersion": false,
 	"IsExperimentalVersion": true,
 	"Installed" : false,
@@ -22,7 +22,7 @@
 			"Enabled": true
 		},
 		{
-			"Name": "Chooser",
+			"Name": "RigVM",
 			"Enabled": true
 		}
 	],
@@ -33,6 +33,11 @@
 			"Type" : "Runtime",
 			"LoadingPhase" : "PreDefault"
 		},
+		{
+			"Name" : "AnimNextAnimGraphEditor",
+			"Type" : "Editor",
+			"LoadingPhase" : "Default"
+		},
 		{
 			"Name" : "AnimNextAnimGraphUncookedOnly",
 			"Type" : "UncookedOnly",
diff --git a/Engine/Plugins/Experimental/AnimNextAnimGraph/Config/BaseAnimNextAnimGraph.ini b/Engine/Plugins/Experimental/AnimNextAnimGraph/Config/BaseAnimNextAnimGraph.ini
--- a/Engine/Plugins/Experimental/AnimNextAnimGraph/Config/BaseAnimNextAnimGraph.ini
+++ b/Engine/Plugins/Experimental/AnimNextAnimGraph/Config/BaseAnimNextAnimGraph.ini
@@ -1,3 +1,7 @@
+; Default settings/mappings for factories
+[/Script/AnimNextAnimGraph.AnimNextAnimGraphSettings]
++AssetGraphMappings=(AssetType="/Script/Engine.AnimSequence",AnimationGraph="/AnimNextAnimGraph/FactoryGraphs/AG_AnimSequencePlayer.AG_AnimSequencePlayer",Variable="AnimSequence")
+
 [CoreRedirects]
 ; Traits rename
 +StructRedirects=(OldName="/Script/AnimNext.AnimNextBlendByBoolDecoratorSharedData", NewName="/Script/AnimNextAnimGraph.AnimNextBlendByBoolTraitSharedData")
@@ -28,4 +32,8 @@
 +PropertyRedirects=(OldName="/Script/AnimNextGraph.AnimNode_AnimNextGraph.AnimNextGraph", NewName="Module")
 
 ; Rename K2 Node
-+ClassRedirects=(OldName="/Script/AnimNextAnimGraphUncookedOnly.K2Node_PlayAnim", NewName="/Script/AnimNextAnimGraphUncookedOnly.K2Node_AnimNextPlayAnim")
\ No newline at end of file
++ClassRedirects=(OldName="/Script/AnimNextAnimGraphUncookedOnly.K2Node_PlayAnim", NewName="/Script/AnimNextAnimGraphUncookedOnly.K2Node_AnimNextPlayAnim")
+
+; Rename PlayAnim parameter
++PropertyRedirects=(OldName="/Script/AnimNextAnimGraphUncookedOnly.CreateProxyObjectForPlayAnim.AnimationObject", NewName="AnimSequence")
++PropertyRedirects=(OldName="/Script/AnimNextAnimGraphUncookedOnly.K2Node_PlayAnim.AnimationObject", NewName="AnimSequence")
diff --git a/Engine/Plugins/Experimental/AnimNextAnimGraph/Source/AnimNextAnimGraph/AnimNextAnimGraph.Build.cs b/Engine/Plugins/Experimental/AnimNextAnimGraph/Source/AnimNextAnimGraph/AnimNextAnimGraph.Build.cs
--- a/Engine/Plugins/Experimental/AnimNextAnimGraph/Source/AnimNextAnimGraph/AnimNextAnimGraph.Build.cs
+++ b/Engine/Plugins/Experimental/AnimNextAnimGraph/Source/AnimNextAnimGraph/AnimNextAnimGraph.Build.cs
@@ -14,7 +14,6 @@ namespace UnrealBuildTool.Rules
 					"RigVM",
 					"Engine",
 					"AnimNext",
-					"Chooser",
 				}
 			);
 		}
diff --git a/Engine/Plugins/Experimental/AnimNextAnimGraph/Source/AnimNextAnimGraph/Internal/AnimNextAnimGraphSettings.h b/Engine/Plugins/Experimental/AnimNextAnimGraph/Source/AnimNextAnimGraph/Internal/AnimNextAnimGraphSettings.h
new file mode 100644
--- /dev/null
+++ b/Engine/Plugins/Experimental/AnimNextAnimGraph/Source/AnimNextAnimGraph/Internal/AnimNextAnimGraphSettings.h
@@ -0,0 +1,94 @@
+﻿// Copyright Epic Games, Inc. All Rights Reserved.
+
+#pragma once
+
+#include "StructUtils/StructView.h"
+#include "Templates/SubclassOf.h"
+#include "UObject/ObjectKey.h"
+
+#include "AnimNextAnimGraphSettings.generated.h"
+
+struct FInstancedPropertyBag;
+class UAnimNextAnimationGraph;
+class UAnimNextAnimGraphSettings;
+struct FAnimNextDataInterfacePayload;
+
+USTRUCT()
+struct FAnimNextAssetGraphMapping
+{
+	GENERATED_BODY()
+
+private:
+	friend UAnimNextAnimGraphSettings;
+
+	// The object type/class that this mapping handles
+	UPROPERTY(EditAnywhere, Category = "Mapping")
+	TSoftClassPtr<UObject> AssetType;
+
+	// The animation graph class that the type maps to
+	UPROPERTY(EditAnywhere, Category = "Mapping")
+	TSoftObjectPtr<const UAnimNextAnimationGraph> AnimationGraph;
+
+	// The public variable that will be set when creating an instance of AnimationGraph, from an asset of AssetType.
+	// For example, when mapping from an AnimSequence asset to a graph, the graph will have the AnimSequence asset set on a property of this name.
+	// Leave as 'None' to not set the asset into a variable of the graph.
+	// Must be an 'object' property of a compatible type to AssetType.
+	UPROPERTY(EditAnywhere, Category = "Mapping")
+	FName Variable;
+};
+
+UCLASS(Config=AnimNextAnimGraph, defaultconfig)
+class ANIMNEXTANIMGRAPH_API UAnimNextAnimGraphSettings : public UObject
+{
+	GENERATED_BODY()
+
+public:
+	// Given an object, return an animation graph to instantiate and a set of native interface structs that can be used to communicate with the instance
+	const UAnimNextAnimationGraph* GetGraphFromObject(const UObject* InObject) const;
+
+	// Given an asset class, return whether an animation graph can be made via GetGraphFromObject.
+	bool CanGetGraphFromAssetClass(const UClass* InClass) const;
+
+	// Given an object, return the name of the variable to inject into the graph's payload when making a payload for that object's graph (the
+	// 'Variable' property of its mapping.
+	FName GetInjectedVariableNameFromObject(const UObject* InObject) const;
+
+	// Given an object and the resulting graph, generate native interface payloads for the graph.
+	// If the native interface already exists in InOutGraphPayload, it will not be created by this call.
+	void GetNativePayloadFromGraph(const UObject* InObject, const UAnimNextAnimationGraph* InAnimationGraph, FAnimNextDataInterfacePayload& InOutGraphPayload) const;
+	void GetNativePayloadFromGraph(const UObject* InObject, const UAnimNextAnimationGraph* InAnimationGraph, TArray<FInstancedStruct>& InOutGraphPayload) const;
+
+	// Given an object and the resulting graph, generate an interface payload for the graph, regardless if a native payload exists for the graph.
+	// The generated instanced property bag will encompass public variables of ALL data interfaces implemented by the graph in one bundle
+	void GetNonNativePayloadFromGraph(const UObject* InObject, const UAnimNextAnimationGraph* InAnimationGraph, FAnimNextDataInterfacePayload& InOutGraphPayload) const;
+	void GetNonNativePayloadFromGraph(const UObject* InObject, const UAnimNextAnimationGraph* InAnimationGraph, FInstancedPropertyBag& InOutGraphPayload) const;
+
+	// Rebuild mappings (loading assets synchronously if required) for lookup
+	void LoadAndRebuildMappings();
+
+private:
+#if WITH_EDITOR
+	// UObject interface
+	virtual void PostEditChangeProperty(FPropertyChangedEvent& PropertyChangedEvent) override;
+#endif
+
+	// Gets all allowed asset classes that users can reference that map via GetGraphFromObject
+	UFUNCTION()
+	static TArray<UClass*> GetAllowedAssetClasses();
+	TArray<UClass*> GetAllowedAssetClassesImpl() const;
+	
+private:
+	// Mappings from assets to animation graphs
+	UPROPERTY(Config, EditAnywhere, DisplayName = "Asset/Graph Mappings", Category = "Mappings")
+	TArray<FAnimNextAssetGraphMapping> AssetGraphMappings;
+
+	// References to loaded default graphs to prevent GC
+	UPROPERTY(Transient)
+	TArray<TObjectPtr<const UAnimNextAnimationGraph>> LoadedGraphs;
+
+	// Map derived from AssetGraphMappings
+	TMap<TObjectKey<const UClass>, int32> AssetGraphMap;
+
+	// Flag to verify correct mapping rebuild behavior during startup
+	bool bMappingsBuiltAtLeastOnce = false;
+};
\ No newline at end of file
diff --git a/Engine/Plugins/Experimental/AnimNextAnimGraph/Source/AnimNextAnimGraph/Internal/PlayAnim/PlayAnimCallbackProxy.h b/Engine/Plugins/Experimental/AnimNextAnimGraph/Source/AnimNextAnimGraph/Internal/PlayAnim/PlayAnimCallbackProxy.h
--- a/Engine/Plugins/Experimental/AnimNextAnimGraph/Source/AnimNextAnimGraph/Internal/PlayAnim/PlayAnimCallbackProxy.h
+++ b/Engine/Plugins/Experimental/AnimNextAnimGraph/Source/AnimNextAnimGraph/Internal/PlayAnim/PlayAnimCallbackProxy.h
@@ -30,7 +30,7 @@ class UPlayAnimCallbackProxy : public UObject
 	static ANIMNEXTANIMGRAPH_API UPlayAnimCallbackProxy* CreateProxyObjectForPlayAnim(
 		class UAnimNextComponent* AnimNextComponent,
 		FName SlotName,
-		class UAnimSequence* AnimationObject,
+		class UAnimSequence* AnimSequence,
 		float PlayRate = 1.0f,
 		float StartPosition = 0.0f,
 		FAnimNextPlayAnimBlendSettings BlendInSettings = FAnimNextPlayAnimBlendSettings(),
@@ -69,7 +69,7 @@ protected:
 	ANIMNEXTANIMGRAPH_API bool Play(
 		class UAnimNextComponent* AnimNextComponent,
 		FName SlotName,
-		class UAnimSequence* AnimationObject,
+		class UAnimSequence* AnimSequence,
 		float PlayRate,
 		float StartPosition,
 		const UE::AnimNext::FPlayAnimBlendSettings& BlendInSettings,
diff --git a/Engine/Plugins/Experimental/AnimNextAnimGraph/Source/AnimNextAnimGraph/Internal/Traits/BlendStackTrait.h b/Engine/Plugins/Experimental/AnimNextAnimGraph/Source/AnimNextAnimGraph/Internal/Traits/BlendStackTrait.h
--- a/Engine/Plugins/Experimental/AnimNextAnimGraph/Source/AnimNextAnimGraph/Internal/Traits/BlendStackTrait.h
+++ b/Engine/Plugins/Experimental/AnimNextAnimGraph/Source/AnimNextAnimGraph/Internal/Traits/BlendStackTrait.h
@@ -31,7 +31,7 @@ struct FAnimNextBlendStackTraitSharedData : public FAnimNextTraitSharedData
 
 	/** Asset to use as a sub-graph */
 	UPROPERTY(EditAnywhere, Category = "Default")
-	TObjectPtr<const UAnimNextAnimationGraph> AnimationGraph;
+	TObjectPtr<const UObject> Object;
 
 	/** Blend time to use when sub-graph pin input changes */
 	UPROPERTY(EditAnywhere, Category = "Default")
@@ -43,7 +43,7 @@ struct FAnimNextBlendStackTraitSharedData : public FAnimNextTraitSharedData
 
 	// Latent pin support boilerplate
 	#define TRAIT_LATENT_PROPERTIES_ENUMERATOR(GeneratorMacro) \
-		GeneratorMacro(AnimationGraph) \
+		GeneratorMacro(Object) \
 		GeneratorMacro(BlendTime) \
 		GeneratorMacro(bForceBlend) \
 
@@ -58,7 +58,7 @@ struct FAnimNextBlendStackRequesterTraitSharedData : public FAnimNextTraitShared
 
 	/** Asset to use as a sub-graph */
 	UPROPERTY(EditAnywhere, Category = "Default")
-	TObjectPtr<const UAnimNextAnimationGraph> AnimationGraph;
+	TObjectPtr<const UObject> Object;
 
 	/** Blend time to use when sub-graph pin input changes */
 	UPROPERTY(EditAnywhere, Category = "Default")
@@ -70,7 +70,7 @@ struct FAnimNextBlendStackRequesterTraitSharedData : public FAnimNextTraitShared
 
 	// Latent pin support boilerplate
 	#define TRAIT_LATENT_PROPERTIES_ENUMERATOR(GeneratorMacro) \
-		GeneratorMacro(AnimationGraph) \
+		GeneratorMacro(Object) \
 		GeneratorMacro(BlendTime) \
 		GeneratorMacro(bForceBlend) \
 
@@ -172,8 +172,8 @@ namespace UE::AnimNext
 	/**
 	 * FBlendStackTrait
 	 * 
-	 * Standalone blend stack that exposes desired graph and blend time as pins. 
-	 * A new graph will be pushed every time the newest graph does not match the pin (or bForceBlend is true).
+	 * Standalone blend stack that exposes desired object and blend time as pins.
+	 * A new graph will be pushed every time the newest object does not match the pin (or bForceBlend is true).
 	 */	
 	struct FBlendStackTrait : FBlendStackCoreTrait, IUpdate
 	{
@@ -203,6 +203,5 @@ namespace UE::AnimNext
 
 		// IUpdate impl
 		virtual void PreUpdate(FUpdateTraversalContext& Context, const TTraitBinding<IUpdate>& Binding, const FTraitUpdateState& TraitState) const override;
-		virtual void OnBecomeRelevant(FUpdateTraversalContext& Context, const TTraitBinding<IUpdate>& Binding, const FTraitUpdateState& TraitState) const override;
 	};
 }
diff --git a/Engine/Plugins/Experimental/AnimNextAnimGraph/Source/AnimNextAnimGraph/Private/AnimNextAnimGraphSettings.cpp b/Engine/Plugins/Experimental/AnimNextAnimGraph/Source/AnimNextAnimGraph/Private/AnimNextAnimGraphSettings.cpp
new file mode 100644
--- /dev/null
+++ b/Engine/Plugins/Experimental/AnimNextAnimGraph/Source/AnimNextAnimGraph/Private/AnimNextAnimGraphSettings.cpp
@@ -0,0 +1,292 @@
+﻿// Copyright Epic Games, Inc. All Rights Reserved.
+
+#include "AnimNextAnimGraphSettings.h"
+
+#include "AnimNextDataInterfacePayload.h"
+#include "DataInterface/AnimNextNativeDataInterface.h"
+#include "Graph/AnimNextAnimationGraph.h"
+#include "Logging/StructuredLog.h"
+#include "Misc/EnumerateRange.h"
+#include "TraitInterfaces/IPayloadFactory.h"
+
+#if WITH_EDITOR
+void UAnimNextAnimGraphSettings::PostEditChangeProperty(FPropertyChangedEvent& PropertyChangedEvent)
+{
+	if(PropertyChangedEvent.Property)
+	{
+		if( PropertyChangedEvent.Property->GetFName() == GET_MEMBER_NAME_CHECKED(UAnimNextAnimGraphSettings, AssetGraphMappings) ||
+			PropertyChangedEvent.Property->GetFName() == GET_MEMBER_NAME_CHECKED(FAnimNextAssetGraphMapping, AssetType) ||
+			PropertyChangedEvent.Property->GetFName() == GET_MEMBER_NAME_CHECKED(FAnimNextAssetGraphMapping, AnimationGraph) ||
+			PropertyChangedEvent.Property->GetFName() == GET_MEMBER_NAME_CHECKED(FAnimNextAssetGraphMapping, Variable))
+		{
+			LoadAndRebuildMappings();
+		}
+	}
+}
+#endif
+
+const UAnimNextAnimationGraph* UAnimNextAnimGraphSettings::GetGraphFromObject(const UObject* InObject) const
+{
+	check(InObject);
+	ensure(bMappingsBuiltAtLeastOnce);
+
+	const UAnimNextAnimationGraph* AnimationGraph = Cast<UAnimNextAnimationGraph>(InObject);
+	if(AnimationGraph)
+	{
+		// Object is an animation graph, so we early out
+		return AnimationGraph;
+	}
+
+	// Check this class and its parent classes in turn
+	const UClass* Class = InObject->GetClass();
+	int32 Index = INDEX_NONE;
+	while(Class)
+	{
+		const int32* IndexPtr = AssetGraphMap.Find(Class);
+		if(IndexPtr != nullptr)
+		{
+			Index = *IndexPtr;
+			break;
+		}
+
+		Class = Class->GetSuperClass();
+	}
+
+	if(!AssetGraphMappings.IsValidIndex(Index))
+	{
+		return nullptr;
+	}
+
+	return AssetGraphMappings[Index].AnimationGraph.Get();
+}
+
+bool UAnimNextAnimGraphSettings::CanGetGraphFromAssetClass(const UClass* InClass) const
+{
+	ensure(bMappingsBuiltAtLeastOnce);
+
+	if(InClass == nullptr)
+	{
+		return false;
+	}
+	
+	if(InClass == UAnimNextAnimationGraph::StaticClass())
+	{
+		return true;
+	}
+
+	const UClass* Class = InClass;
+	while(Class)
+	{
+		if(AssetGraphMap.Find(Class) != nullptr)
+		{
+			return true;
+		}
+
+		Class = Class->GetSuperClass();
+	}
+
+	return false;
+}
+
+FName UAnimNextAnimGraphSettings::GetInjectedVariableNameFromObject(const UObject* InObject) const
+{
+	check(InObject);
+	ensure(bMappingsBuiltAtLeastOnce);
+
+	const UAnimNextAnimationGraph* AnimationGraph = Cast<UAnimNextAnimationGraph>(InObject);
+	if(AnimationGraph)
+	{
+		// Object is an animation graph, so we early out
+		return NAME_None;
+	}
+
+	// Check this class and its parent classes in turn
+	const UClass* Class = InObject->GetClass();
+	int32 Index = INDEX_NONE;
+	while(Class)
+	{
+		const int32* IndexPtr = AssetGraphMap.Find(Class);
+		if(IndexPtr != nullptr)
+		{
+			Index = *IndexPtr;
+			break;
+		}
+
+		Class = Class->GetSuperClass();
+	}
+
+	if(!AssetGraphMappings.IsValidIndex(Index))
+	{
+		return NAME_None;
+	}
+
+	return AssetGraphMappings[Index].Variable;
+}
+
+
+void UAnimNextAnimGraphSettings::GetNativePayloadFromGraph(const UObject* InObject, const UAnimNextAnimationGraph* InAnimationGraph, FAnimNextDataInterfacePayload& InOutGraphPayload) const
+{
+	GetNativePayloadFromGraph(InObject, InAnimationGraph, InOutGraphPayload.OwnedNativePayloads);
+	InOutGraphPayload.bCombinedPayloadsDirty = true;
+}
+
+void UAnimNextAnimGraphSettings::GetNativePayloadFromGraph(const UObject* InObject, const UAnimNextAnimationGraph* InAnimationGraph, TArray<FInstancedStruct>& InOutGraphPayload) const
+{
+	if(InObject == nullptr || InAnimationGraph == nullptr)
+	{
+		return;
+	}
+
+	auto FindNativeInterface = [&InOutGraphPayload](const UScriptStruct* InNativeInterface)
+	{
+		return InOutGraphPayload.FindByPredicate([InNativeInterface](const FInstancedStruct& InInstancedStruct)
+		{
+			return InInstancedStruct.GetScriptStruct() == InNativeInterface;
+		});
+	};
+
+	// First generate any structs for the supplied graph from the graph's native interfaces, if required
+	for(const FAnimNextImplementedDataInterface& ImplementedInterface : InAnimationGraph->GetImplementedInterfaces())
+	{
+		if(ImplementedInterface.NativeInterface)
+		{
+			// Check if we have an interface already supplied
+			const FInstancedStruct* NativeInterfaceInstance = FindNativeInterface(ImplementedInterface.NativeInterface);
+			if(NativeInterfaceInstance == nullptr)
+			{
+				check(ImplementedInterface.NativeInterface->IsChildOf(FAnimNextNativeDataInterface::StaticStruct()));
+
+				// Create a new native interface
+				FInstancedStruct NewNativeInterface;
+				NewNativeInterface.InitializeAs(ImplementedInterface.NativeInterface);
+
+				FAnimNextNativeDataInterface::FBindToFactoryObjectContext Context;
+				Context.FactoryObject = InObject;
+				Context.DataInterface = InAnimationGraph;
+				NewNativeInterface.GetMutable<FAnimNextNativeDataInterface>().BindToFactoryObject(Context);
+
+				InOutGraphPayload.Add(MoveTemp(NewNativeInterface));
+			}
+		}
+	}
+}
+
+void UAnimNextAnimGraphSettings::GetNonNativePayloadFromGraph(const UObject* InObject, const UAnimNextAnimationGraph* InAnimationGraph, FAnimNextDataInterfacePayload& InOutGraphPayload) const
+{
+	GetNonNativePayloadFromGraph(InObject, InAnimationGraph, InOutGraphPayload.OwnedPayload);
+	InOutGraphPayload.bCombinedPayloadsDirty = true;
+}
+
+void UAnimNextAnimGraphSettings::GetNonNativePayloadFromGraph(const UObject* InObject, const UAnimNextAnimationGraph* InAnimationGraph, FInstancedPropertyBag& InOutGraphPayload) const
+{
+	// Duplicate the public variables
+	InOutGraphPayload = InAnimationGraph->GetPublicVariableDefaults();
+
+	// Setup the factory-mapped variable
+	FName VariableName = GetInjectedVariableNameFromObject(InObject);
+	if(VariableName.IsNone())
+	{
+		return;
+	}
+
+	const FPropertyBagPropertyDesc* Desc = InOutGraphPayload.FindPropertyDescByName(VariableName);
+	if(Desc == nullptr)
+	{
+		return;
+	}
+
+	const FObjectProperty* ObjectProperty = CastField<FObjectProperty>(Desc->CachedProperty);
+	if(ObjectProperty == nullptr)
+	{
+		return;
+	}
+
+#if WITH_EDITOR
+	// Remove editor visibility for the mapped property as UI doesnt want to display it twice
+	const_cast<FObjectProperty*>(ObjectProperty)->PropertyFlags &= ~CPF_Edit;
+#endif
+
+	uint8* ValuePtr = ObjectProperty->ContainerPtrToValuePtr<uint8>(InOutGraphPayload.GetMutableValue().GetMemory());
+	ObjectProperty->SetObjectPropertyValue(ValuePtr, const_cast<UObject*>(InObject));
+}
+
+void UAnimNextAnimGraphSettings::LoadAndRebuildMappings()
+{
+	AssetGraphMap.Reset();
+	LoadedGraphs.Reset();
+
+	for(TEnumerateRef<FAnimNextAssetGraphMapping> Mapping : EnumerateRange(AssetGraphMappings))
+	{
+		const UAnimNextAnimationGraph* AnimationGraph = Mapping->AnimationGraph.LoadSynchronous();
+		if(AnimationGraph == nullptr)
+		{
+			UE_LOGFMT(LogAnimation, Error, "UAnimNextAnimGraphSettings::LoadAndRebuildMappings: Could not load animation graph '{GraphName}'", Mapping->AnimationGraph.ToString());
+			continue;
+		}
+
+		LoadedGraphs.Add(AnimationGraph);
+
+		const UClass* AssetType = Mapping->AssetType.LoadSynchronous();
+		if(AssetType == nullptr)
+		{
+			UE_LOGFMT(LogAnimation, Error, "UAnimNextAnimGraphSettings::LoadAndRebuildMappings: Could not load asset class '{ClassName}'", Mapping->AssetType.ToString());
+			continue;
+		}
+
+		AssetGraphMap.Add(AssetType, Mapping.GetIndex());
+
+		// Warn if the variable is not settable in the public interface of the graph
+		if(AnimationGraph != nullptr && AssetType != nullptr && !Mapping->Variable.IsNone())
+		{
+			const FInstancedPropertyBag& PublicVariableDefaults = AnimationGraph->GetPublicVariableDefaults();
+			const FPropertyBagPropertyDesc* Desc = PublicVariableDefaults.FindPropertyDescByName(Mapping->Variable);
+			if(Desc == nullptr)
+			{
+				UE_LOGFMT(LogAnimation, Error, "UAnimNextAnimGraphSettings::LoadAndRebuildMappings: Could not find public variable '{VariableName}' in graph '{GraphName}'", Mapping->Variable, AnimationGraph->GetPathName());
+				continue;
+			}
+
+			const FObjectProperty* ObjectProperty = CastField<FObjectProperty>(Desc->CachedProperty);
+			if(ObjectProperty == nullptr)
+			{
+				UE_LOGFMT(LogAnimation, Error, "UAnimNextAnimGraphSettings::LoadAndRebuildMappings: Variable '{VariableName}' in graph '{GraphName}' is not of object type", Mapping->Variable, AnimationGraph->GetPathName());
+				return;
+			}
+
+			if(!AssetType->IsChildOf(ObjectProperty->PropertyClass))
+			{
+				UE_LOGFMT(LogAnimation, Error, "UAnimNextAnimGraphSettings::LoadAndRebuildMappings: Variable '{VariableName}' in graph '{GraphName}' is not of a compatible object type: '{TargetType}' vs '{SourceType}'", Mapping->Variable, AnimationGraph->GetPathName(), ObjectProperty->PropertyClass->GetFullName(), AssetType->GetFullName());
+				return;
+			}
+		}
+	}
+
+	bMappingsBuiltAtLeastOnce = true;
+}
+
+TArray<UClass*> UAnimNextAnimGraphSettings::GetAllowedAssetClasses()
+{
+	return GetDefault<UAnimNextAnimGraphSettings>()->GetAllowedAssetClassesImpl();
+}
+
+TArray<UClass*> UAnimNextAnimGraphSettings::GetAllowedAssetClassesImpl() const
+{
+	ensure(bMappingsBuiltAtLeastOnce);
+
+	TArray<UClass*> AllowedClasses;
+
+	// Add all derived classes of UAnimNextAnimationGraph
+	AllowedClasses.Add(UAnimNextAnimationGraph::StaticClass());
+	GetDerivedClasses(UAnimNextAnimationGraph::StaticClass(), AllowedClasses);
+
+	// Add all mappings
+	for(const FAnimNextAssetGraphMapping& Mapping : AssetGraphMappings)
+	{
+		if(UClass* AssetClass = Mapping.AssetType.Get())
+		{
+			AllowedClasses.Add(AssetClass);
+		}
+	}
+
+	return AllowedClasses;
+}
diff --git a/Engine/Plugins/Experimental/AnimNextAnimGraph/Source/AnimNextAnimGraph/Private/GraphInterfaces/AnimNextNativeDataInterface_AnimSequencePlayer.cpp b/Engine/Plugins/Experimental/AnimNextAnimGraph/Source/AnimNextAnimGraph/Private/GraphInterfaces/AnimNextNativeDataInterface_AnimSequencePlayer.cpp
new file mode 100644
--- /dev/null
+++ b/Engine/Plugins/Experimental/AnimNextAnimGraph/Source/AnimNextAnimGraph/Private/GraphInterfaces/AnimNextNativeDataInterface_AnimSequencePlayer.cpp
@@ -0,0 +1,9 @@
+﻿// Copyright Epic Games, Inc. All Rights Reserved.
+
+#include "GraphInterfaces/AnimNextNativeDataInterface_AnimSequencePlayer.h"
+#include "Animation/AnimSequence.h"
+
+void FAnimNextNativeDataInterface_AnimSequencePlayer::BindToFactoryObject(const FBindToFactoryObjectContext& InContext)
+{
+	AnimSequence = CastChecked<UAnimSequence>(InContext.FactoryObject);
+}
diff --git a/Engine/Plugins/Experimental/AnimNextAnimGraph/Source/AnimNextAnimGraph/Private/Module.cpp b/Engine/Plugins/Experimental/AnimNextAnimGraph/Source/AnimNextAnimGraph/Private/Module.cpp
--- a/Engine/Plugins/Experimental/AnimNextAnimGraph/Source/AnimNextAnimGraph/Private/Module.cpp
+++ b/Engine/Plugins/Experimental/AnimNextAnimGraph/Source/AnimNextAnimGraph/Private/Module.cpp
@@ -1,6 +1,7 @@
 // Copyright Epic Games, Inc. All Rights Reserved.
 
 #include "CoreMinimal.h"
+#include "AnimNextAnimGraphSettings.h"
 #include "Animation/BlendProfile.h"
 #include "Curves/CurveFloat.h"
 #include "Module/AnimNextModule.h"
@@ -8,22 +9,41 @@
 #include "Modules/ModuleManager.h"
 #include "RigVMCore/RigVMRegistry.h"
 #include "Animation/AnimSequence.h"
-#include "Chooser.h"
 
 #include "IAnimNextModuleInterface.h"
 #include "TraitInterfaces/IEvaluate.h"
 #include "TraitInterfaces/IUpdate.h"
 #include "EvaluationVM/EvaluationVM.h"
 #include "Graph/AnimNext_LODPose.h"
+#include "PlayAnim/PlayAnimRequest.h"
 
 namespace UE::AnimNext::AnimGraph
 {
 
+// Delay loading of default asset factory mappings in settings 
+static FDelayedAutoRegisterHelper AutoRegisterFactoryMappings(EDelayedRegisterRunPhase::EndOfEngineInit, []
+{
+	UAnimNextAnimGraphSettings* Settings = GetMutableDefault<UAnimNextAnimGraphSettings>();
+	Settings->LoadAndRebuildMappings();
+});
+
 class FModule : public IModuleInterface, IAnimNextAnimGraph
 {
 public:
 	virtual void StartupModule() override
 	{
+		// Ensure that AnimNext modules are loaded so we can correctly load plugin content
+		FModuleManager::LoadModuleChecked<IModuleInterface>("AnimNext");
+#if WITH_EDITORONLY_DATA
+		FModuleManager::LoadModuleChecked<IModuleInterface>("AnimNextUncookedOnly");
+#endif
+
+		// Setup default settings/factories
+		{
+			UAnimNextAnimGraphSettings* Settings = GetMutableDefault<UAnimNextAnimGraphSettings>();
+			Settings->LoadConfig();
+		}
+
 		static TPair<UClass*, FRigVMRegistry::ERegisterObjectOperation> const AllowedObjectTypes[] =
 		{
 			{ UAnimSequence::StaticClass(), FRigVMRegistry::ERegisterObjectOperation::Class },
@@ -31,7 +51,6 @@ public:
 			{ UBlendProfile::StaticClass(), FRigVMRegistry::ERegisterObjectOperation::Class },
 			{ UCurveFloat::StaticClass(), FRigVMRegistry::ERegisterObjectOperation::Class },
 			{ UAnimNextModule::StaticClass(), FRigVMRegistry::ERegisterObjectOperation::Class },
-			{ UChooserTable::StaticClass(), FRigVMRegistry::ERegisterObjectOperation::Class },
 		};
 
 		FRigVMRegistry::Get().RegisterObjectTypes(AllowedObjectTypes);
diff --git a/Engine/Plugins/Experimental/AnimNextAnimGraph/Source/AnimNextAnimGraph/Private/PlayAnim/PlayAnimCallbackProxy.cpp b/Engine/Plugins/Experimental/AnimNextAnimGraph/Source/AnimNextAnimGraph/Private/PlayAnim/PlayAnimCallbackProxy.cpp
--- a/Engine/Plugins/Experimental/AnimNextAnimGraph/Source/AnimNextAnimGraph/Private/PlayAnim/PlayAnimCallbackProxy.cpp
+++ b/Engine/Plugins/Experimental/AnimNextAnimGraph/Source/AnimNextAnimGraph/Private/PlayAnim/PlayAnimCallbackProxy.cpp
@@ -2,6 +2,7 @@
 
 #include "PlayAnim/PlayAnimCallbackProxy.h"
 #include "Animation/AnimSequence.h"
+#include "GraphInterfaces/AnimNextNativeDataInterface_AnimSequencePlayer.h"
 
 #include UE_INLINE_GENERATED_CPP_BY_NAME(PlayAnimCallbackProxy)
 
@@ -16,7 +17,7 @@ UPlayAnimCallbackProxy::UPlayAnimCallbackProxy(const FObjectInitializer& ObjectI
 UPlayAnimCallbackProxy* UPlayAnimCallbackProxy::CreateProxyObjectForPlayAnim(
 	UAnimNextComponent* AnimNextComponent,
 	FName SlotName,
-	UAnimSequence* AnimationObject,
+	UAnimSequence* AnimSequence,
 	float PlayRate,
 	float StartPosition,
 	UE::AnimNext::FPlayAnimBlendSettings BlendInSettings,
@@ -24,7 +25,7 @@ UPlayAnimCallbackProxy* UPlayAnimCallbackProxy::CreateProxyObjectForPlayAnim(
 {
 	UPlayAnimCallbackProxy* Proxy = NewObject<UPlayAnimCallbackProxy>();
 	Proxy->SetFlags(RF_StrongRefOnFrame);
-	Proxy->Play(AnimNextComponent, SlotName, AnimationObject, PlayRate, StartPosition, BlendInSettings, BlendOutSettings);
+	Proxy->Play(AnimNextComponent, SlotName, AnimSequence, PlayRate, StartPosition, BlendInSettings, BlendOutSettings);
 	return Proxy;
 }
 
@@ -46,19 +47,19 @@ UPlayAnimCallbackProxy* UPlayAnimCallbackProxy::CreateProxyObjectForPlayAsset(
 bool UPlayAnimCallbackProxy::Play(
 	UAnimNextComponent* AnimNextComponent,
 	FName SlotName,
-	UAnimSequence* AnimationObject,
+	UAnimSequence* AnimSequence,
 	float PlayRate,
 	float StartPosition,
 	const UE::AnimNext::FPlayAnimBlendSettings& BlendInSettings,
 	const UE::AnimNext::FPlayAnimBlendSettings& BlendOutSettings)
 {
-	FInstancedStruct Payload;
-	Payload.InitializeAs<FAnimNextPlayAnimPayload>();
-	FAnimNextPlayAnimPayload& PlayAnimPayload = Payload.GetMutable<FAnimNextPlayAnimPayload>();
-	PlayAnimPayload.AnimationObject = AnimationObject;
-	PlayAnimPayload.PlayRate = PlayRate;
-	PlayAnimPayload.StartPosition = StartPosition;
-	return Play(AnimNextComponent, SlotName, AnimationObject, MoveTemp(Payload), BlendInSettings, BlendOutSettings);
+	FInstancedStruct NativeInterface;
+	NativeInterface.InitializeAs<FAnimNextNativeDataInterface_AnimSequencePlayer>();
+	FAnimNextNativeDataInterface_AnimSequencePlayer& PlayNativeInterface = NativeInterface.GetMutable<FAnimNextNativeDataInterface_AnimSequencePlayer>();
+	PlayNativeInterface.AnimSequence = AnimSequence;
+	PlayNativeInterface.PlayRate = PlayRate;
+	PlayNativeInterface.StartPosition = StartPosition;
+	return Play(AnimNextComponent, SlotName, AnimSequence, MoveTemp(NativeInterface), BlendInSettings, BlendOutSettings);
 }
 
 bool UPlayAnimCallbackProxy::Play(
@@ -77,7 +78,7 @@ bool UPlayAnimCallbackProxy::Play(
 		RequestArgs.Object = Object;
 		RequestArgs.BlendInSettings = BlendInSettings;
 		RequestArgs.BlendOutSettings = BlendOutSettings;
-		RequestArgs.Payload = MoveTemp(Payload);
+		RequestArgs.Payload.AddNative(MoveTemp(Payload));
 
 		auto Request = UE::AnimNext::MakePlayAnimRequest();
 		Request->OnCompleted.BindUObject(this, &UPlayAnimCallbackProxy::OnPlayAnimCompleted);
diff --git a/Engine/Plugins/Experimental/AnimNextAnimGraph/Source/AnimNextAnimGraph/Private/PlayAnim/PlayAnimRequest.cpp b/Engine/Plugins/Experimental/AnimNextAnimGraph/Source/AnimNextAnimGraph/Private/PlayAnim/PlayAnimRequest.cpp
--- a/Engine/Plugins/Experimental/AnimNextAnimGraph/Source/AnimNextAnimGraph/Private/PlayAnim/PlayAnimRequest.cpp
+++ b/Engine/Plugins/Experimental/AnimNextAnimGraph/Source/AnimNextAnimGraph/Private/PlayAnim/PlayAnimRequest.cpp
@@ -10,7 +10,7 @@ namespace UE::AnimNext
 	bool FPlayAnimRequest::Play(FPlayAnimRequestArgs&& InRequestArgs, UAnimNextComponent* InComponent)
 	{
 		check(IsInGameThread());
-		if (!InRequestArgs.Payload.IsValid() || InComponent == nullptr)
+		if (InRequestArgs.Payload.Get().Num() == 0 || InComponent == nullptr)
 		{
 			return false;	// Nothing to play
 		}
@@ -159,8 +159,7 @@ namespace UE::AnimNext
 
 	void FPlayAnimRequest::AddReferencedObjects(FReferenceCollector& Collector)
 	{
-		Collector.AddReferencedObject(RequestArgs.Object);
-		Collector.AddPropertyReferencesWithStructARO(FInstancedPropertyBag::StaticStruct(), &RequestArgs.Payload);
+		Collector.AddPropertyReferences(FPlayAnimRequestArgs::StaticStruct(), &RequestArgs);
 		Collector.AddReferencedObject(Component);
 	}
 }
diff --git a/Engine/Plugins/Experimental/AnimNextAnimGraph/Source/AnimNextAnimGraph/Private/PlayAnim/PlayAnimSlotTrait.cpp b/Engine/Plugins/Experimental/AnimNextAnimGraph/Source/AnimNextAnimGraph/Private/PlayAnim/PlayAnimSlotTrait.cpp
--- a/Engine/Plugins/Experimental/AnimNextAnimGraph/Source/AnimNextAnimGraph/Private/PlayAnim/PlayAnimSlotTrait.cpp
+++ b/Engine/Plugins/Experimental/AnimNextAnimGraph/Source/AnimNextAnimGraph/Private/PlayAnim/PlayAnimSlotTrait.cpp
@@ -2,6 +2,8 @@
 
 #include "PlayAnim/PlayAnimSlotTrait.h"
 
+#include "AnimNextAnimGraphSettings.h"
+#include "AnimNextDataInterfacePayload.h"
 #include "Animation/AnimSequence.h"
 #include "DataInterface/AnimNextDataInterfaceHost.h"
 #include "DataInterface/DataInterfaceStructAdapter.h"
@@ -11,6 +13,8 @@
 #include "Graph/AnimNextGraphInstance.h"
 #include "Logging/StructuredLog.h"
 #include "Module/ModuleEvents.h"
+#include "TraitInterfaces/IGraphFactory.h"
+#include "TraitInterfaces/IPayloadFactory.h"
 
 namespace UE::AnimNext
 {
@@ -34,11 +38,12 @@ namespace UE::AnimNext
 	#undef TRAIT_INTERFACE_ENUMERATOR
 	#undef TRAIT_EVENT_ENUMERATOR
 
-	void FPlayAnimSlotTrait::FPlayAnimSlotRequest::Initialize(FPlayAnimRequestPtr InRequest, const FPlayAnimBlendSettings& InBlendSettings, const UAnimNextAnimationGraph* InAnimationGraph)
+	void FPlayAnimSlotTrait::FPlayAnimSlotRequest::Initialize(FPlayAnimRequestPtr InRequest, const FPlayAnimBlendSettings& InBlendSettings, const UAnimNextAnimationGraph* InAnimationGraph, FAnimNextDataInterfacePayload&& InPayload)
 	{
 		Request = InRequest;
 		BlendSettings = InBlendSettings;
 		AnimationGraph = InAnimationGraph;
+		Payload = MoveTemp(InPayload);
 
 		// If no input is provided, we'll use the source
 		State = InAnimationGraph != nullptr ? EPlayAnimRequestState::Active : EPlayAnimRequestState::ActiveSource;
@@ -159,32 +164,15 @@ namespace UE::AnimNext
 
 			FPlayAnimBlendSettings BlendSettings;
 			const UAnimNextAnimationGraph* AnimationGraph = nullptr;
+			FAnimNextDataInterfacePayload Payload;
+
 			if (Request)
 			{
-				// This is a new pending request, lookup the sub-graph to use with our chooser and the desired animation object
-
-				// TODO: Rather than this built-in choosing logic, this could be implemented by a new optional trait:
-				// - Implement some new trait that converts an event payload's object into a graph + payload struct for variable binding
-				// - Could possibly remove Object from FAnimNextPlayAnimRequestArgs, or AnimationObject from FAnimNextPlayAnimPayload, as they are duplicated
-				// - Call this trait's interface here and use the resulting animation graph (binding variables in OnBlendInitiated)
-				// - Reentrancy needs to handle 'same graph, different parameters'
-				const FAnimNextPlayAnimRequestArgs& RequestArgs = Request->GetArgs();
+				// This is a new pending request, lookup the sub-graph to use with any factory trait we have and the desired animation object
+				FAnimNextPlayAnimRequestArgs& RequestArgs = Request->GetMutableArgs();
 				if (RequestArgs.Object)
 				{
-					if (const UChooserTable* Chooser = SharedData->GetSubGraphChooser(Binding))
-					{
-						FAnimNextPlayAnimChooserParameters ChooserParameters;
-						ChooserParameters.AnimationObjectType = RequestArgs.Object->GetClass();
-
-						FChooserEvaluationContext ChooserContext;
-						ChooserContext.AddStructParam(ChooserParameters);
-
-						UChooserTable::EvaluateChooser(ChooserContext, Chooser, FObjectChooserBase::FObjectChooserIteratorCallback::CreateLambda([&AnimationGraph](UObject* InResult)
-							{
-								AnimationGraph = Cast<UAnimNextAnimationGraph>(InResult);
-								return FObjectChooserBase::EIteratorStatus::Stop;
-							}));
-					}
+					AnimationGraph = IGraphFactory::GetGraphFromObjectWithFallback(Context, Binding, RequestArgs.Object);
 
 					if (AnimationGraph != nullptr)
 					{
@@ -195,11 +183,16 @@ namespace UE::AnimNext
 						{
 							if (OwnerGraphInstance->UsesAnimationGraph(AnimationGraph) && OwnerGraphInstance->UsesEntryPoint(EntryPoint))
 							{
+								UE_LOGFMT(LogAnimation, Warning, "Ignoring PlayAnim request for {0}, re-entrancy detected", AnimationGraph->GetFName());
 								return;
 							}
 
 							OwnerGraphInstance = OwnerGraphInstance->GetParentGraphInstance();
 						}
+
+						Payload = MoveTemp(RequestArgs.Payload);
+
+						IPayloadFactory::GetPayloadForGraphWithFallback(Context, Binding, RequestArgs.Object, AnimationGraph, Payload);
 					}
 
 					BlendSettings = RequestArgs.BlendInSettings;
@@ -224,7 +217,7 @@ namespace UE::AnimNext
 			const int32 FreeRequestIndex = FindFreeRequestIndexOrAdd(*InstanceData);
 
 			FPlayAnimSlotRequest& SlotRequest = InstanceData->SlotRequests[FreeRequestIndex];
-			SlotRequest.Initialize(Request, BlendSettings, AnimationGraph);
+			SlotRequest.Initialize(Request, BlendSettings, AnimationGraph, MoveTemp(Payload));
 
 			const int32 OldChildIndex = InstanceData->CurrentlyActiveRequestIndex;
 			const int32 NewChildIndex = FreeRequestIndex;
@@ -295,7 +288,7 @@ namespace UE::AnimNext
 					const int32 FreeRequestIndex = FindFreeRequestIndexOrAdd(*InstanceData);
 
 					FPlayAnimSlotRequest& FreeSlotRequest = InstanceData->SlotRequests[FreeRequestIndex];
-					FreeSlotRequest.Initialize(FPlayAnimRequestPtr(), RequestArgs.BlendOutSettings, nullptr);
+					FreeSlotRequest.Initialize(FPlayAnimRequestPtr(), RequestArgs.BlendOutSettings, nullptr, FAnimNextDataInterfacePayload());
 
 					const int32 OldChildIndex = InstanceData->CurrentlyActiveRequestIndex;
 					const int32 NewChildIndex = FreeRequestIndex;
@@ -384,10 +377,7 @@ namespace UE::AnimNext
 				const FName EntryPoint = SlotRequest.AnimationGraph->DefaultEntryPoint;
 				SlotRequest.GraphInstance = SlotRequest.AnimationGraph->AllocateInstance(nullptr, &Binding.GetTraitPtr().GetNodeInstance()->GetOwner(), EntryPoint);
 				SlotRequest.ChildPtr = SlotRequest.GraphInstance.IsValid() ? SlotRequest.GraphInstance->GetGraphRootPtr() : FTraitPtr();
-
-				// Note: args are mutable here as bindings allow writes!
-				FDataInterfaceStructAdapter VariableBinding(SlotRequest.AnimationGraph, SlotRequest.Request->GetMutableArgs().Payload);
-				SlotRequest.GraphInstance->BindPublicVariables({ &VariableBinding });
+				SlotRequest.GraphInstance->BindPublicVariables(SlotRequest.Payload.Get());
 
 				// TODO: Validate that our child implements the ITimeline interface
 
diff --git a/Engine/Plugins/Experimental/AnimNextAnimGraph/Source/AnimNextAnimGraph/Private/PlayAnim/PlayAnimSlotTrait.h b/Engine/Plugins/Experimental/AnimNextAnimGraph/Source/AnimNextAnimGraph/Private/PlayAnim/PlayAnimSlotTrait.h
--- a/Engine/Plugins/Experimental/AnimNextAnimGraph/Source/AnimNextAnimGraph/Private/PlayAnim/PlayAnimSlotTrait.h
+++ b/Engine/Plugins/Experimental/AnimNextAnimGraph/Source/AnimNextAnimGraph/Private/PlayAnim/PlayAnimSlotTrait.h
@@ -14,7 +14,6 @@
 #include "Graph/AnimNextAnimationGraph.h"
 #include "Graph/AnimNextGraphInstance.h"
 #include "PlayAnim/PlayAnimEvents.h"
-#include "Chooser.h"
 
 #include "PlayAnimSlotTrait.generated.h"
 
@@ -41,14 +40,9 @@ struct FAnimNextPlayAnimSlotTraitSharedData : public FAnimNextTraitSharedData
 	UPROPERTY(EditAnywhere, Category = "Default")
 	FName SlotName;
 
-	/** The chooser instance to use to convert animation objects into a sub-graph to play. */
-	UPROPERTY(EditAnywhere, Category = "Default")
-	TObjectPtr<UChooserTable> SubGraphChooser;
-
 	// Latent pin support boilerplate
 	#define TRAIT_LATENT_PROPERTIES_ENUMERATOR(GeneratorMacro) \
 		GeneratorMacro(SlotName) \
-		GeneratorMacro(SubGraphChooser) \
 
 	GENERATE_TRAIT_LATENT_PROPERTIES(FAnimNextPlayAnimSlotTraitSharedData, TRAIT_LATENT_PROPERTIES_ENUMERATOR)
 	#undef TRAIT_LATENT_PROPERTIES_ENUMERATOR
@@ -110,6 +104,9 @@ namespace UE::AnimNext
 			// The module used by the graph instance, as selected by the chooser
 			TObjectPtr<const UAnimNextAnimationGraph> AnimationGraph;
 
+			// Payloads for the graph instance when instantiated
+			FAnimNextDataInterfacePayload Payload;
+
 			// The graph instance
 			TSharedPtr<FAnimNextGraphInstance> GraphInstance;
 
@@ -124,7 +121,7 @@ namespace UE::AnimNext
 			bool bWasRelevant = false;
 
 			// Initializes a request to begin playing
-			void Initialize(FPlayAnimRequestPtr InRequest, const FPlayAnimBlendSettings& InBlendSettings, const UAnimNextAnimationGraph* InAnimationGraph);
+			void Initialize(FPlayAnimRequestPtr InRequest, const FPlayAnimBlendSettings& InBlendSettings, const UAnimNextAnimationGraph* InAnimationGraph, FAnimNextDataInterfacePayload&& InPayload);
 		};
 
 		using FSharedData = FAnimNextPlayAnimSlotTraitSharedData;
diff --git a/Engine/Plugins/Experimental/AnimNextAnimGraph/Source/AnimNextAnimGraph/Private/TraitInterfaces/IGraphFactory.cpp b/Engine/Plugins/Experimental/AnimNextAnimGraph/Source/AnimNextAnimGraph/Private/TraitInterfaces/IGraphFactory.cpp
new file mode 100644
--- /dev/null
+++ b/Engine/Plugins/Experimental/AnimNextAnimGraph/Source/AnimNextAnimGraph/Private/TraitInterfaces/IGraphFactory.cpp
@@ -0,0 +1,50 @@
+﻿// Copyright Epic Games, Inc. All Rights Reserved.
+
+#include "TraitInterfaces/IGraphFactory.h"
+
+#include "AnimNextAnimGraphSettings.h"
+#include "TraitCore/ExecutionContext.h"
+
+namespace UE::AnimNext
+{
+	AUTO_REGISTER_ANIM_TRAIT_INTERFACE(IGraphFactory)
+
+#if WITH_EDITOR
+	const FText& IGraphFactory::GetDisplayName() const
+	{
+		static FText InterfaceName = NSLOCTEXT("TraitInterfaces", "TraitInterface_IGraphFactory_Name", "Graph Factory");
+		return InterfaceName;
+	}
+	const FText& IGraphFactory::GetDisplayShortName() const
+	{
+		static FText InterfaceShortName = NSLOCTEXT("TraitInterfaces", "TraitInterface_IGraphFactory_ShortName", "GF");
+		return InterfaceShortName;
+	}
+#endif // WITH_EDITOR
+
+	const UAnimNextAnimationGraph* IGraphFactory::GetGraphFromObject(FExecutionContext& Context, const TTraitBinding<IGraphFactory>& Binding, const UObject* InObject) const
+	{
+		TTraitBinding<IGraphFactory> SuperBinding;
+		if (Binding.GetStackInterfaceSuper(SuperBinding))
+		{
+			return SuperBinding.GetGraphFromObject(Context, InObject);
+		}
+		return nullptr;
+	}
+
+	const UAnimNextAnimationGraph* IGraphFactory::GetGraphFromObjectWithFallback(FExecutionContext& Context, const FTraitBinding& InBinding, const UObject* InObject)
+	{
+		TTraitBinding<IGraphFactory> GraphFactoryBinding;
+		if(InBinding.GetStackInterface<IGraphFactory>(GraphFactoryBinding))
+		{
+			const UAnimNextAnimationGraph* AnimationGraph = GraphFactoryBinding.GetGraphFromObject(Context, InObject);
+			if(AnimationGraph != nullptr)
+			{
+				return AnimationGraph;
+			}
+		}
+
+		return GetDefault<UAnimNextAnimGraphSettings>()->GetGraphFromObject(InObject);
+	}
+
+}
diff --git a/Engine/Plugins/Experimental/AnimNextAnimGraph/Source/AnimNextAnimGraph/Private/TraitInterfaces/IPayloadFactory.cpp b/Engine/Plugins/Experimental/AnimNextAnimGraph/Source/AnimNextAnimGraph/Private/TraitInterfaces/IPayloadFactory.cpp
new file mode 100644
--- /dev/null
+++ b/Engine/Plugins/Experimental/AnimNextAnimGraph/Source/AnimNextAnimGraph/Private/TraitInterfaces/IPayloadFactory.cpp
@@ -0,0 +1,46 @@
+﻿// Copyright Epic Games, Inc. All Rights Reserved.
+
+#include "TraitInterfaces/IPayloadFactory.h"
+
+#include "AnimNextAnimGraphSettings.h"
+#include "TraitCore/ExecutionContext.h"
+
+namespace UE::AnimNext
+{
+	AUTO_REGISTER_ANIM_TRAIT_INTERFACE(IPayloadFactory)
+
+#if WITH_EDITOR
+	const FText& IPayloadFactory::GetDisplayName() const
+	{
+		static FText InterfaceName = NSLOCTEXT("TraitInterfaces", "TraitInterface_IPayloadFactory_Name", "Payload Factory");
+		return InterfaceName;
+	}
+	const FText& IPayloadFactory::GetDisplayShortName() const
+	{
+		static FText InterfaceShortName = NSLOCTEXT("TraitInterfaces", "TraitInterface_IPayloadFactory_ShortName", "PF");
+		return InterfaceShortName;
+	}
+#endif // WITH_EDITOR
+
+	void IPayloadFactory::GetPayloadFromGraph(FExecutionContext& Context, const TTraitBinding<IPayloadFactory>& Binding, const UObject* InObject, const UAnimNextAnimationGraph* InGraph, FAnimNextDataInterfacePayload& InOutPayload) const
+	{
+		TTraitBinding<IPayloadFactory> SuperBinding;
+		if (Binding.GetStackInterfaceSuper(SuperBinding))
+		{
+			SuperBinding.GetPayloadFromGraph(Context, InObject, InGraph, InOutPayload);
+		}
+	}
+
+	void IPayloadFactory::GetPayloadForGraphWithFallback(FExecutionContext& Context, const FTraitBinding& InBinding, const UObject* InObject, const UAnimNextAnimationGraph* InGraph, FAnimNextDataInterfacePayload& InOutPayload)
+	{
+		TTraitBinding<IPayloadFactory> PayloadFactoryBinding;
+		if(InBinding.GetStackInterface<IPayloadFactory>(PayloadFactoryBinding))
+		{
+			PayloadFactoryBinding.GetPayloadFromGraph(Context, InObject, InGraph, InOutPayload);
+		}
+		else
+		{
+			GetDefault<UAnimNextAnimGraphSettings>()->GetNativePayloadFromGraph(InObject, InGraph, InOutPayload);
+		}
+	}
+}
diff --git a/Engine/Plugins/Experimental/AnimNextAnimGraph/Source/AnimNextAnimGraph/Private/Traits/BlendStackTrait.cpp b/Engine/Plugins/Experimental/AnimNextAnimGraph/Source/AnimNextAnimGraph/Private/Traits/BlendStackTrait.cpp
--- a/Engine/Plugins/Experimental/AnimNextAnimGraph/Source/AnimNextAnimGraph/Private/Traits/BlendStackTrait.cpp
+++ b/Engine/Plugins/Experimental/AnimNextAnimGraph/Source/AnimNextAnimGraph/Private/Traits/BlendStackTrait.cpp
@@ -1,9 +1,14 @@
 // Copyright Epic Games, Inc. All Rights Reserved.
 
 #include "Traits/BlendStackTrait.h"
+
+#include "AnimNextAnimGraphSettings.h"
+#include "AnimNextDataInterfacePayload.h"
 #include "DataInterface/DataInterfaceStructAdapter.h"
 #include "TraitCore/NodeInstance.h"
 #include "TraitCore/ExecutionContext.h"
+#include "TraitInterfaces/IGraphFactory.h"
+#include "TraitInterfaces/IPayloadFactory.h"
 
 namespace UE::AnimNext
 {
@@ -288,34 +293,10 @@ namespace UE::AnimNext
 		TTraitBinding<IDiscreteBlend> DiscreteBlendTrait;
 		Binding.GetStackInterface(DiscreteBlendTrait);
 		DiscreteBlendTrait.OnBlendTransition(Context, OldChildIndex, NewChildIndex);
-		
-		InstanceData->CurrentlyActiveGraphIndex = NewChildIndex;
 
-		const int32 NumPayloads = Graph.Request.Payloads.Num();
-		if (NumPayloads > 0)
-		{
-			// stack allocating uninitialized arrays to contain FDataInterfaceStructAdapter(s)
-			TArrayView<FDataInterfaceStructAdapter> Adapters((FDataInterfaceStructAdapter*)FMemory_Alloca(NumPayloads * sizeof(FDataInterfaceStructAdapter)), NumPayloads);
-			TArrayView<IDataInterfaceHost*> PayloadVariableBindings((IDataInterfaceHost**)FMemory_Alloca(NumPayloads * sizeof(IDataInterfaceHost*)), NumPayloads);
-
-			// initializing the FDataInterfaceStructAdapter(s) with their respective Graph.Request.Payloads
-			for (int32 PayloadIndex = 0; PayloadIndex < NumPayloads; ++PayloadIndex)
-			{
-				FInstancedStruct& Payload = Graph.Request.Payloads[PayloadIndex];
-				PayloadVariableBindings[PayloadIndex] = new (&Adapters[PayloadIndex]) FDataInterfaceStructAdapter(Graph.Request.AnimationGraph, Payload);
-			}
-			
-			// binding 
-			Graph.Instance->BindPublicVariables(PayloadVariableBindings);
-			
-			// uninitializing the FDataInterfaceStructAdapter(s)
-			for (int32 PayloadIndex = 0; PayloadIndex < NumPayloads; ++PayloadIndex)
-			{
-				Adapters[PayloadIndex].~FDataInterfaceStructAdapter();
-			}
+		InstanceData->CurrentlyActiveGraphIndex = NewChildIndex;
 
-			// abandoning no longer initialized stack allocated arrays to their demise
-		}
+		Graph.Instance->BindPublicVariables(Graph.Request.GraphPayload.Get());
 	}
 
 	void FBlendStackCoreTrait::GetActiveGraph(FExecutionContext& Context, const TTraitBinding<IBlendStack>& Binding, IBlendStack::FGraphRequestPtr& OutGraphRequest) const
@@ -337,21 +318,32 @@ namespace UE::AnimNext
 		FInstanceData* InstanceData = Binding.GetInstanceData<FInstanceData>();
 		const FSharedData* SharedData = Binding.GetSharedData<FSharedData>();
  
-		const UAnimNextAnimationGraph* DesiredAnimationGraph = SharedData->GetAnimationGraph(Binding);
-		const int32 CurrentlyActiveGraphIndex = InstanceData->CurrentlyActiveGraphIndex;
-
-		const bool bForceBlend = SharedData->GetbForceBlend(Binding);
-		const bool bIsEmpty = InstanceData->CurrentlyActiveGraphIndex == INDEX_NONE;
-		if (DesiredAnimationGraph && (bForceBlend || bIsEmpty || (DesiredAnimationGraph != InstanceData->ChildGraphs[CurrentlyActiveGraphIndex].Request.AnimationGraph)))
+		const UObject* DesiredObject = SharedData->GetObject(Binding);
+		if (DesiredObject != nullptr)
 		{
-			TTraitBinding<IBlendStack> BlendStackTrait;
-			Binding.GetStackInterface(BlendStackTrait);
-
-			IBlendStack::FGraphRequest GraphRequest;
-			GraphRequest.AnimationGraph = DesiredAnimationGraph;
-			GraphRequest.BlendArgs.BlendTime = SharedData->GetBlendTime(Binding);
-
-			BlendStackTrait.PushGraph(Context, MoveTemp(GraphRequest));
+			const UAnimNextAnimationGraph* AnimationGraph = IGraphFactory::GetGraphFromObjectWithFallback(Context, Binding, DesiredObject);
+			if (AnimationGraph != nullptr)
+			{
+				const int32 CurrentlyActiveGraphIndex = InstanceData->CurrentlyActiveGraphIndex;
+				const bool bForceBlend = SharedData->GetbForceBlend(Binding);
+				const bool bIsEmpty = CurrentlyActiveGraphIndex == INDEX_NONE;
+				if (bForceBlend || bIsEmpty || (DesiredObject != InstanceData->ChildGraphs[CurrentlyActiveGraphIndex].Request.FactoryObject))
+				{
+					FAnimNextDataInterfacePayload GraphPayload;
+					IPayloadFactory::GetPayloadForGraphWithFallback(Context, Binding, DesiredObject, AnimationGraph, GraphPayload);
+
+					TTraitBinding<IBlendStack> BlendStackTrait;
+					Binding.GetStackInterface(BlendStackTrait);
+
+					IBlendStack::FGraphRequest GraphRequest;
+					GraphRequest.FactoryObject = DesiredObject;
+					GraphRequest.AnimationGraph = AnimationGraph;
+					GraphRequest.BlendArgs.BlendTime = SharedData->GetBlendTime(Binding);
+					GraphRequest.GraphPayload = MoveTemp(GraphPayload);
+
+					BlendStackTrait.PushGraph(Context, MoveTemp(GraphRequest));
+				}
+			}
 		}
 	}
 
@@ -359,41 +351,32 @@ namespace UE::AnimNext
 	{
 		const FSharedData* SharedData = Binding.GetSharedData<FSharedData>();
 
-		const UAnimNextAnimationGraph* DesiredAnimationGraph = SharedData->GetAnimationGraph(Binding);
-		if (DesiredAnimationGraph != nullptr)
+		const UObject* DesiredObject = SharedData->GetObject(Binding);
+		if (DesiredObject != nullptr)
 		{
-			TTraitBinding<IBlendStack> BlendStackTrait;
-			Binding.GetStackInterface(BlendStackTrait);
-
-			IBlendStack::FGraphRequestPtr ActiveGraphRequest;
-			BlendStackTrait.GetActiveGraph(Context, ActiveGraphRequest);
-			const bool bForceBlend = SharedData->GetbForceBlend(Binding);
-			if (bForceBlend || !ActiveGraphRequest || (DesiredAnimationGraph != ActiveGraphRequest->AnimationGraph))
+			const UAnimNextAnimationGraph* AnimationGraph = IGraphFactory::GetGraphFromObjectWithFallback(Context, Binding, DesiredObject);
+			if (AnimationGraph != nullptr)
 			{
-				IBlendStack::FGraphRequest GraphRequest;
-				GraphRequest.AnimationGraph = DesiredAnimationGraph;
-				GraphRequest.BlendArgs.BlendTime = SharedData->GetBlendTime(Binding);
-
-				BlendStackTrait.PushGraph(Context, MoveTemp(GraphRequest));
+				TTraitBinding<IBlendStack> BlendStackTrait;
+				Binding.GetStackInterface(BlendStackTrait);
+
+				IBlendStack::FGraphRequestPtr ActiveGraphRequest;
+				BlendStackTrait.GetActiveGraph(Context, ActiveGraphRequest);
+				const bool bForceBlend = SharedData->GetbForceBlend(Binding);
+				if (bForceBlend || !ActiveGraphRequest || (DesiredObject != ActiveGraphRequest->FactoryObject))
+				{
+					FAnimNextDataInterfacePayload GraphPayload;
+					IPayloadFactory::GetPayloadForGraphWithFallback(Context, Binding, DesiredObject, AnimationGraph, GraphPayload);
+
+					IBlendStack::FGraphRequest GraphRequest;
+					GraphRequest.FactoryObject = DesiredObject;
+					GraphRequest.AnimationGraph = AnimationGraph;
+					GraphRequest.BlendArgs.BlendTime = SharedData->GetBlendTime(Binding);
+					GraphRequest.GraphPayload = MoveTemp(GraphPayload);
+
+					BlendStackTrait.PushGraph(Context, MoveTemp(GraphRequest));
+				}
 			}
 		}
 	}
-
-	void FBlendStackRequesterTrait::OnBecomeRelevant(FUpdateTraversalContext& Context, const TTraitBinding<IUpdate>& Binding, const FTraitUpdateState& TraitState) const
-	{
-		const FSharedData* SharedData = Binding.GetSharedData<FSharedData>();
-
-		TTraitBinding<IBlendStack> BlendStackTrait;
-		Binding.GetStackInterface(BlendStackTrait);
-
-		const UAnimNextAnimationGraph* DesiredAnimationGraph = SharedData->GetAnimationGraph(Binding);
-		if (DesiredAnimationGraph != nullptr)
-		{
-			IBlendStack::FGraphRequest GraphRequest;
-			GraphRequest.AnimationGraph = DesiredAnimationGraph;
-			GraphRequest.BlendArgs.BlendTime = SharedData->GetBlendTime(Binding);
-
-			BlendStackTrait.PushGraph(Context, MoveTemp(GraphRequest));
-		}
-	}
 }
diff --git a/Engine/Plugins/Experimental/AnimNextAnimGraph/Source/AnimNextAnimGraph/Private/Traits/MakeGraphFromAsset.cpp b/Engine/Plugins/Experimental/AnimNextAnimGraph/Source/AnimNextAnimGraph/Private/Traits/MakeGraphFromAsset.cpp
new file mode 100644
--- /dev/null
+++ b/Engine/Plugins/Experimental/AnimNextAnimGraph/Source/AnimNextAnimGraph/Private/Traits/MakeGraphFromAsset.cpp
@@ -0,0 +1,29 @@
+﻿// Copyright Epic Games, Inc. All Rights Reserved.
+
+#include "Traits/MakeGraphFromAsset.h"
+
+#include "AnimNextAnimGraphSettings.h"
+#include "TraitCore/ExecutionContext.h"
+
+namespace UE::AnimNext
+{
+	AUTO_REGISTER_ANIM_TRAIT(FMakeGraphFromAssetTrait)
+
+	// Trait implementation boilerplate
+	#define TRAIT_INTERFACE_ENUMERATOR(GeneratorMacro) \
+		GeneratorMacro(IGraphFactory) \
+
+		GENERATE_ANIM_TRAIT_IMPLEMENTATION(FMakeGraphFromAssetTrait, TRAIT_INTERFACE_ENUMERATOR, NULL_ANIM_TRAIT_INTERFACE_ENUMERATOR, NULL_ANIM_TRAIT_EVENT_ENUMERATOR)
+	#undef TRAIT_INTERFACE_ENUMERATOR
+
+	const UAnimNextAnimationGraph* FMakeGraphFromAssetTrait::GetGraphFromObject(FExecutionContext& Context, const TTraitBinding<IGraphFactory>& Binding, const UObject* InObject) const
+	{
+		if(InObject == nullptr)
+		{
+			return nullptr;
+		}
+
+		return GetDefault<UAnimNextAnimGraphSettings>()->GetGraphFromObject(InObject);
+	}
+
+}
diff --git a/Engine/Plugins/Experimental/AnimNextAnimGraph/Source/AnimNextAnimGraph/Private/Traits/MakeGraphFromAsset.h b/Engine/Plugins/Experimental/AnimNextAnimGraph/Source/AnimNextAnimGraph/Private/Traits/MakeGraphFromAsset.h
new file mode 100644
--- /dev/null
+++ b/Engine/Plugins/Experimental/AnimNextAnimGraph/Source/AnimNextAnimGraph/Private/Traits/MakeGraphFromAsset.h
@@ -0,0 +1,37 @@
+﻿// Copyright Epic Games, Inc. All Rights Reserved.
+
+#pragma once
+
+#include "CoreMinimal.h"
+#include "TraitCore/Trait.h"
+#include "TraitInterfaces/IGraphFactory.h"
+
+#include "MakeGraphFromAsset.generated.h"
+
+USTRUCT(meta = (DisplayName = "Make Graph From Asset"))
+struct FAnimNextMakeGraphFromAssetTraitSharedData : public FAnimNextTraitSharedData
+{
+	GENERATED_BODY()
+};
+
+namespace UE::AnimNext
+{
+	/**
+	 * FMakeGraphFromAssetTrait
+	 *
+	 * A trait that makes an AnimNext animation graph from an asset
+	 */
+	struct FMakeGraphFromAssetTrait : FAdditiveTrait, IGraphFactory
+	{
+		DECLARE_ANIM_TRAIT(FMakeGraphFromAssetTrait, 0x31b4fc8d, FAdditiveTrait)
+
+		using FSharedData = FAnimNextMakeGraphFromAssetTraitSharedData;
+
+		struct FInstanceData : FTrait::FInstanceData
+		{
+		};
+
+		// IGraphFactory impl
+		virtual const UAnimNextAnimationGraph* GetGraphFromObject(FExecutionContext& Context, const TTraitBinding<IGraphFactory>& Binding, const UObject* InObject) const;
+	};
+}
diff --git a/Engine/Plugins/Experimental/AnimNextAnimGraph/Source/AnimNextAnimGraph/Public/GraphInterfaces/AnimNextNativeDataInterface_AnimSequencePlayer.h b/Engine/Plugins/Experimental/AnimNextAnimGraph/Source/AnimNextAnimGraph/Public/GraphInterfaces/AnimNextNativeDataInterface_AnimSequencePlayer.h
new file mode 100644
--- /dev/null
+++ b/Engine/Plugins/Experimental/AnimNextAnimGraph/Source/AnimNextAnimGraph/Public/GraphInterfaces/AnimNextNativeDataInterface_AnimSequencePlayer.h
@@ -0,0 +1,37 @@
+﻿// Copyright Epic Games, Inc. All Rights Reserved.
+
+#pragma once
+
+#include "CoreMinimal.h"
+#include "DataInterface/AnimNextNativeDataInterface.h"
+#include "AnimNextNativeDataInterface_AnimSequencePlayer.generated.h"
+
+class UAnimSequence;
+
+/**
+* Native interface to built-in AnimSequence player graph at: /AnimNextAnimGraph/FactoryGraphs/AG_AnimSequencePlayer
+ */
+USTRUCT(BlueprintType, DisplayName = "AnimNext Anim Sequence Player Data")
+struct FAnimNextNativeDataInterface_AnimSequencePlayer : public FAnimNextNativeDataInterface
+{
+	GENERATED_BODY()
+
+	// FAnimNextNativeDataInterface interface
+	ANIMNEXTANIMGRAPH_API virtual void BindToFactoryObject(const FBindToFactoryObjectContext& InContext) override;
+
+	// The animation object to play
+	UPROPERTY(BlueprintReadWrite, Category = "Anim Sequence")	// Not editable to hide it in cases where we would end up duplicating factory source asset
+	TObjectPtr<const UAnimSequence> AnimSequence;
+
+	// The play rate to use
+	UPROPERTY(EditAnywhere, BlueprintReadWrite, Category = "Anim Sequence", meta = (EnableCategories))
+	double PlayRate = 1.0f;
+
+	// The timeline's start position
+	UPROPERTY(EditAnywhere, BlueprintReadWrite, Category = "Anim Sequence", meta = (EnableCategories))
+	double StartPosition = 0.0f;
+
+	// Whether to loop the animation
+	UPROPERTY(EditAnywhere, BlueprintReadWrite, Category = "Anim Sequence", meta = (EnableCategories))
+	bool Loop = false;
+};
diff --git a/Engine/Plugins/Experimental/AnimNextAnimGraph/Source/AnimNextAnimGraph/Public/PlayAnim/PlayAnimRequest.h b/Engine/Plugins/Experimental/AnimNextAnimGraph/Source/AnimNextAnimGraph/Public/PlayAnim/PlayAnimRequest.h
--- a/Engine/Plugins/Experimental/AnimNextAnimGraph/Source/AnimNextAnimGraph/Public/PlayAnim/PlayAnimRequest.h
+++ b/Engine/Plugins/Experimental/AnimNextAnimGraph/Source/AnimNextAnimGraph/Public/PlayAnim/PlayAnimRequest.h
@@ -5,6 +5,7 @@
 #include "CoreMinimal.h"
 
 #include "AlphaBlend.h"
+#include "AnimNextDataInterfacePayload.h"
 #include "TraitCore/TraitEvent.h"
 #include "TraitInterfaces/ITimeline.h"
 #include "PlayAnim/PlayAnimStatus.h"
@@ -59,33 +60,6 @@ struct FAnimNextPlayAnimBlendSettings
 	EAnimNextPlayAnimBlendMode BlendMode = EAnimNextPlayAnimBlendMode::Standard;
 };
 
-/**
- * PlayAnim payload
- *
- * Encapsulates the data interface used to play a simple animation
- */
-USTRUCT(BlueprintType)
-struct FAnimNextPlayAnimPayload
-{
-	GENERATED_BODY()
-
-	// The animation object to play with this request
-	UPROPERTY(EditAnywhere, BlueprintReadWrite, Category = "PlayAnim")
-	TObjectPtr<UAnimSequence> AnimationObject;
-
-	// The play rate of the request
-	UPROPERTY(EditAnywhere, BlueprintReadWrite, Category = "PlayAnim")
-	double PlayRate = 1.0f;
-
-	// The timeline start position of the request
-	UPROPERTY(EditAnywhere, BlueprintReadWrite, Category = "PlayAnim")
-	double StartPosition = 0.0f;
-
-	// Whether to loop the animation
-	UPROPERTY(EditAnywhere, BlueprintReadWrite, Category = "PlayAnim")
-	bool IsLooping = false;
-};
-
 /**
  * PlayAnim Request Arguments
  *
@@ -115,7 +89,7 @@ struct FAnimNextPlayAnimRequestArgs
 
 	// Payload that will be applied to the animation graph's variables via its data interfaces.
 	UPROPERTY()
-	FInstancedStruct Payload;
+	FAnimNextDataInterfacePayload Payload;
 };
 
 namespace UE::AnimNext
diff --git a/Engine/Plugins/Experimental/AnimNextAnimGraph/Source/AnimNextAnimGraph/Public/TraitInterfaces/IBlendStack.h b/Engine/Plugins/Experimental/AnimNextAnimGraph/Source/AnimNextAnimGraph/Public/TraitInterfaces/IBlendStack.h
--- a/Engine/Plugins/Experimental/AnimNextAnimGraph/Source/AnimNextAnimGraph/Public/TraitInterfaces/IBlendStack.h
+++ b/Engine/Plugins/Experimental/AnimNextAnimGraph/Source/AnimNextAnimGraph/Public/TraitInterfaces/IBlendStack.h
@@ -6,6 +6,8 @@
 #include "TraitCore/ITraitInterface.h"
 #include "TraitCore/TraitBinding.h"
 #include "AlphaBlend.h"
+#include "AnimNextDataInterfacePayload.h"
+#include "IPayloadFactory.h"
 #include "Graph/AnimNextAnimationGraph.h"
 #include "Graph/AnimNextGraphInstancePtr.h"
 
@@ -27,14 +29,18 @@ namespace UE::AnimNext
 			 * Consider making dynamic payload if we want to implement a special blend framework.
 			 */
 
-			// Payloads applied to the AnimationGraph's variables via its data interfaces.
-			TArray<FInstancedStruct> Payloads;
+			// Payload applied to the AnimationGraph's variables via its data interfaces.
+			FAnimNextDataInterfacePayload GraphPayload;
 
 			// Blend-in duration for this graph request 
 			FAlphaBlendArgs BlendArgs;
 
 			// The template graph to use for the new graph instance
 			TObjectPtr<const UAnimNextAnimationGraph> AnimationGraph;
+
+			// The object that was used for factory generation of the graph.
+			// Used as a key to allow identification of this request without inspecting payloads.
+			TObjectPtr<const UObject> FactoryObject;
 		};
 
 		typedef FGraphRequest* FGraphRequestPtr;
diff --git a/Engine/Plugins/Experimental/AnimNextAnimGraph/Source/AnimNextAnimGraph/Public/TraitInterfaces/IGraphFactory.h b/Engine/Plugins/Experimental/AnimNextAnimGraph/Source/AnimNextAnimGraph/Public/TraitInterfaces/IGraphFactory.h
new file mode 100644
--- /dev/null
+++ b/Engine/Plugins/Experimental/AnimNextAnimGraph/Source/AnimNextAnimGraph/Public/TraitInterfaces/IGraphFactory.h
@@ -0,0 +1,55 @@
+﻿// Copyright Epic Games, Inc. All Rights Reserved.
+
+#pragma once
+
+#include "CoreMinimal.h"
+#include "StructUtils/StructView.h"
+#include "TraitCore/ITraitInterface.h"
+#include "TraitCore/TraitBinding.h"
+
+class UAnimNextAnimationGraph;
+
+namespace UE::AnimNext
+{
+	/**
+	 * IGraphFactory
+	 *
+	 * This interface translates an object and payload struct into a UAnimNextAnimationGraph and payload instance(s), ready for instantiation
+	 */
+	struct ANIMNEXTANIMGRAPH_API IGraphFactory : ITraitInterface
+	{
+		DECLARE_ANIM_TRAIT_INTERFACE(IGraphFactory, 0xb5c08dd6)
+
+		// Gets a graph and a set of native interfaces from an object
+		// @param	InObject				The object to manufacture a graph for (e.g. AnimSequence). If this is a UAnimNextAnimationGraph, it should be returned verbatim
+		// @return The graph to be used.
+		virtual const UAnimNextAnimationGraph* GetGraphFromObject(FExecutionContext& Context, const TTraitBinding<IGraphFactory>& Binding, const UObject* InObject) const;
+
+		// Gets a graph and a set of native interfaces from an object
+		// Falls back to using user settings if no interface is found in the current stack, or if an interface was found but returns nullptr
+		// @param	InObject				The object to manufacture a graph for (e.g. AnimSequence). If this is a UAnimNextAnimationGraph, it should be returned verbatim
+		// @return The graph to be used.
+		static const UAnimNextAnimationGraph* GetGraphFromObjectWithFallback(FExecutionContext& Context, const FTraitBinding& InBinding, const UObject* InObject);
+
+#if WITH_EDITOR
+		virtual const FText& GetDisplayName() const override;
+		virtual const FText& GetDisplayShortName() const override;
+#endif // WITH_EDITOR
+	};
+
+	/**
+	 * Specialization for trait binding.
+	 */
+	template<>
+	struct TTraitBinding<IGraphFactory> : FTraitBinding
+	{
+		// @see IGraphFactory::GetGraphFromObject
+		const UAnimNextAnimationGraph* GetGraphFromObject(FExecutionContext& Context, const UObject* InObject) const
+		{
+			return GetInterface()->GetGraphFromObject(Context, *this, InObject);
+		}
+
+	protected:
+		const IGraphFactory* GetInterface() const { return GetInterfaceTyped<IGraphFactory>(); }
+	};
+}
diff --git a/Engine/Plugins/Experimental/AnimNextAnimGraph/Source/AnimNextAnimGraph/Public/TraitInterfaces/IPayloadFactory.h b/Engine/Plugins/Experimental/AnimNextAnimGraph/Source/AnimNextAnimGraph/Public/TraitInterfaces/IPayloadFactory.h
new file mode 100644
--- /dev/null
+++ b/Engine/Plugins/Experimental/AnimNextAnimGraph/Source/AnimNextAnimGraph/Public/TraitInterfaces/IPayloadFactory.h
@@ -0,0 +1,61 @@
+﻿// Copyright Epic Games, Inc. All Rights Reserved.
+
+#pragma once
+
+#include "CoreMinimal.h"
+#include "StructUtils/PropertyBag.h"
+#include "StructUtils/StructView.h"
+#include "TraitCore/ITraitInterface.h"
+#include "TraitCore/TraitBinding.h"
+
+class UAnimNextAnimationGraph;
+struct FAnimNextDataInterfacePayload;
+
+namespace UE::AnimNext
+{
+
+/**
+ * IPayloadFactory
+ *
+ * This interface translates an object and a UAnimNextAnimationGraph into payloads for binding to a graph instance
+ */
+struct ANIMNEXTANIMGRAPH_API IPayloadFactory : ITraitInterface
+{
+	DECLARE_ANIM_TRAIT_INTERFACE(IPayloadFactory, 0x766ce644)
+
+
+	// Gets set of payloads for binding to data interfaces from an object & graph
+	// @param	InObject					The object (e.g. AnimSequence). If this is a UAnimNextAnimationGraph, it should be the same object as InGraph
+	// @param	InGraph						The graph to be used.
+	// @param	InOutGraphPayload			Payload that can be used to communicate with a graph instance
+	virtual void GetPayloadFromGraph(FExecutionContext& Context, const TTraitBinding<IPayloadFactory>& Binding, const UObject* InObject, const UAnimNextAnimationGraph* InGraph, FAnimNextDataInterfacePayload& InOutPayload) const;
+
+	// Gets set of payloads for binding to data interfaces from an object & graph
+	// Falls back to using user settings if no interface is found in the current stack 
+	// @param	InObject					The object (e.g. AnimSequence). If this is a UAnimNextAnimationGraph, it should be the same object as InGraph
+	// @param	InGraph						The graph to be used.
+	// @param	InOutGraphPayload			Payload that can be used to communicate with a graph instance
+	static void GetPayloadForGraphWithFallback(FExecutionContext& Context, const FTraitBinding& Binding, const UObject* InObject, const UAnimNextAnimationGraph* InGraph, FAnimNextDataInterfacePayload& InOutPayload);
+
+#if WITH_EDITOR
+	virtual const FText& GetDisplayName() const override;
+	virtual const FText& GetDisplayShortName() const override;
+#endif // WITH_EDITOR
+};
+
+/**
+ * Specialization for trait binding.
+ */
+template<>
+struct TTraitBinding<IPayloadFactory> : FTraitBinding
+{
+	// @see IPayloadFactory::GetPayloadsFromGraph
+	void GetPayloadFromGraph(FExecutionContext& Context, const UObject* InObject, const UAnimNextAnimationGraph* InGraph, FAnimNextDataInterfacePayload& InOutPayload) const
+	{
+		GetInterface()->GetPayloadFromGraph(Context, *this, InObject, InGraph, InOutPayload);
+	}
+
+protected:
+	const IPayloadFactory* GetInterface() const { return GetInterfaceTyped<IPayloadFactory>(); }
+};
+}
diff --git a/Engine/Plugins/Experimental/AnimNextAnimGraph/Source/AnimNextAnimGraphEditor/AnimNextAnimGraphEditor.Build.cs b/Engine/Plugins/Experimental/AnimNextAnimGraph/Source/AnimNextAnimGraphEditor/AnimNextAnimGraphEditor.Build.cs
new file mode 100644
--- /dev/null
+++ b/Engine/Plugins/Experimental/AnimNextAnimGraph/Source/AnimNextAnimGraphEditor/AnimNextAnimGraphEditor.Build.cs
@@ -0,0 +1,20 @@
+// Copyright Epic Games, Inc. All Rights Reserved.
+
+namespace UnrealBuildTool.Rules
+{
+	public class AnimNextAnimGraphEditor : ModuleRules
+	{
+		public AnimNextAnimGraphEditor(ReadOnlyTargetRules Target) : base(Target)
+		{
+			PublicDependencyModuleNames.AddRange(
+				new string[]
+				{
+					"Core",
+					"CoreUObject",
+					"AnimNextAnimGraph",
+					"Settings",
+				}
+			);
+		}
+	}
+}
\ No newline at end of file
diff --git a/Engine/Plugins/Experimental/AnimNextAnimGraph/Source/AnimNextAnimGraphEditor/Private/AnimNextAnimGraphEditorModule.cpp b/Engine/Plugins/Experimental/AnimNextAnimGraph/Source/AnimNextAnimGraphEditor/Private/AnimNextAnimGraphEditorModule.cpp
new file mode 100644
--- /dev/null
+++ b/Engine/Plugins/Experimental/AnimNextAnimGraph/Source/AnimNextAnimGraphEditor/Private/AnimNextAnimGraphEditorModule.cpp
@@ -0,0 +1,40 @@
+// Copyright Epic Games, Inc. All Rights Reserved.
+
+#include "Modules/ModuleInterface.h"
+#include "Modules/ModuleManager.h"
+#include "ISettingsModule.h"
+#include "AnimNextAnimGraphSettings.h"
+
+#define LOCTEXT_NAMESPACE "FAnimNextAnimGraphEditorModule"
+
+namespace UE::AnimNext::Editor
+{
+
+class FAnimNextAnimGraphEditorModule : public IModuleInterface
+{
+	virtual void StartupModule() override
+	{
+		// Register settings
+		ISettingsModule& SettingsModule = FModuleManager::LoadModuleChecked<ISettingsModule>("Settings");
+		SettingsModule.RegisterSettings("Project", "Plugins", "AnimNextAnimGraph",
+			LOCTEXT("SettingsName", "AnimNext Anim Graph"),
+			LOCTEXT("SettingsDescription", "Configure options for AnimNext animation graphs."),
+			GetMutableDefault<UAnimNextAnimGraphSettings>());
+	}
+
+	virtual void ShutdownModule() override
+	{
+		ISettingsModule* SettingsModule = FModuleManager::GetModulePtr<ISettingsModule>("Settings");
+		if (SettingsModule)
+		{
+			SettingsModule->UnregisterSettings("Project", "Plugins", "AnimNextAnimGraph");
+		}
+	}
+};
+
+}
+
+#undef LOCTEXT_NAMESPACE
+
+IMPLEMENT_MODULE(UE::AnimNext::Editor::FAnimNextAnimGraphEditorModule, AnimNextAnimGraphEditor);
+
diff --git a/Engine/Plugins/Experimental/AnimNextChooser/AnimNextChooser.uplugin b/Engine/Plugins/Experimental/AnimNextChooser/AnimNextChooser.uplugin
--- a/Engine/Plugins/Experimental/AnimNextChooser/AnimNextChooser.uplugin
+++ b/Engine/Plugins/Experimental/AnimNextChooser/AnimNextChooser.uplugin
@@ -21,6 +21,10 @@
 			"Name": "AnimNext",
 			"Enabled": true
 		},
+		{
+			"Name": "AnimNextAnimGraph",
+			"Enabled": true
+		},
 		{
 			"Name": "Chooser",
 			"Enabled": true
diff --git a/Engine/Plugins/Experimental/AnimNextChooser/Source/AnimNextChooser/AnimNextChooser.Build.cs b/Engine/Plugins/Experimental/AnimNextChooser/Source/AnimNextChooser/AnimNextChooser.Build.cs
--- a/Engine/Plugins/Experimental/AnimNextChooser/Source/AnimNextChooser/AnimNextChooser.Build.cs
+++ b/Engine/Plugins/Experimental/AnimNextChooser/Source/AnimNextChooser/AnimNextChooser.Build.cs
@@ -6,7 +6,7 @@ namespace UnrealBuildTool.Rules
 	{
 		public AnimNextChooser(ReadOnlyTargetRules Target) : base(Target)
 		{
-			PublicDependencyModuleNames.AddRange(
+			PrivateDependencyModuleNames.AddRange(
 				new string[]
 				{
 					"Core",
@@ -16,6 +16,7 @@ namespace UnrealBuildTool.Rules
 					"Engine",
 					"AnimNext",
 					"Chooser",
+					"AnimNextAnimGraph",
 				}
 			);
 		}
diff --git a/Engine/Plugins/Experimental/AnimNextChooser/Source/AnimNextChooser/Private/MakeGraphFromChooser.cpp b/Engine/Plugins/Experimental/AnimNextChooser/Source/AnimNextChooser/Private/MakeGraphFromChooser.cpp
new file mode 100644
--- /dev/null
+++ b/Engine/Plugins/Experimental/AnimNextChooser/Source/AnimNextChooser/Private/MakeGraphFromChooser.cpp
@@ -0,0 +1,48 @@
+﻿// Copyright Epic Games, Inc. All Rights Reserved.
+
+#include "MakeGraphFromChooser.h"
+
+#include "Chooser.h"
+#include "IObjectChooser.h"
+#include "Graph/AnimNextAnimationGraph.h"
+#include "TraitCore/ExecutionContext.h"
+
+namespace UE::AnimNext
+{
+	AUTO_REGISTER_ANIM_TRAIT(FMakeGraphFromChooserTrait)
+
+	// Trait implementation boilerplate
+	#define TRAIT_INTERFACE_ENUMERATOR(GeneratorMacro) \
+		GeneratorMacro(IGraphFactory) \
+
+		GENERATE_ANIM_TRAIT_IMPLEMENTATION(FMakeGraphFromChooserTrait, TRAIT_INTERFACE_ENUMERATOR, NULL_ANIM_TRAIT_INTERFACE_ENUMERATOR, NULL_ANIM_TRAIT_EVENT_ENUMERATOR)
+	#undef TRAIT_INTERFACE_ENUMERATOR
+
+	const UAnimNextAnimationGraph* FMakeGraphFromChooserTrait::GetGraphFromObject(FExecutionContext& Context, const TTraitBinding<IGraphFactory>& Binding, const UObject* InObject) const
+	{
+		if(InObject == nullptr)
+		{
+			return nullptr;
+		}
+
+		const FSharedData* SharedData = Binding.GetSharedData<FSharedData>();
+		const UChooserTable* Chooser = SharedData->GetChooser(Binding);
+		if (Chooser == nullptr)
+		{
+			return nullptr;
+		}
+
+		FChooserEvaluationContext ChooserContext;
+		ChooserContext.AddObjectParam(const_cast<UObject*>(InObject));
+
+		const UAnimNextAnimationGraph* AnimationGraph = nullptr;
+		UChooserTable::EvaluateChooser(ChooserContext, Chooser, FObjectChooserBase::FObjectChooserIteratorCallback::CreateLambda([&AnimationGraph](UObject* InResult)
+		{
+			AnimationGraph = Cast<UAnimNextAnimationGraph>(InResult);
+			return FObjectChooserBase::EIteratorStatus::Stop;
+		}));
+
+		return AnimationGraph;
+	}
+
+}
diff --git a/Engine/Plugins/Experimental/AnimNextChooser/Source/AnimNextChooser/Private/MakeGraphFromChooser.h b/Engine/Plugins/Experimental/AnimNextChooser/Source/AnimNextChooser/Private/MakeGraphFromChooser.h
new file mode 100644
--- /dev/null
+++ b/Engine/Plugins/Experimental/AnimNextChooser/Source/AnimNextChooser/Private/MakeGraphFromChooser.h
@@ -0,0 +1,50 @@
+﻿// Copyright Epic Games, Inc. All Rights Reserved.
+
+#pragma once
+
+#include "CoreMinimal.h"
+#include "TraitCore/Trait.h"
+#include "TraitInterfaces/IGraphFactory.h"
+
+#include "MakeGraphFromChooser.generated.h"
+
+class UChooserTable;
+
+USTRUCT(meta = (DisplayName = "Make Graph From Chooser"))
+struct FAnimNextMakeGraphFromChooserTraitSharedData : public FAnimNextTraitSharedData
+{
+	GENERATED_BODY()
+
+	// The chooser to use to pick an animation graph 
+	UPROPERTY(EditAnywhere, Category = "Chooser")
+	TObjectPtr<UChooserTable> Chooser;
+
+	// Latent pin support boilerplate
+	#define TRAIT_LATENT_PROPERTIES_ENUMERATOR(GeneratorMacro) \
+		GeneratorMacro(Chooser) \
+
+	GENERATE_TRAIT_LATENT_PROPERTIES(FAnimNextMakeGraphFromChooserTraitSharedData, TRAIT_LATENT_PROPERTIES_ENUMERATOR)
+	#undef TRAIT_LATENT_PROPERTIES_ENUMERATOR
+};
+
+namespace UE::AnimNext
+{
+	/**
+	 * FMakeGraphFromChooserTrait
+	 *
+	 * A trait that uses a chooser to make an AnimNext animation graph
+	 */
+	struct FMakeGraphFromChooserTrait : FAdditiveTrait, IGraphFactory
+	{
+		DECLARE_ANIM_TRAIT(FMakeGraphFromChooserTrait, 0xe8c28546, FAdditiveTrait)
+
+		using FSharedData = FAnimNextMakeGraphFromChooserTraitSharedData;
+
+		struct FInstanceData : FTrait::FInstanceData
+		{
+		};
+
+		// IGraphFactory impl
+		virtual const UAnimNextAnimationGraph* GetGraphFromObject(FExecutionContext& Context, const TTraitBinding<IGraphFactory>& Binding, const UObject* InObject) const;
+	};
+}
diff --git a/Engine/Plugins/Experimental/AnimNextChooser/Source/AnimNextChooser/Private/Module.cpp b/Engine/Plugins/Experimental/AnimNextChooser/Source/AnimNextChooser/Private/Module.cpp
--- a/Engine/Plugins/Experimental/AnimNextChooser/Source/AnimNextChooser/Private/Module.cpp
+++ b/Engine/Plugins/Experimental/AnimNextChooser/Source/AnimNextChooser/Private/Module.cpp
@@ -19,6 +19,12 @@ class FModule : public IModuleInterface
 public:
 	virtual void StartupModule() override
 	{
+		static TPair<UClass*, FRigVMRegistry::ERegisterObjectOperation> const AllowedObjectTypes[] =
+		{
+			{ UChooserTable::StaticClass(), FRigVMRegistry::ERegisterObjectOperation::Class },
+		};
+
+		FRigVMRegistry::Get().RegisterObjectTypes(AllowedObjectTypes);
 	}
 
 	virtual void ShutdownModule() override
diff --git a/Engine/Plugins/Experimental/AnimNextPoseSearch/Source/AnimNextPoseSearch/Private/MotionMatchingTrait.cpp b/Engine/Plugins/Experimental/AnimNextPoseSearch/Source/AnimNextPoseSearch/Private/MotionMatchingTrait.cpp
--- a/Engine/Plugins/Experimental/AnimNextPoseSearch/Source/AnimNextPoseSearch/Private/MotionMatchingTrait.cpp
+++ b/Engine/Plugins/Experimental/AnimNextPoseSearch/Source/AnimNextPoseSearch/Private/MotionMatchingTrait.cpp
@@ -1,13 +1,17 @@
 // Copyright Epic Games, Inc. All Rights Reserved.
 
 #include "MotionMatchingTrait.h"
+
+#include "AnimNextAnimGraphSettings.h"
 #include "Animation/AnimSequence.h"
 #include "DataInterface/DataInterfaceStructAdapter.h"
 #include "IPoseHistory.h"
-#include "PlayAnim/PlayAnimRequest.h"
+#include "GraphInterfaces/AnimNextNativeDataInterface_AnimSequencePlayer.h"
 #include "PoseSearch/PoseSearchDatabase.h"
 #include "PoseSearch/PoseSearchDerivedData.h"
 #include "TraitInterfaces/IBlendStack.h"
+#include "TraitInterfaces/IGraphFactory.h"
+#include "TraitInterfaces/ITimeline.h"
 
 namespace UE::AnimNext
 {
@@ -152,19 +156,28 @@ namespace UE::AnimNext
 			{
 				if (const FPoseSearchDatabaseAnimationAssetBase* DatabaseAnimationAssetBase = MotionMatchingState.CurrentSearchResult.Database->GetDatabaseAnimationAsset<FPoseSearchDatabaseAnimationAssetBase>(*SearchIndexAsset))
 				{
-					if (UAnimSequence* AnimSequence = Cast<UAnimSequence>(DatabaseAnimationAssetBase->GetAnimationAsset()))
+					const UObject* AnimationAsset = DatabaseAnimationAssetBase->GetAnimationAsset();
+					
+					UE::AnimNext::IBlendStack::FGraphRequest NewGraphRequest;
+					NewGraphRequest.BlendArgs = SharedData->BlendArguments;
+
+					const UAnimNextAnimationGraph* AnimationGraph = IGraphFactory::GetGraphFromObjectWithFallback(Context, Binding, AnimationAsset);
+					if(AnimationGraph != nullptr)
 					{
-						UE::AnimNext::IBlendStack::FGraphRequest NewGraphRequest;
-						NewGraphRequest.BlendArgs = SharedData->BlendArguments;
-						NewGraphRequest.AnimationGraph = SharedData->AnimationGraph;
-
-						FInstancedStruct& NewPayload = NewGraphRequest.Payloads.AddDefaulted_GetRef();
-						NewPayload.InitializeAs<FAnimNextPlayAnimPayload>();
-						FAnimNextPlayAnimPayload& NewPlayAnimPayload = NewPayload.GetMutable<FAnimNextPlayAnimPayload>();
-						NewPlayAnimPayload.AnimationObject = AnimSequence;
-						NewPlayAnimPayload.PlayRate = MotionMatchingState.WantedPlayRate;
-						NewPlayAnimPayload.StartPosition = MotionMatchingState.CurrentSearchResult.AssetTime;
-						NewPlayAnimPayload.IsLooping = SearchIndexAsset->IsLooping();
+						// Create a default 'native data interface' payload via which we expect to communicate with the graph 
+						FInstancedStruct NewPayload;
+						NewPayload.InitializeAs<FAnimNextNativeDataInterface_AnimSequencePlayer>();
+						FAnimNextNativeDataInterface_AnimSequencePlayer& NewNativeInterface = NewPayload.GetMutable<FAnimNextNativeDataInterface_AnimSequencePlayer>();
+						NewNativeInterface.AnimSequence = Cast<UAnimSequence>(AnimationAsset);
+						NewNativeInterface.PlayRate = MotionMatchingState.WantedPlayRate;
+						NewNativeInterface.StartPosition = MotionMatchingState.CurrentSearchResult.AssetTime;
+						NewNativeInterface.Loop = SearchIndexAsset->IsLooping();
+
+						NewGraphRequest.GraphPayload.AddNative(MoveTemp(NewPayload));
+
+						IPayloadFactory::GetPayloadForGraphWithFallback(Context, Binding, AnimationAsset, AnimationGraph, NewGraphRequest.GraphPayload);
+
+						NewGraphRequest.AnimationGraph = AnimationGraph;
 
 						BlendStackBinding.PushGraph(Context, MoveTemp(NewGraphRequest));
 					}
diff --git a/Engine/Plugins/Experimental/AnimNextPoseSearch/Source/AnimNextPoseSearch/Private/MotionMatchingTrait.h b/Engine/Plugins/Experimental/AnimNextPoseSearch/Source/AnimNextPoseSearch/Private/MotionMatchingTrait.h
--- a/Engine/Plugins/Experimental/AnimNextPoseSearch/Source/AnimNextPoseSearch/Private/MotionMatchingTrait.h
+++ b/Engine/Plugins/Experimental/AnimNextPoseSearch/Source/AnimNextPoseSearch/Private/MotionMatchingTrait.h
@@ -14,10 +14,6 @@ struct FMotionMatchingTraitSharedData : public FAnimNextTraitSharedData
 {
 	GENERATED_BODY()
 
-	// Asset to use as a sub-graph
-	UPROPERTY(EditAnywhere, Category = "Default")
-	TObjectPtr<const UAnimNextAnimationGraph> AnimationGraph;
-
 	// Entry point in the Subgraph that we will use
 	UPROPERTY(EditAnywhere, Category = "Default", meta=(CustomWidget = "ParamName", AllowedParamType = "FAnimNextEntryPoint"))
 	FName EntryPoint;
@@ -61,7 +57,6 @@ struct FMotionMatchingTraitSharedData : public FAnimNextTraitSharedData
 
 	// Latent pin support boilerplate
 	#define TRAIT_LATENT_PROPERTIES_ENUMERATOR(GeneratorMacro) \
-		GeneratorMacro(AnimationGraph) \
 		GeneratorMacro(EntryPoint) \
 		GeneratorMacro(Databases) \
 		GeneratorMacro(BlendArguments) \
diff --git a/Engine/Plugins/Experimental/AnimNextStateTree/Source/AnimNextStateTree/AnimNextStateTree.Build.cs b/Engine/Plugins/Experimental/AnimNextStateTree/Source/AnimNextStateTree/AnimNextStateTree.Build.cs
--- a/Engine/Plugins/Experimental/AnimNextStateTree/Source/AnimNextStateTree/AnimNextStateTree.Build.cs
+++ b/Engine/Plugins/Experimental/AnimNextStateTree/Source/AnimNextStateTree/AnimNextStateTree.Build.cs
@@ -20,7 +20,6 @@ namespace UnrealBuildTool.Rules
 					"CoreUObject",
 					"StateTreeModule",
 					"Engine",
-					"Chooser",
 					"AnimNext", 
 				}
 			);
diff --git a/Engine/Plugins/Experimental/AnimNextStateTree/Source/AnimNextStateTree/Private/AnimNextStateTreeContext.cpp b/Engine/Plugins/Experimental/AnimNextStateTree/Source/AnimNextStateTree/Private/AnimNextStateTreeContext.cpp
--- a/Engine/Plugins/Experimental/AnimNextStateTree/Source/AnimNextStateTree/Private/AnimNextStateTreeContext.cpp
+++ b/Engine/Plugins/Experimental/AnimNextStateTree/Source/AnimNextStateTree/Private/AnimNextStateTreeContext.cpp
@@ -2,47 +2,55 @@
 
 #include "AnimNextStateTreeContext.h"
 
-#include "DataInterface/DataInterfaceStructAdapter.h"
+#include "AnimNextAnimGraphSettings.h"
 #include "TraitCore/TraitStackBinding.h"
 #include "TraitInterfaces/IBlendStack.h"
+#include "TraitInterfaces/IGraphFactory.h"
 #include "TraitInterfaces/ISmoothBlend.h"
 #include "TraitInterfaces/ITimeline.h"
 
-bool FAnimNextStateTreeTraitContext::PushAnimationGraphOntoBlendStack(TNonNullPtr<UAnimNextAnimationGraph> InAnimationGraph, const FAlphaBlendArgs& InBlendArguments, TArrayView<const FInstancedStruct> PayloadStructs) const
-{		
-	if (Binding && Context)
+bool FAnimNextStateTreeTraitContext::PushAssetOntoBlendStack(TNonNullPtr<UObject> InAsset, const FAlphaBlendArgs& InBlendArguments, FAnimNextDataInterfacePayload&& InPayload) const
+{
+	using namespace UE::AnimNext;
+
+	if (Binding == nullptr || Context == nullptr)
 	{
-		UE::AnimNext::TTraitBinding<UE::AnimNext::IBlendStack> BlendStackBinding;
-		if(Binding->GetInterface<UE::AnimNext::IBlendStack>(BlendStackBinding))
-		{
-			UE::AnimNext::IBlendStack::FGraphRequestPtr ActiveGraphRequest;
-			BlendStackBinding.GetActiveGraph(*Context, ActiveGraphRequest);
-
-			if (ActiveGraphRequest && ActiveGraphRequest->AnimationGraph == InAnimationGraph)
-			{
-				return true;
-			}
-			
-			UE::AnimNext::IBlendStack::FGraphRequest NewGraphRequest;
-			NewGraphRequest.BlendArgs = InBlendArguments;
-			NewGraphRequest.AnimationGraph = InAnimationGraph.Get();
-			NewGraphRequest.Payloads.Append(PayloadStructs);
-			
-			BlendStackBinding.PushGraph(*Context, MoveTemp(NewGraphRequest));
-
-			return true;
-		}
+		return false;
+	}
+
+	TTraitBinding<IBlendStack> BlendStackBinding;
+	if(!Binding->GetStackInterface<IBlendStack>(BlendStackBinding))
+	{
+		return false;
+	}
+
+	// Get a graph to instantiate
+	const UAnimNextAnimationGraph* AnimationGraph = IGraphFactory::GetGraphFromObjectWithFallback(*Context, *Binding, InAsset);
+	if(AnimationGraph == nullptr)
+	{
+		return false;
 	}
 
-	return false;
+	// Get a payload to instantiate to control the graph
+	IPayloadFactory::GetPayloadForGraphWithFallback(*Context, *Binding, InAsset, AnimationGraph, InPayload);
+
+	IBlendStack::FGraphRequest NewGraphRequest;
+	NewGraphRequest.BlendArgs = InBlendArguments;
+	NewGraphRequest.FactoryObject = InAsset;
+	NewGraphRequest.AnimationGraph = AnimationGraph;
+	NewGraphRequest.GraphPayload = MoveTemp(InPayload);
+
+	BlendStackBinding.PushGraph(*Context, MoveTemp(NewGraphRequest));
+
+	return true;
 }
 
-float FAnimNextStateTreeTraitContext::QueryPlaybackRatio(TNonNullPtr<UAnimNextAnimationGraph> InAnimationGraph) const
+float FAnimNextStateTreeTraitContext::QueryPlaybackRatio(TNonNullPtr<UObject> InAsset) const
 {
 	if (Binding && Context)
 	{
 		UE::AnimNext::TTraitBinding<UE::AnimNext::ITimeline> Timeline;
-		if (Binding->GetInterface<UE::AnimNext::ITimeline>(Timeline))
+		if (Binding->GetStackInterface<UE::AnimNext::ITimeline>(Timeline))
 		{
 			UE::AnimNext::FTimelineProgress Progress = Timeline.GetProgress(*Context);
 
diff --git a/Engine/Plugins/Experimental/AnimNextStateTree/Source/AnimNextStateTree/Private/AnimNextStateTreeGraphInstanceTask.cpp b/Engine/Plugins/Experimental/AnimNextStateTree/Source/AnimNextStateTree/Private/AnimNextStateTreeGraphInstanceTask.cpp
--- a/Engine/Plugins/Experimental/AnimNextStateTree/Source/AnimNextStateTree/Private/AnimNextStateTreeGraphInstanceTask.cpp
+++ b/Engine/Plugins/Experimental/AnimNextStateTree/Source/AnimNextStateTree/Private/AnimNextStateTreeGraphInstanceTask.cpp
@@ -1,12 +1,16 @@
 // Copyright Epic Games, Inc. All Rights Reserved.
 
 #include "Tasks/AnimNextStateTreeGraphInstanceTask.h"
+
+#include "AnimNextAnimGraphSettings.h"
 #include "StateTreeLinker.h"
 #include "StateTreeExecutionContext.h"
-#include "DataInterface/DataInterfaceStructAdapter.h"
+#include "AnimNextDataInterfacePayload.h"
 
 FAnimNextStateTreeGraphInstanceTask::FAnimNextStateTreeGraphInstanceTask()
 {
+	// Re-selecting the same state should not cause a re-trigger of EnterState()
+	bShouldStateChangeOnReselect = false;
 }
 
 bool FAnimNextStateTreeGraphInstanceTask::Link(FStateTreeLinker& Linker)
@@ -20,12 +24,17 @@ EStateTreeRunStatus FAnimNextStateTreeGraphInstanceTask::EnterState(FStateTreeEx
 	FAnimNextStateTreeTraitContext& ExecContext = Context.GetExternalData(TraitContextHandle);
 	FInstanceDataType& InstanceData = Context.GetInstanceData(*this);
 
-	FInstancedStruct Payload;
-	Payload.InitializeAs(Context.GetMutableInstanceData()->GetMutableStorage().GetMutableGlobalParameters().GetScriptStruct(), Context.GetMutableInstanceData()->GetMutableStorage().GetMutableGlobalParameters().GetMemory());
+	// Reference task's default payload
+	FAnimNextDataInterfacePayload Payload;
+	Payload.AddStructView(InstanceData.Payload.GetMutableValue());
+
+	FInstancedStruct NewNativePayload;
+	NewNativePayload.InitializeAs(Context.GetMutableInstanceData()->GetMutableStorage().GetMutableGlobalParameters().GetScriptStruct(), Context.GetMutableInstanceData()->GetMutableStorage().GetMutableGlobalParameters().GetMemory());
+	Payload.AddNative(MoveTemp(NewNativePayload));
 
-	if(ExecContext.PushAnimationGraphOntoBlendStack(InstanceData.AnimationGraph.Get(), InstanceData.BlendOptions, MakeConstArrayView({Payload})))
+	if (ExecContext.PushAssetOntoBlendStack(InstanceData.Asset.Get(), InstanceData.BlendOptions, MoveTemp(Payload)))
 	{
-		return EStateTreeRunStatus::Running;	
+		return EStateTreeRunStatus::Running;
 	}
 	
 	return EStateTreeRunStatus::Failed;
@@ -38,7 +47,7 @@ EStateTreeRunStatus FAnimNextStateTreeGraphInstanceTask::Tick(FStateTreeExecutio
 	{
 		FAnimNextStateTreeTraitContext& ExecContext = Context.GetExternalData(TraitContextHandle);
 
-		InstanceData.PlaybackRatio = ExecContext.QueryPlaybackRatio(InstanceData.AnimationGraph.Get());
+		InstanceData.PlaybackRatio = ExecContext.QueryPlaybackRatio(InstanceData.Asset.Get());
 
 		return EStateTreeRunStatus::Running;
 	}
@@ -50,3 +59,22 @@ void FAnimNextStateTreeGraphInstanceTask::ExitState(FStateTreeExecutionContext&
 {
 	FAnimNextStateTreeTaskBase::ExitState(Context, Transition);
 }
+
+#if WITH_EDITOR
+void FAnimNextStateTreeGraphInstanceTask::PostEditInstanceDataChangeChainProperty(const FPropertyChangedChainEvent& PropertyChangedEvent, FStateTreeDataView InstanceDataView)
+{
+	const FProperty* Property = PropertyChangedEvent.Property;
+	if(Property == nullptr)
+	{
+		return;
+	}
+
+	if(Property->GetFName() == GET_MEMBER_NAME_CHECKED(FAnimNextGraphInstanceTaskInstanceData, Asset))
+	{
+		// Repopulate payload property bag if the selected asset changes
+		FAnimNextGraphInstanceTaskInstanceData& InstanceData = InstanceDataView.GetMutable<FAnimNextGraphInstanceTaskInstanceData>();
+		const UAnimNextAnimationGraph* AnimationGraph = GetDefault<UAnimNextAnimGraphSettings>()->GetGraphFromObject(InstanceData.Asset);
+		GetDefault<UAnimNextAnimGraphSettings>()->GetNonNativePayloadFromGraph(InstanceData.Asset, AnimationGraph, InstanceData.Payload);
+	}
+}
+#endif
\ No newline at end of file
diff --git a/Engine/Plugins/Experimental/AnimNextStateTree/Source/AnimNextStateTree/Private/AnimStateTreeTrait.cpp b/Engine/Plugins/Experimental/AnimNextStateTree/Source/AnimNextStateTree/Private/AnimStateTreeTrait.cpp
--- a/Engine/Plugins/Experimental/AnimNextStateTree/Source/AnimNextStateTree/Private/AnimStateTreeTrait.cpp
+++ b/Engine/Plugins/Experimental/AnimNextStateTree/Source/AnimNextStateTree/Private/AnimStateTreeTrait.cpp
@@ -45,7 +45,7 @@ namespace UE::AnimNext
 			UObject* Owner = GetTransientPackage();
 			FStateTreeExecutionContext StateTreeExecutionContext(*Owner, *InstanceData->StateTree, InstanceData->InstanceData);
 
-			FAnimNextStateTreeTraitContext TraitContext(Context, Binding.GetStack());
+			FAnimNextStateTreeTraitContext TraitContext(Context, &Binding);
 			StateTreeExecutionContext.SetContextDataByName(UStateTreeAnimNextSchema::AnimStateTreeExecutionContextName, FStateTreeDataView(FAnimNextStateTreeTraitContext::StaticStruct(), reinterpret_cast<uint8*>(&TraitContext)));
 
 			FAnimNextGraphInstance* OwnerGraphInstance = &Binding.GetTraitPtr().GetNodeInstance()->GetOwner();
@@ -107,7 +107,7 @@ namespace UE::AnimNext
 			FAnimNextGraphInstance* OwnerGraphInstance = &Binding.GetTraitPtr().GetNodeInstance()->GetOwner();
 			if(OwnerGraphInstance && InstanceData->StateTree->GetDefaultParameters().GetPropertyBagStruct() == OwnerGraphInstance->GetVariables().GetPropertyBagStruct())
 			{
-				FAnimNextStateTreeTraitContext TraitContext(Context, Binding.GetStack());
+				FAnimNextStateTreeTraitContext TraitContext(Context, &Binding);
 				StateTreeExecutionContext.SetContextDataByName(UStateTreeAnimNextSchema::AnimStateTreeExecutionContextName, FStateTreeDataView(FAnimNextStateTreeTraitContext::StaticStruct(), reinterpret_cast<uint8*>(&TraitContext)));
 
 				const FInstancedPropertyBag& StateTreeParameters = InstanceData->StateTree->GetDefaultParameters();
diff --git a/Engine/Plugins/Experimental/AnimNextStateTree/Source/AnimNextStateTree/Public/AnimNextStateTreeContext.h b/Engine/Plugins/Experimental/AnimNextStateTree/Source/AnimNextStateTree/Public/AnimNextStateTreeContext.h
--- a/Engine/Plugins/Experimental/AnimNextStateTree/Source/AnimNextStateTree/Public/AnimNextStateTreeContext.h
+++ b/Engine/Plugins/Experimental/AnimNextStateTree/Source/AnimNextStateTree/Public/AnimNextStateTreeContext.h
@@ -9,13 +9,14 @@
 
 namespace UE::AnimNext
 {
-struct FTraitStackBinding;
+struct FTraitBinding;
 struct FExecutionContext;
 struct FStateTreeTrait;
 }
 
 class UAnimNextAnimationGraph;
 struct FAlphaBlendArgs;
+struct FAnimNextDataInterfacePayload;
 
 USTRUCT()
 struct FAnimNextStateTreeTraitContext
@@ -26,12 +27,12 @@ struct FAnimNextStateTreeTraitContext
 	
 	FAnimNextStateTreeTraitContext() {}
 
-	bool PushAnimationGraphOntoBlendStack(TNonNullPtr<UAnimNextAnimationGraph> InAnimationGraph, const FAlphaBlendArgs& InBlendArguments, const TArrayView<const FInstancedStruct> PayloadStructs) const;
-	float QueryPlaybackRatio(TNonNullPtr<UAnimNextAnimationGraph> InAnimationGraph) const;
+	bool PushAssetOntoBlendStack(TNonNullPtr<UObject> InAsset, const FAlphaBlendArgs& InBlendArguments, FAnimNextDataInterfacePayload&& InPayload) const;
+	float QueryPlaybackRatio(TNonNullPtr<UObject> InAsset) const;
 protected:
-	FAnimNextStateTreeTraitContext(UE::AnimNext::FExecutionContext& InContext, const UE::AnimNext::FTraitStackBinding* InBinding) : Context(&InContext), Binding(InBinding) {}
+	FAnimNextStateTreeTraitContext(UE::AnimNext::FExecutionContext& InContext, const UE::AnimNext::FTraitBinding* InBinding) : Context(&InContext), Binding(InBinding) {}
 protected:
 	UE::AnimNext::FExecutionContext* Context = nullptr;
-	const UE::AnimNext::FTraitStackBinding* Binding = nullptr;
+	const UE::AnimNext::FTraitBinding* Binding = nullptr;
 };
 
diff --git a/Engine/Plugins/Experimental/AnimNextStateTree/Source/AnimNextStateTree/Public/Tasks/AnimNextStateTreeGraphInstanceTask.h b/Engine/Plugins/Experimental/AnimNextStateTree/Source/AnimNextStateTree/Public/Tasks/AnimNextStateTreeGraphInstanceTask.h
--- a/Engine/Plugins/Experimental/AnimNextStateTree/Source/AnimNextStateTree/Public/Tasks/AnimNextStateTreeGraphInstanceTask.h
+++ b/Engine/Plugins/Experimental/AnimNextStateTree/Source/AnimNextStateTree/Public/Tasks/AnimNextStateTreeGraphInstanceTask.h
@@ -6,6 +6,7 @@
 #include "AnimNextStateTreeTypes.h"
 #include "Graph/AnimNextAnimationGraph.h"
 #include "AlphaBlend.h"
+#include "StructUtils/PropertyBag.h"
 
 #include "AnimNextStateTreeGraphInstanceTask.generated.h"
 
@@ -13,16 +14,24 @@ USTRUCT()
 struct ANIMNEXTSTATETREE_API FAnimNextGraphInstanceTaskInstanceData
 {
 	GENERATED_BODY()
-	
-	UPROPERTY(EditAnywhere, Category = Animation)
-	TObjectPtr<UAnimNextAnimationGraph> AnimationGraph = nullptr;
 
+	// The asset to instantiate
+	UPROPERTY(EditAnywhere, Category = Animation, meta=(GetAllowedClasses="/Script/AnimNextAnimGraph.AnimNextAnimGraphSettings:GetAllowedAssetClasses"))
+	TObjectPtr<UObject> Asset = nullptr;
+
+	// The payload to use for the asset when instanced
+	UPROPERTY(EditAnywhere, Category = Input, meta = (Optional, FixedLayout))
+	FInstancedPropertyBag Payload;
+
+	// Blend options for when the state is pushed
 	UPROPERTY(EditAnywhere, Category = Animation)
 	FAlphaBlendArgs BlendOptions;
 
+	// Whether this task should continue to tick once state is entered
 	UPROPERTY(EditAnywhere, Category = Animation)
 	bool bContinueTicking = true;
 
+	// Current playback ratio (debug)
 	UPROPERTY(VisibleAnywhere, Category = Animation)
 	float PlaybackRatio = 1.0f;
 };
@@ -43,7 +52,9 @@ protected:
 	virtual EStateTreeRunStatus EnterState(FStateTreeExecutionContext& Context, const FStateTreeTransitionResult& Transition) const override;
 	virtual EStateTreeRunStatus Tick(FStateTreeExecutionContext& Context, const float DeltaTime) const override;
 	virtual void ExitState(FStateTreeExecutionContext& Context, const FStateTreeTransitionResult& Transition) const override;
-
+#if WITH_EDITOR
+	virtual void PostEditInstanceDataChangeChainProperty(const FPropertyChangedChainEvent& PropertyChangedEvent, FStateTreeDataView InstanceDataView) override;
+#endif
 public:
 	TStateTreeExternalDataHandle<FAnimNextStateTreeTraitContext> TraitContextHandle;
 };
diff --git a/Engine/Plugins/Experimental/AnimNextStateTree/Source/AnimNextStateTreeEditor/AnimNextStateTreeEditor.Build.cs b/Engine/Plugins/Experimental/AnimNextStateTree/Source/AnimNextStateTreeEditor/AnimNextStateTreeEditor.Build.cs
--- a/Engine/Plugins/Experimental/AnimNextStateTree/Source/AnimNextStateTreeEditor/AnimNextStateTreeEditor.Build.cs
+++ b/Engine/Plugins/Experimental/AnimNextStateTree/Source/AnimNextStateTreeEditor/AnimNextStateTreeEditor.Build.cs
@@ -18,7 +18,9 @@ namespace UnrealBuildTool.Rules
 				    "UnrealEd",
 				    "AssetDefinition",
 				    "AnimNextAnimGraph",
-				    "RigVM"
+				    "RigVM",
+				    "StructUtilsEditor",
+				    "PropertyEditor",
 			    });
 
 			PublicDependencyModuleNames.AddRange(
