[Backout] - CL37017466
Original CL Desc
-----------------------------------------------------------------
AnimNext: Factories

Added to new trait interfaces: IGraphFactory and IPayloadFactory. Both use fallbacks to project-defined (with engine defaults) mappings between asset classes (e.g. UAnimSequence) and anim next graphs. The mapping is managed by UAnimNextAnimGraphSettings.

The idea is that anything that wants to implement the hosting of some subgraph can use the IGraphFactory to not be limited to only handling UAnimNextAnimationGraph assets, but can use it to make UAnimNextAnimationGraphs to instantiate from some other object key. The object key can be optionally injected as a variable into the facotry graph via the user mapping.

IGraphFactory has two implementations, FMakeGraphFromAssetTrait (which just implements the fallback to UAnimNextAnimGraphSettings) and FMakeGraphFromChooserTrait (which allows a chooser to decide the graph to instantiate).
IPayloadFactory does not have any implementations yet (so fallback is always used), but the intent is that we can implement this in future to supply per-variable overrides via traits that expose values on pins.
Updated the following traits to use the new interfaces: FBlendStackTrait, FBlendStackRequesterTrait, FMotionMatchingTrait, FPlayAnimSlotTrait.

Added the concept of a 'native interface' to all AnimNext RigVM assets. This is a structure that is intended to be used for native C++ communication with AnimNext asset instances via struct-adapters. This is compiled into all 'implemented interface' definitions.

Updated StateTree to use the new interfaces. Updated state tree's task to allow any object that is supported by the IGraphFactory interface. Also updated the UI to allow selecting and parameterizing the instanced graph via non-native interfaces.

Added one plugin-level asset to implement an anim sequence player: /AnimNextAnimGraph/FactoryGraphs/AG_AnimSequencePlayer.

Removed AnimNextAnimGraph and AnimNextStateTree's dependency on Chooser module.

#virtualized

diff --git a/Engine/Plugins/Experimental/AnimNext/Source/AnimNext/Internal/DataInterface/AnimNextDataInterface.h b/Engine/Plugins/Experimental/AnimNext/Source/AnimNext/Internal/DataInterface/AnimNextDataInterface.h
--- a/Engine/Plugins/Experimental/AnimNext/Source/AnimNext/Internal/DataInterface/AnimNextDataInterface.h
+++ b/Engine/Plugins/Experimental/AnimNext/Source/AnimNext/Internal/DataInterface/AnimNextDataInterface.h
@@ -7,7 +7,6 @@
 #include "AnimNextDataInterface.generated.h"
 
 class UAnimNextDataInterfaceFactory;
-class UAnimNextDataInterface;
 
 namespace UE::AnimNext::Tests
 {
@@ -25,14 +24,9 @@ struct FAnimNextImplementedDataInterface
 {
 	GENERATED_BODY()
 
-	// Reference to the implemented data interface asset
 	UPROPERTY()
 	TObjectPtr<const UAnimNextDataInterface> DataInterface;
 
-	// Native struct used for communication with the data interface
-	UPROPERTY()
-	TObjectPtr<const UScriptStruct> NativeInterface;
-
 	// Index of the first variable that implements the interface
 	UPROPERTY()
 	int32 VariableIndex = INDEX_NONE;
diff --git a/Engine/Plugins/Experimental/AnimNext/Source/AnimNext/Private/DataInterface/DataInterfaceStructAdapter.cpp b/Engine/Plugins/Experimental/AnimNext/Source/AnimNext/Private/DataInterface/DataInterfaceStructAdapter.cpp
--- a/Engine/Plugins/Experimental/AnimNext/Source/AnimNext/Private/DataInterface/DataInterfaceStructAdapter.cpp
+++ b/Engine/Plugins/Experimental/AnimNext/Source/AnimNext/Private/DataInterface/DataInterfaceStructAdapter.cpp
@@ -14,13 +14,7 @@ const UAnimNextDataInterface* FDataInterfaceStructAdapter::GetDataInterface() co
 
 uint8* FDataInterfaceStructAdapter::GetMemoryForVariable(int32 InVariableIndex, FName InVariableName, const FProperty* InVariableProperty) const
 {
-	const UScriptStruct* Struct = StructView.GetScriptStruct();
-	if(Struct == nullptr)
-	{
-		return nullptr;
-	}
-
-	const FProperty* Property = Struct->FindPropertyByName(InVariableName);
+	const FProperty* Property = StructView.GetScriptStruct()->FindPropertyByName(InVariableName);
 	if(Property == nullptr)
 	{
 		return nullptr;
diff --git a/Engine/Plugins/Experimental/AnimNext/Source/AnimNext/Private/Graph/AnimNextGraphInstance.cpp b/Engine/Plugins/Experimental/AnimNext/Source/AnimNext/Private/Graph/AnimNextGraphInstance.cpp
--- a/Engine/Plugins/Experimental/AnimNext/Source/AnimNext/Private/Graph/AnimNextGraphInstance.cpp
+++ b/Engine/Plugins/Experimental/AnimNext/Source/AnimNext/Private/Graph/AnimNextGraphInstance.cpp
@@ -4,7 +4,6 @@
 
 #include "AnimNextStats.h"
 #include "DataInterface/AnimNextDataInterfaceHost.h"
-#include "DataInterface/DataInterfaceStructAdapter.h"
 #include "Graph/AnimNextAnimationGraph.h"
 #include "TraitCore/ExecutionContext.h"
 #include "Graph/GC_GraphInstanceComponent.h"
@@ -286,43 +285,6 @@ bool FAnimNextGraphInstance::BindToHostHelper(const HostType& InHost, bool bInAu
 	return bPublicVariablesBound;
 }
 
-void FAnimNextGraphInstance::BindPublicVariables(TArrayView<FStructView> InHostStructs)
-{
-	using namespace UE::AnimNext;
-
-	const UAnimNextAnimationGraph* AnimationGraph = GetAnimationGraph();
-	if (AnimationGraph == nullptr)
-	{
-		return;
-	}
-
-	const int32 NumStructs = InHostStructs.Num();
-	if(NumStructs > 0)
-	{
-		// stack allocating uninitialized arrays to contain FDataInterfaceStructAdapter(s)
-		TArrayView<FDataInterfaceStructAdapter> Adapters((FDataInterfaceStructAdapter*)FMemory_Alloca(NumStructs * sizeof(FDataInterfaceStructAdapter)), NumStructs);
-		TArrayView<IDataInterfaceHost*> PayloadVariableBindings((IDataInterfaceHost**)FMemory_Alloca(NumStructs * sizeof(IDataInterfaceHost*)), NumStructs);
-
-		// initializing the FDataInterfaceStructAdapter(s) with their respective HostStruct
-		for (int32 StructIndex = 0; StructIndex < NumStructs; ++StructIndex)
-		{
-			FStructView& StructView = InHostStructs[StructIndex];
-			PayloadVariableBindings[StructIndex] = new (&Adapters[StructIndex]) FDataInterfaceStructAdapter(AnimationGraph, StructView);
-		}
-
-		// Inner binding 
-		BindPublicVariables(PayloadVariableBindings);
-
-		// uninitializing the FDataInterfaceStructAdapter(s)
-		for (int32 StructIndex = 0; StructIndex < NumStructs; ++StructIndex)
-		{
-			Adapters[StructIndex].~FDataInterfaceStructAdapter();
-		}
-
-		// abandoning no longer initialized stack allocated arrays to their demise
-	}
-}
-
 void FAnimNextGraphInstance::BindPublicVariables(TConstArrayView<UE::AnimNext::IDataInterfaceHost*> InHosts)
 {
 	using namespace UE::AnimNext;
diff --git a/Engine/Plugins/Experimental/AnimNext/Source/AnimNext/Public/AnimNextDataInterfacePayload.h b/Engine/Plugins/Experimental/AnimNext/Source/AnimNext/Public/AnimNextDataInterfacePayload.h
deleted file mode 100644
diff --git a/Engine/Plugins/Experimental/AnimNext/Source/AnimNext/Public/DataInterface/AnimNextNativeDataInterface.h b/Engine/Plugins/Experimental/AnimNext/Source/AnimNext/Public/DataInterface/AnimNextNativeDataInterface.h
deleted file mode 100644
diff --git a/Engine/Plugins/Experimental/AnimNext/Source/AnimNext/Public/Graph/AnimNextGraphInstance.h b/Engine/Plugins/Experimental/AnimNext/Source/AnimNext/Public/Graph/AnimNextGraphInstance.h
--- a/Engine/Plugins/Experimental/AnimNext/Source/AnimNext/Public/Graph/AnimNextGraphInstance.h
+++ b/Engine/Plugins/Experimental/AnimNext/Source/AnimNext/Public/Graph/AnimNextGraphInstance.h
@@ -141,10 +141,6 @@ private:
 	template<typename HostType>
 	bool BindToHostHelper(const HostType& InHost, bool bInAutoBind);
 
-	// Bind the variables in the supplied traits in scope to their respective public variables, represented by structs held on a host
-	// Note: struct views are mutable here as bindings allow writes!
-	void BindPublicVariables(TArrayView<FStructView> InHostStructs);
-
 	// Bind the variables in the supplied traits in scope to their respective public variables, so they point at host memory
 	void BindPublicVariables(TConstArrayView<UE::AnimNext::IDataInterfaceHost*> InHosts);
 
diff --git a/Engine/Plugins/Experimental/AnimNext/Source/AnimNextUncookedOnly/Internal/AnimNextRigVMAssetEditorData.h b/Engine/Plugins/Experimental/AnimNext/Source/AnimNextUncookedOnly/Internal/AnimNextRigVMAssetEditorData.h
--- a/Engine/Plugins/Experimental/AnimNext/Source/AnimNextUncookedOnly/Internal/AnimNextRigVMAssetEditorData.h
+++ b/Engine/Plugins/Experimental/AnimNext/Source/AnimNextUncookedOnly/Internal/AnimNextRigVMAssetEditorData.h
@@ -420,10 +420,6 @@ protected:
 	UPROPERTY()
 	FRigVMGraphFunctionStore GraphFunctionStore;
 
-	// The native C++ struct that is used to communicate with this asset
-	UPROPERTY(EditAnywhere, Category = "Native", meta = (ShowDisplayNames, MetaStruct = "/Script/AnimNext.AnimNextNativeDataInterface"))
-	TObjectPtr<const UScriptStruct> NativeInterface;
-
 	UPROPERTY(EditAnywhere, Category = "User Interface")
 	FRigVMEdGraphDisplaySettings RigGraphDisplaySettings;
 
diff --git a/Engine/Plugins/Experimental/AnimNext/Source/AnimNextUncookedOnly/Private/UncookedOnlyUtils.cpp b/Engine/Plugins/Experimental/AnimNext/Source/AnimNextUncookedOnly/Private/UncookedOnlyUtils.cpp
--- a/Engine/Plugins/Experimental/AnimNext/Source/AnimNextUncookedOnly/Private/UncookedOnlyUtils.cpp
+++ b/Engine/Plugins/Experimental/AnimNext/Source/AnimNextUncookedOnly/Private/UncookedOnlyUtils.cpp
@@ -70,12 +70,11 @@ void FUtils::CompileVariables(UAnimNextRigVMAsset* InAsset)
 
 	struct FStructEntryInfo
 	{
-		const UAnimNextDataInterface* FromDataInterface = nullptr;
-		const UScriptStruct* NativeInterface = nullptr;
+		const UAnimNextDataInterface* FromDataInterface;
 		FName Name;
 		FAnimNextParamType Type;
-		EAnimNextExportAccessSpecifier AccessSpecifier = EAnimNextExportAccessSpecifier::Private;
-		bool bAutoBindDataInterfaceToHost = false;
+		EAnimNextExportAccessSpecifier AccessSpecifier;
+		bool bAutoBindDataInterfaceToHost;
 		TConstArrayView<uint8> Value;
 	};
 
@@ -149,13 +148,9 @@ void FUtils::CompileVariables(UAnimNextRigVMAsset* InAsset)
 		}
 		else
 		{
-			const UAnimNextDataInterface_EditorData* FromInterfaceEditorData = GetEditorData<const UAnimNextDataInterface_EditorData>(InFromInterface);
-			check(FromInterfaceEditorData != nullptr);
-
 			int32 Index = StructEntryInfos.Add(
 			{
 				InFromInterface,
-				FromInterfaceEditorData->NativeInterface,
 				Name,
 				FAnimNextParamType(Type.GetValueType(), Type.GetContainerType(), Type.GetValueTypeObject()),
 				AccessSpecifier,
@@ -258,7 +253,6 @@ void FUtils::CompileVariables(UAnimNextRigVMAsset* InAsset)
 			{
 				FAnimNextImplementedDataInterface& NewImplementedDataInterface = DataInterface->ImplementedInterfaces.AddDefaulted_GetRef();
 				NewImplementedDataInterface.DataInterface = StructEntryInfo->FromDataInterface;
-				NewImplementedDataInterface.NativeInterface = StructEntryInfo->NativeInterface;
 				NewImplementedDataInterface.VariableIndex = StructEntryInfo.GetIndex();
 				NewImplementedDataInterface.NumVariables = 1;
 				NewImplementedDataInterface.bAutoBindToHost = StructEntryInfo->bAutoBindDataInterfaceToHost; 
diff --git a/Engine/Plugins/Experimental/AnimNextAnimGraph/AnimNextAnimGraph.uplugin b/Engine/Plugins/Experimental/AnimNextAnimGraph/AnimNextAnimGraph.uplugin
--- a/Engine/Plugins/Experimental/AnimNextAnimGraph/AnimNextAnimGraph.uplugin
+++ b/Engine/Plugins/Experimental/AnimNextAnimGraph/AnimNextAnimGraph.uplugin
@@ -11,7 +11,7 @@
 	"MarketplaceURL" : "",
 	"SupportURL" : "",
 	"EnabledByDefault" : false,
-	"CanContainContent" : true,
+	"CanContainContent" : false,
 	"IsBetaVersion": false,
 	"IsExperimentalVersion": true,
 	"Installed" : false,
@@ -22,7 +22,7 @@
 			"Enabled": true
 		},
 		{
-			"Name": "RigVM",
+			"Name": "Chooser",
 			"Enabled": true
 		}
 	],
@@ -33,11 +33,6 @@
 			"Type" : "Runtime",
 			"LoadingPhase" : "PreDefault"
 		},
-		{
-			"Name" : "AnimNextAnimGraphEditor",
-			"Type" : "Editor",
-			"LoadingPhase" : "Default"
-		},
 		{
 			"Name" : "AnimNextAnimGraphUncookedOnly",
 			"Type" : "UncookedOnly",
diff --git a/Engine/Plugins/Experimental/AnimNextAnimGraph/Config/BaseAnimNextAnimGraph.ini b/Engine/Plugins/Experimental/AnimNextAnimGraph/Config/BaseAnimNextAnimGraph.ini
--- a/Engine/Plugins/Experimental/AnimNextAnimGraph/Config/BaseAnimNextAnimGraph.ini
+++ b/Engine/Plugins/Experimental/AnimNextAnimGraph/Config/BaseAnimNextAnimGraph.ini
@@ -1,7 +1,3 @@
-; Default settings/mappings for factories
-[/Script/AnimNextAnimGraph.AnimNextAnimGraphSettings]
-+AssetGraphMappings=(AssetType="/Script/Engine.AnimSequence",AnimationGraph="/AnimNextAnimGraph/FactoryGraphs/AG_AnimSequencePlayer.AG_AnimSequencePlayer",Variable="AnimSequence")
-
 [CoreRedirects]
 ; Traits rename
 +StructRedirects=(OldName="/Script/AnimNext.AnimNextBlendByBoolDecoratorSharedData", NewName="/Script/AnimNextAnimGraph.AnimNextBlendByBoolTraitSharedData")
@@ -32,8 +28,4 @@
 +PropertyRedirects=(OldName="/Script/AnimNextGraph.AnimNode_AnimNextGraph.AnimNextGraph", NewName="Module")
 
 ; Rename K2 Node
-+ClassRedirects=(OldName="/Script/AnimNextAnimGraphUncookedOnly.K2Node_PlayAnim", NewName="/Script/AnimNextAnimGraphUncookedOnly.K2Node_AnimNextPlayAnim")
-
-; Rename PlayAnim parameter
-+PropertyRedirects=(OldName="/Script/AnimNextAnimGraphUncookedOnly.CreateProxyObjectForPlayAnim.AnimationObject", NewName="AnimSequence")
-+PropertyRedirects=(OldName="/Script/AnimNextAnimGraphUncookedOnly.K2Node_PlayAnim.AnimationObject", NewName="AnimSequence")
++ClassRedirects=(OldName="/Script/AnimNextAnimGraphUncookedOnly.K2Node_PlayAnim", NewName="/Script/AnimNextAnimGraphUncookedOnly.K2Node_AnimNextPlayAnim")
\ No newline at end of file
diff --git a/Engine/Plugins/Experimental/AnimNextAnimGraph/Source/AnimNextAnimGraph/AnimNextAnimGraph.Build.cs b/Engine/Plugins/Experimental/AnimNextAnimGraph/Source/AnimNextAnimGraph/AnimNextAnimGraph.Build.cs
--- a/Engine/Plugins/Experimental/AnimNextAnimGraph/Source/AnimNextAnimGraph/AnimNextAnimGraph.Build.cs
+++ b/Engine/Plugins/Experimental/AnimNextAnimGraph/Source/AnimNextAnimGraph/AnimNextAnimGraph.Build.cs
@@ -14,6 +14,7 @@ namespace UnrealBuildTool.Rules
 					"RigVM",
 					"Engine",
 					"AnimNext",
+					"Chooser",
 				}
 			);
 		}
diff --git a/Engine/Plugins/Experimental/AnimNextAnimGraph/Source/AnimNextAnimGraph/Internal/AnimNextAnimGraphSettings.h b/Engine/Plugins/Experimental/AnimNextAnimGraph/Source/AnimNextAnimGraph/Internal/AnimNextAnimGraphSettings.h
deleted file mode 100644
diff --git a/Engine/Plugins/Experimental/AnimNextAnimGraph/Source/AnimNextAnimGraph/Internal/PlayAnim/PlayAnimCallbackProxy.h b/Engine/Plugins/Experimental/AnimNextAnimGraph/Source/AnimNextAnimGraph/Internal/PlayAnim/PlayAnimCallbackProxy.h
--- a/Engine/Plugins/Experimental/AnimNextAnimGraph/Source/AnimNextAnimGraph/Internal/PlayAnim/PlayAnimCallbackProxy.h
+++ b/Engine/Plugins/Experimental/AnimNextAnimGraph/Source/AnimNextAnimGraph/Internal/PlayAnim/PlayAnimCallbackProxy.h
@@ -30,7 +30,7 @@ class UPlayAnimCallbackProxy : public UObject
 	static ANIMNEXTANIMGRAPH_API UPlayAnimCallbackProxy* CreateProxyObjectForPlayAnim(
 		class UAnimNextComponent* AnimNextComponent,
 		FName SlotName,
-		class UAnimSequence* AnimSequence,
+		class UAnimSequence* AnimationObject,
 		float PlayRate = 1.0f,
 		float StartPosition = 0.0f,
 		FAnimNextPlayAnimBlendSettings BlendInSettings = FAnimNextPlayAnimBlendSettings(),
@@ -69,7 +69,7 @@ protected:
 	ANIMNEXTANIMGRAPH_API bool Play(
 		class UAnimNextComponent* AnimNextComponent,
 		FName SlotName,
-		class UAnimSequence* AnimSequence,
+		class UAnimSequence* AnimationObject,
 		float PlayRate,
 		float StartPosition,
 		const UE::AnimNext::FPlayAnimBlendSettings& BlendInSettings,
diff --git a/Engine/Plugins/Experimental/AnimNextAnimGraph/Source/AnimNextAnimGraph/Internal/Traits/BlendStackTrait.h b/Engine/Plugins/Experimental/AnimNextAnimGraph/Source/AnimNextAnimGraph/Internal/Traits/BlendStackTrait.h
--- a/Engine/Plugins/Experimental/AnimNextAnimGraph/Source/AnimNextAnimGraph/Internal/Traits/BlendStackTrait.h
+++ b/Engine/Plugins/Experimental/AnimNextAnimGraph/Source/AnimNextAnimGraph/Internal/Traits/BlendStackTrait.h
@@ -31,7 +31,7 @@ struct FAnimNextBlendStackTraitSharedData : public FAnimNextTraitSharedData
 
 	/** Asset to use as a sub-graph */
 	UPROPERTY(EditAnywhere, Category = "Default")
-	TObjectPtr<const UObject> Object;
+	TObjectPtr<const UAnimNextAnimationGraph> AnimationGraph;
 
 	/** Blend time to use when sub-graph pin input changes */
 	UPROPERTY(EditAnywhere, Category = "Default")
@@ -43,7 +43,7 @@ struct FAnimNextBlendStackTraitSharedData : public FAnimNextTraitSharedData
 
 	// Latent pin support boilerplate
 	#define TRAIT_LATENT_PROPERTIES_ENUMERATOR(GeneratorMacro) \
-		GeneratorMacro(Object) \
+		GeneratorMacro(AnimationGraph) \
 		GeneratorMacro(BlendTime) \
 		GeneratorMacro(bForceBlend) \
 
@@ -58,7 +58,7 @@ struct FAnimNextBlendStackRequesterTraitSharedData : public FAnimNextTraitShared
 
 	/** Asset to use as a sub-graph */
 	UPROPERTY(EditAnywhere, Category = "Default")
-	TObjectPtr<const UObject> Object;
+	TObjectPtr<const UAnimNextAnimationGraph> AnimationGraph;
 
 	/** Blend time to use when sub-graph pin input changes */
 	UPROPERTY(EditAnywhere, Category = "Default")
@@ -70,7 +70,7 @@ struct FAnimNextBlendStackRequesterTraitSharedData : public FAnimNextTraitShared
 
 	// Latent pin support boilerplate
 	#define TRAIT_LATENT_PROPERTIES_ENUMERATOR(GeneratorMacro) \
-		GeneratorMacro(Object) \
+		GeneratorMacro(AnimationGraph) \
 		GeneratorMacro(BlendTime) \
 		GeneratorMacro(bForceBlend) \
 
@@ -172,8 +172,8 @@ namespace UE::AnimNext
 	/**
 	 * FBlendStackTrait
 	 * 
-	 * Standalone blend stack that exposes desired object and blend time as pins.
-	 * A new graph will be pushed every time the newest object does not match the pin (or bForceBlend is true).
+	 * Standalone blend stack that exposes desired graph and blend time as pins. 
+	 * A new graph will be pushed every time the newest graph does not match the pin (or bForceBlend is true).
 	 */	
 	struct FBlendStackTrait : FBlendStackCoreTrait, IUpdate
 	{
@@ -203,5 +203,6 @@ namespace UE::AnimNext
 
 		// IUpdate impl
 		virtual void PreUpdate(FUpdateTraversalContext& Context, const TTraitBinding<IUpdate>& Binding, const FTraitUpdateState& TraitState) const override;
+		virtual void OnBecomeRelevant(FUpdateTraversalContext& Context, const TTraitBinding<IUpdate>& Binding, const FTraitUpdateState& TraitState) const override;
 	};
 }
diff --git a/Engine/Plugins/Experimental/AnimNextAnimGraph/Source/AnimNextAnimGraph/Private/AnimNextAnimGraphSettings.cpp b/Engine/Plugins/Experimental/AnimNextAnimGraph/Source/AnimNextAnimGraph/Private/AnimNextAnimGraphSettings.cpp
deleted file mode 100644
diff --git a/Engine/Plugins/Experimental/AnimNextAnimGraph/Source/AnimNextAnimGraph/Private/GraphInterfaces/AnimNextNativeDataInterface_AnimSequencePlayer.cpp b/Engine/Plugins/Experimental/AnimNextAnimGraph/Source/AnimNextAnimGraph/Private/GraphInterfaces/AnimNextNativeDataInterface_AnimSequencePlayer.cpp
deleted file mode 100644
diff --git a/Engine/Plugins/Experimental/AnimNextAnimGraph/Source/AnimNextAnimGraph/Private/Module.cpp b/Engine/Plugins/Experimental/AnimNextAnimGraph/Source/AnimNextAnimGraph/Private/Module.cpp
--- a/Engine/Plugins/Experimental/AnimNextAnimGraph/Source/AnimNextAnimGraph/Private/Module.cpp
+++ b/Engine/Plugins/Experimental/AnimNextAnimGraph/Source/AnimNextAnimGraph/Private/Module.cpp
@@ -1,7 +1,6 @@
 // Copyright Epic Games, Inc. All Rights Reserved.
 
 #include "CoreMinimal.h"
-#include "AnimNextAnimGraphSettings.h"
 #include "Animation/BlendProfile.h"
 #include "Curves/CurveFloat.h"
 #include "Module/AnimNextModule.h"
@@ -9,41 +8,22 @@
 #include "Modules/ModuleManager.h"
 #include "RigVMCore/RigVMRegistry.h"
 #include "Animation/AnimSequence.h"
+#include "Chooser.h"
 
 #include "IAnimNextModuleInterface.h"
 #include "TraitInterfaces/IEvaluate.h"
 #include "TraitInterfaces/IUpdate.h"
 #include "EvaluationVM/EvaluationVM.h"
 #include "Graph/AnimNext_LODPose.h"
-#include "PlayAnim/PlayAnimRequest.h"
 
 namespace UE::AnimNext::AnimGraph
 {
 
-// Delay loading of default asset factory mappings in settings 
-static FDelayedAutoRegisterHelper AutoRegisterFactoryMappings(EDelayedRegisterRunPhase::EndOfEngineInit, []
-{
-	UAnimNextAnimGraphSettings* Settings = GetMutableDefault<UAnimNextAnimGraphSettings>();
-	Settings->LoadAndRebuildMappings();
-});
-
 class FModule : public IModuleInterface, IAnimNextAnimGraph
 {
 public:
 	virtual void StartupModule() override
 	{
-		// Ensure that AnimNext modules are loaded so we can correctly load plugin content
-		FModuleManager::LoadModuleChecked<IModuleInterface>("AnimNext");
-#if WITH_EDITORONLY_DATA
-		FModuleManager::LoadModuleChecked<IModuleInterface>("AnimNextUncookedOnly");
-#endif
-
-		// Setup default settings/factories
-		{
-			UAnimNextAnimGraphSettings* Settings = GetMutableDefault<UAnimNextAnimGraphSettings>();
-			Settings->LoadConfig();
-		}
-
 		static TPair<UClass*, FRigVMRegistry::ERegisterObjectOperation> const AllowedObjectTypes[] =
 		{
 			{ UAnimSequence::StaticClass(), FRigVMRegistry::ERegisterObjectOperation::Class },
@@ -51,6 +31,7 @@ public:
 			{ UBlendProfile::StaticClass(), FRigVMRegistry::ERegisterObjectOperation::Class },
 			{ UCurveFloat::StaticClass(), FRigVMRegistry::ERegisterObjectOperation::Class },
 			{ UAnimNextModule::StaticClass(), FRigVMRegistry::ERegisterObjectOperation::Class },
+			{ UChooserTable::StaticClass(), FRigVMRegistry::ERegisterObjectOperation::Class },
 		};
 
 		FRigVMRegistry::Get().RegisterObjectTypes(AllowedObjectTypes);
diff --git a/Engine/Plugins/Experimental/AnimNextAnimGraph/Source/AnimNextAnimGraph/Private/PlayAnim/PlayAnimCallbackProxy.cpp b/Engine/Plugins/Experimental/AnimNextAnimGraph/Source/AnimNextAnimGraph/Private/PlayAnim/PlayAnimCallbackProxy.cpp
--- a/Engine/Plugins/Experimental/AnimNextAnimGraph/Source/AnimNextAnimGraph/Private/PlayAnim/PlayAnimCallbackProxy.cpp
+++ b/Engine/Plugins/Experimental/AnimNextAnimGraph/Source/AnimNextAnimGraph/Private/PlayAnim/PlayAnimCallbackProxy.cpp
@@ -2,7 +2,6 @@
 
 #include "PlayAnim/PlayAnimCallbackProxy.h"
 #include "Animation/AnimSequence.h"
-#include "GraphInterfaces/AnimNextNativeDataInterface_AnimSequencePlayer.h"
 
 #include UE_INLINE_GENERATED_CPP_BY_NAME(PlayAnimCallbackProxy)
 
@@ -17,7 +16,7 @@ UPlayAnimCallbackProxy::UPlayAnimCallbackProxy(const FObjectInitializer& ObjectI
 UPlayAnimCallbackProxy* UPlayAnimCallbackProxy::CreateProxyObjectForPlayAnim(
 	UAnimNextComponent* AnimNextComponent,
 	FName SlotName,
-	UAnimSequence* AnimSequence,
+	UAnimSequence* AnimationObject,
 	float PlayRate,
 	float StartPosition,
 	UE::AnimNext::FPlayAnimBlendSettings BlendInSettings,
@@ -25,7 +24,7 @@ UPlayAnimCallbackProxy* UPlayAnimCallbackProxy::CreateProxyObjectForPlayAnim(
 {
 	UPlayAnimCallbackProxy* Proxy = NewObject<UPlayAnimCallbackProxy>();
 	Proxy->SetFlags(RF_StrongRefOnFrame);
-	Proxy->Play(AnimNextComponent, SlotName, AnimSequence, PlayRate, StartPosition, BlendInSettings, BlendOutSettings);
+	Proxy->Play(AnimNextComponent, SlotName, AnimationObject, PlayRate, StartPosition, BlendInSettings, BlendOutSettings);
 	return Proxy;
 }
 
@@ -47,19 +46,19 @@ UPlayAnimCallbackProxy* UPlayAnimCallbackProxy::CreateProxyObjectForPlayAsset(
 bool UPlayAnimCallbackProxy::Play(
 	UAnimNextComponent* AnimNextComponent,
 	FName SlotName,
-	UAnimSequence* AnimSequence,
+	UAnimSequence* AnimationObject,
 	float PlayRate,
 	float StartPosition,
 	const UE::AnimNext::FPlayAnimBlendSettings& BlendInSettings,
 	const UE::AnimNext::FPlayAnimBlendSettings& BlendOutSettings)
 {
-	FInstancedStruct NativeInterface;
-	NativeInterface.InitializeAs<FAnimNextNativeDataInterface_AnimSequencePlayer>();
-	FAnimNextNativeDataInterface_AnimSequencePlayer& PlayNativeInterface = NativeInterface.GetMutable<FAnimNextNativeDataInterface_AnimSequencePlayer>();
-	PlayNativeInterface.AnimSequence = AnimSequence;
-	PlayNativeInterface.PlayRate = PlayRate;
-	PlayNativeInterface.StartPosition = StartPosition;
-	return Play(AnimNextComponent, SlotName, AnimSequence, MoveTemp(NativeInterface), BlendInSettings, BlendOutSettings);
+	FInstancedStruct Payload;
+	Payload.InitializeAs<FAnimNextPlayAnimPayload>();
+	FAnimNextPlayAnimPayload& PlayAnimPayload = Payload.GetMutable<FAnimNextPlayAnimPayload>();
+	PlayAnimPayload.AnimationObject = AnimationObject;
+	PlayAnimPayload.PlayRate = PlayRate;
+	PlayAnimPayload.StartPosition = StartPosition;
+	return Play(AnimNextComponent, SlotName, AnimationObject, MoveTemp(Payload), BlendInSettings, BlendOutSettings);
 }
 
 bool UPlayAnimCallbackProxy::Play(
@@ -78,7 +77,7 @@ bool UPlayAnimCallbackProxy::Play(
 		RequestArgs.Object = Object;
 		RequestArgs.BlendInSettings = BlendInSettings;
 		RequestArgs.BlendOutSettings = BlendOutSettings;
-		RequestArgs.Payload.AddNative(MoveTemp(Payload));
+		RequestArgs.Payload = MoveTemp(Payload);
 
 		auto Request = UE::AnimNext::MakePlayAnimRequest();
 		Request->OnCompleted.BindUObject(this, &UPlayAnimCallbackProxy::OnPlayAnimCompleted);
diff --git a/Engine/Plugins/Experimental/AnimNextAnimGraph/Source/AnimNextAnimGraph/Private/PlayAnim/PlayAnimRequest.cpp b/Engine/Plugins/Experimental/AnimNextAnimGraph/Source/AnimNextAnimGraph/Private/PlayAnim/PlayAnimRequest.cpp
--- a/Engine/Plugins/Experimental/AnimNextAnimGraph/Source/AnimNextAnimGraph/Private/PlayAnim/PlayAnimRequest.cpp
+++ b/Engine/Plugins/Experimental/AnimNextAnimGraph/Source/AnimNextAnimGraph/Private/PlayAnim/PlayAnimRequest.cpp
@@ -10,7 +10,7 @@ namespace UE::AnimNext
 	bool FPlayAnimRequest::Play(FPlayAnimRequestArgs&& InRequestArgs, UAnimNextComponent* InComponent)
 	{
 		check(IsInGameThread());
-		if (InRequestArgs.Payload.Get().Num() == 0 || InComponent == nullptr)
+		if (!InRequestArgs.Payload.IsValid() || InComponent == nullptr)
 		{
 			return false;	// Nothing to play
 		}
@@ -159,7 +159,8 @@ namespace UE::AnimNext
 
 	void FPlayAnimRequest::AddReferencedObjects(FReferenceCollector& Collector)
 	{
-		Collector.AddPropertyReferences(FPlayAnimRequestArgs::StaticStruct(), &RequestArgs);
+		Collector.AddReferencedObject(RequestArgs.Object);
+		Collector.AddPropertyReferencesWithStructARO(FInstancedPropertyBag::StaticStruct(), &RequestArgs.Payload);
 		Collector.AddReferencedObject(Component);
 	}
 }
diff --git a/Engine/Plugins/Experimental/AnimNextAnimGraph/Source/AnimNextAnimGraph/Private/PlayAnim/PlayAnimSlotTrait.cpp b/Engine/Plugins/Experimental/AnimNextAnimGraph/Source/AnimNextAnimGraph/Private/PlayAnim/PlayAnimSlotTrait.cpp
--- a/Engine/Plugins/Experimental/AnimNextAnimGraph/Source/AnimNextAnimGraph/Private/PlayAnim/PlayAnimSlotTrait.cpp
+++ b/Engine/Plugins/Experimental/AnimNextAnimGraph/Source/AnimNextAnimGraph/Private/PlayAnim/PlayAnimSlotTrait.cpp
@@ -2,8 +2,6 @@
 
 #include "PlayAnim/PlayAnimSlotTrait.h"
 
-#include "AnimNextAnimGraphSettings.h"
-#include "AnimNextDataInterfacePayload.h"
 #include "Animation/AnimSequence.h"
 #include "DataInterface/AnimNextDataInterfaceHost.h"
 #include "DataInterface/DataInterfaceStructAdapter.h"
@@ -13,8 +11,6 @@
 #include "Graph/AnimNextGraphInstance.h"
 #include "Logging/StructuredLog.h"
 #include "Module/ModuleEvents.h"
-#include "TraitInterfaces/IGraphFactory.h"
-#include "TraitInterfaces/IPayloadFactory.h"
 
 namespace UE::AnimNext
 {
@@ -38,12 +34,11 @@ namespace UE::AnimNext
 	#undef TRAIT_INTERFACE_ENUMERATOR
 	#undef TRAIT_EVENT_ENUMERATOR
 
-	void FPlayAnimSlotTrait::FPlayAnimSlotRequest::Initialize(FPlayAnimRequestPtr InRequest, const FPlayAnimBlendSettings& InBlendSettings, const UAnimNextAnimationGraph* InAnimationGraph, FAnimNextDataInterfacePayload&& InPayload)
+	void FPlayAnimSlotTrait::FPlayAnimSlotRequest::Initialize(FPlayAnimRequestPtr InRequest, const FPlayAnimBlendSettings& InBlendSettings, const UAnimNextAnimationGraph* InAnimationGraph)
 	{
 		Request = InRequest;
 		BlendSettings = InBlendSettings;
 		AnimationGraph = InAnimationGraph;
-		Payload = MoveTemp(InPayload);
 
 		// If no input is provided, we'll use the source
 		State = InAnimationGraph != nullptr ? EPlayAnimRequestState::Active : EPlayAnimRequestState::ActiveSource;
@@ -164,15 +159,32 @@ namespace UE::AnimNext
 
 			FPlayAnimBlendSettings BlendSettings;
 			const UAnimNextAnimationGraph* AnimationGraph = nullptr;
-			FAnimNextDataInterfacePayload Payload;
-
 			if (Request)
 			{
-				// This is a new pending request, lookup the sub-graph to use with any factory trait we have and the desired animation object
-				FAnimNextPlayAnimRequestArgs& RequestArgs = Request->GetMutableArgs();
+				// This is a new pending request, lookup the sub-graph to use with our chooser and the desired animation object
+
+				// TODO: Rather than this built-in choosing logic, this could be implemented by a new optional trait:
+				// - Implement some new trait that converts an event payload's object into a graph + payload struct for variable binding
+				// - Could possibly remove Object from FAnimNextPlayAnimRequestArgs, or AnimationObject from FAnimNextPlayAnimPayload, as they are duplicated
+				// - Call this trait's interface here and use the resulting animation graph (binding variables in OnBlendInitiated)
+				// - Reentrancy needs to handle 'same graph, different parameters'
+				const FAnimNextPlayAnimRequestArgs& RequestArgs = Request->GetArgs();
 				if (RequestArgs.Object)
 				{
-					AnimationGraph = IGraphFactory::GetGraphFromObjectWithFallback(Context, Binding, RequestArgs.Object);
+					if (const UChooserTable* Chooser = SharedData->GetSubGraphChooser(Binding))
+					{
+						FAnimNextPlayAnimChooserParameters ChooserParameters;
+						ChooserParameters.AnimationObjectType = RequestArgs.Object->GetClass();
+
+						FChooserEvaluationContext ChooserContext;
+						ChooserContext.AddStructParam(ChooserParameters);
+
+						UChooserTable::EvaluateChooser(ChooserContext, Chooser, FObjectChooserBase::FObjectChooserIteratorCallback::CreateLambda([&AnimationGraph](UObject* InResult)
+							{
+								AnimationGraph = Cast<UAnimNextAnimationGraph>(InResult);
+								return FObjectChooserBase::EIteratorStatus::Stop;
+							}));
+					}
 
 					if (AnimationGraph != nullptr)
 					{
@@ -183,16 +195,11 @@ namespace UE::AnimNext
 						{
 							if (OwnerGraphInstance->UsesAnimationGraph(AnimationGraph) && OwnerGraphInstance->UsesEntryPoint(EntryPoint))
 							{
-								UE_LOGFMT(LogAnimation, Warning, "Ignoring PlayAnim request for {0}, re-entrancy detected", AnimationGraph->GetFName());
 								return;
 							}
 
 							OwnerGraphInstance = OwnerGraphInstance->GetParentGraphInstance();
 						}
-
-						Payload = MoveTemp(RequestArgs.Payload);
-
-						IPayloadFactory::GetPayloadForGraphWithFallback(Context, Binding, RequestArgs.Object, AnimationGraph, Payload);
 					}
 
 					BlendSettings = RequestArgs.BlendInSettings;
@@ -217,7 +224,7 @@ namespace UE::AnimNext
 			const int32 FreeRequestIndex = FindFreeRequestIndexOrAdd(*InstanceData);
 
 			FPlayAnimSlotRequest& SlotRequest = InstanceData->SlotRequests[FreeRequestIndex];
-			SlotRequest.Initialize(Request, BlendSettings, AnimationGraph, MoveTemp(Payload));
+			SlotRequest.Initialize(Request, BlendSettings, AnimationGraph);
 
 			const int32 OldChildIndex = InstanceData->CurrentlyActiveRequestIndex;
 			const int32 NewChildIndex = FreeRequestIndex;
@@ -288,7 +295,7 @@ namespace UE::AnimNext
 					const int32 FreeRequestIndex = FindFreeRequestIndexOrAdd(*InstanceData);
 
 					FPlayAnimSlotRequest& FreeSlotRequest = InstanceData->SlotRequests[FreeRequestIndex];
-					FreeSlotRequest.Initialize(FPlayAnimRequestPtr(), RequestArgs.BlendOutSettings, nullptr, FAnimNextDataInterfacePayload());
+					FreeSlotRequest.Initialize(FPlayAnimRequestPtr(), RequestArgs.BlendOutSettings, nullptr);
 
 					const int32 OldChildIndex = InstanceData->CurrentlyActiveRequestIndex;
 					const int32 NewChildIndex = FreeRequestIndex;
@@ -377,7 +384,10 @@ namespace UE::AnimNext
 				const FName EntryPoint = SlotRequest.AnimationGraph->DefaultEntryPoint;
 				SlotRequest.GraphInstance = SlotRequest.AnimationGraph->AllocateInstance(nullptr, &Binding.GetTraitPtr().GetNodeInstance()->GetOwner(), EntryPoint);
 				SlotRequest.ChildPtr = SlotRequest.GraphInstance.IsValid() ? SlotRequest.GraphInstance->GetGraphRootPtr() : FTraitPtr();
-				SlotRequest.GraphInstance->BindPublicVariables(SlotRequest.Payload.Get());
+
+				// Note: args are mutable here as bindings allow writes!
+				FDataInterfaceStructAdapter VariableBinding(SlotRequest.AnimationGraph, SlotRequest.Request->GetMutableArgs().Payload);
+				SlotRequest.GraphInstance->BindPublicVariables({ &VariableBinding });
 
 				// TODO: Validate that our child implements the ITimeline interface
 
diff --git a/Engine/Plugins/Experimental/AnimNextAnimGraph/Source/AnimNextAnimGraph/Private/PlayAnim/PlayAnimSlotTrait.h b/Engine/Plugins/Experimental/AnimNextAnimGraph/Source/AnimNextAnimGraph/Private/PlayAnim/PlayAnimSlotTrait.h
--- a/Engine/Plugins/Experimental/AnimNextAnimGraph/Source/AnimNextAnimGraph/Private/PlayAnim/PlayAnimSlotTrait.h
+++ b/Engine/Plugins/Experimental/AnimNextAnimGraph/Source/AnimNextAnimGraph/Private/PlayAnim/PlayAnimSlotTrait.h
@@ -14,6 +14,7 @@
 #include "Graph/AnimNextAnimationGraph.h"
 #include "Graph/AnimNextGraphInstance.h"
 #include "PlayAnim/PlayAnimEvents.h"
+#include "Chooser.h"
 
 #include "PlayAnimSlotTrait.generated.h"
 
@@ -40,9 +41,14 @@ struct FAnimNextPlayAnimSlotTraitSharedData : public FAnimNextTraitSharedData
 	UPROPERTY(EditAnywhere, Category = "Default")
 	FName SlotName;
 
+	/** The chooser instance to use to convert animation objects into a sub-graph to play. */
+	UPROPERTY(EditAnywhere, Category = "Default")
+	TObjectPtr<UChooserTable> SubGraphChooser;
+
 	// Latent pin support boilerplate
 	#define TRAIT_LATENT_PROPERTIES_ENUMERATOR(GeneratorMacro) \
 		GeneratorMacro(SlotName) \
+		GeneratorMacro(SubGraphChooser) \
 
 	GENERATE_TRAIT_LATENT_PROPERTIES(FAnimNextPlayAnimSlotTraitSharedData, TRAIT_LATENT_PROPERTIES_ENUMERATOR)
 	#undef TRAIT_LATENT_PROPERTIES_ENUMERATOR
@@ -104,9 +110,6 @@ namespace UE::AnimNext
 			// The module used by the graph instance, as selected by the chooser
 			TObjectPtr<const UAnimNextAnimationGraph> AnimationGraph;
 
-			// Payloads for the graph instance when instantiated
-			FAnimNextDataInterfacePayload Payload;
-
 			// The graph instance
 			TSharedPtr<FAnimNextGraphInstance> GraphInstance;
 
@@ -121,7 +124,7 @@ namespace UE::AnimNext
 			bool bWasRelevant = false;
 
 			// Initializes a request to begin playing
-			void Initialize(FPlayAnimRequestPtr InRequest, const FPlayAnimBlendSettings& InBlendSettings, const UAnimNextAnimationGraph* InAnimationGraph, FAnimNextDataInterfacePayload&& InPayload);
+			void Initialize(FPlayAnimRequestPtr InRequest, const FPlayAnimBlendSettings& InBlendSettings, const UAnimNextAnimationGraph* InAnimationGraph);
 		};
 
 		using FSharedData = FAnimNextPlayAnimSlotTraitSharedData;
diff --git a/Engine/Plugins/Experimental/AnimNextAnimGraph/Source/AnimNextAnimGraph/Private/TraitInterfaces/IGraphFactory.cpp b/Engine/Plugins/Experimental/AnimNextAnimGraph/Source/AnimNextAnimGraph/Private/TraitInterfaces/IGraphFactory.cpp
deleted file mode 100644
diff --git a/Engine/Plugins/Experimental/AnimNextAnimGraph/Source/AnimNextAnimGraph/Private/TraitInterfaces/IPayloadFactory.cpp b/Engine/Plugins/Experimental/AnimNextAnimGraph/Source/AnimNextAnimGraph/Private/TraitInterfaces/IPayloadFactory.cpp
deleted file mode 100644
diff --git a/Engine/Plugins/Experimental/AnimNextAnimGraph/Source/AnimNextAnimGraph/Private/Traits/BlendStackTrait.cpp b/Engine/Plugins/Experimental/AnimNextAnimGraph/Source/AnimNextAnimGraph/Private/Traits/BlendStackTrait.cpp
--- a/Engine/Plugins/Experimental/AnimNextAnimGraph/Source/AnimNextAnimGraph/Private/Traits/BlendStackTrait.cpp
+++ b/Engine/Plugins/Experimental/AnimNextAnimGraph/Source/AnimNextAnimGraph/Private/Traits/BlendStackTrait.cpp
@@ -1,14 +1,9 @@
 // Copyright Epic Games, Inc. All Rights Reserved.
 
 #include "Traits/BlendStackTrait.h"
-
-#include "AnimNextAnimGraphSettings.h"
-#include "AnimNextDataInterfacePayload.h"
 #include "DataInterface/DataInterfaceStructAdapter.h"
 #include "TraitCore/NodeInstance.h"
 #include "TraitCore/ExecutionContext.h"
-#include "TraitInterfaces/IGraphFactory.h"
-#include "TraitInterfaces/IPayloadFactory.h"
 
 namespace UE::AnimNext
 {
@@ -293,10 +288,34 @@ namespace UE::AnimNext
 		TTraitBinding<IDiscreteBlend> DiscreteBlendTrait;
 		Binding.GetStackInterface(DiscreteBlendTrait);
 		DiscreteBlendTrait.OnBlendTransition(Context, OldChildIndex, NewChildIndex);
-
+		
 		InstanceData->CurrentlyActiveGraphIndex = NewChildIndex;
 
-		Graph.Instance->BindPublicVariables(Graph.Request.GraphPayload.Get());
+		const int32 NumPayloads = Graph.Request.Payloads.Num();
+		if (NumPayloads > 0)
+		{
+			// stack allocating uninitialized arrays to contain FDataInterfaceStructAdapter(s)
+			TArrayView<FDataInterfaceStructAdapter> Adapters((FDataInterfaceStructAdapter*)FMemory_Alloca(NumPayloads * sizeof(FDataInterfaceStructAdapter)), NumPayloads);
+			TArrayView<IDataInterfaceHost*> PayloadVariableBindings((IDataInterfaceHost**)FMemory_Alloca(NumPayloads * sizeof(IDataInterfaceHost*)), NumPayloads);
+
+			// initializing the FDataInterfaceStructAdapter(s) with their respective Graph.Request.Payloads
+			for (int32 PayloadIndex = 0; PayloadIndex < NumPayloads; ++PayloadIndex)
+			{
+				FInstancedStruct& Payload = Graph.Request.Payloads[PayloadIndex];
+				PayloadVariableBindings[PayloadIndex] = new (&Adapters[PayloadIndex]) FDataInterfaceStructAdapter(Graph.Request.AnimationGraph, Payload);
+			}
+			
+			// binding 
+			Graph.Instance->BindPublicVariables(PayloadVariableBindings);
+			
+			// uninitializing the FDataInterfaceStructAdapter(s)
+			for (int32 PayloadIndex = 0; PayloadIndex < NumPayloads; ++PayloadIndex)
+			{
+				Adapters[PayloadIndex].~FDataInterfaceStructAdapter();
+			}
+
+			// abandoning no longer initialized stack allocated arrays to their demise
+		}
 	}
 
 	void FBlendStackCoreTrait::GetActiveGraph(FExecutionContext& Context, const TTraitBinding<IBlendStack>& Binding, IBlendStack::FGraphRequestPtr& OutGraphRequest) const
@@ -318,32 +337,21 @@ namespace UE::AnimNext
 		FInstanceData* InstanceData = Binding.GetInstanceData<FInstanceData>();
 		const FSharedData* SharedData = Binding.GetSharedData<FSharedData>();
  
-		const UObject* DesiredObject = SharedData->GetObject(Binding);
-		if (DesiredObject != nullptr)
+		const UAnimNextAnimationGraph* DesiredAnimationGraph = SharedData->GetAnimationGraph(Binding);
+		const int32 CurrentlyActiveGraphIndex = InstanceData->CurrentlyActiveGraphIndex;
+
+		const bool bForceBlend = SharedData->GetbForceBlend(Binding);
+		const bool bIsEmpty = InstanceData->CurrentlyActiveGraphIndex == INDEX_NONE;
+		if (DesiredAnimationGraph && (bForceBlend || bIsEmpty || (DesiredAnimationGraph != InstanceData->ChildGraphs[CurrentlyActiveGraphIndex].Request.AnimationGraph)))
 		{
-			const UAnimNextAnimationGraph* AnimationGraph = IGraphFactory::GetGraphFromObjectWithFallback(Context, Binding, DesiredObject);
-			if (AnimationGraph != nullptr)
-			{
-				const int32 CurrentlyActiveGraphIndex = InstanceData->CurrentlyActiveGraphIndex;
-				const bool bForceBlend = SharedData->GetbForceBlend(Binding);
-				const bool bIsEmpty = CurrentlyActiveGraphIndex == INDEX_NONE;
-				if (bForceBlend || bIsEmpty || (DesiredObject != InstanceData->ChildGraphs[CurrentlyActiveGraphIndex].Request.FactoryObject))
-				{
-					FAnimNextDataInterfacePayload GraphPayload;
-					IPayloadFactory::GetPayloadForGraphWithFallback(Context, Binding, DesiredObject, AnimationGraph, GraphPayload);
-
-					TTraitBinding<IBlendStack> BlendStackTrait;
-					Binding.GetStackInterface(BlendStackTrait);
-
-					IBlendStack::FGraphRequest GraphRequest;
-					GraphRequest.FactoryObject = DesiredObject;
-					GraphRequest.AnimationGraph = AnimationGraph;
-					GraphRequest.BlendArgs.BlendTime = SharedData->GetBlendTime(Binding);
-					GraphRequest.GraphPayload = MoveTemp(GraphPayload);
-
-					BlendStackTrait.PushGraph(Context, MoveTemp(GraphRequest));
-				}
-			}
+			TTraitBinding<IBlendStack> BlendStackTrait;
+			Binding.GetStackInterface(BlendStackTrait);
+
+			IBlendStack::FGraphRequest GraphRequest;
+			GraphRequest.AnimationGraph = DesiredAnimationGraph;
+			GraphRequest.BlendArgs.BlendTime = SharedData->GetBlendTime(Binding);
+
+			BlendStackTrait.PushGraph(Context, MoveTemp(GraphRequest));
 		}
 	}
 
@@ -351,32 +359,41 @@ namespace UE::AnimNext
 	{
 		const FSharedData* SharedData = Binding.GetSharedData<FSharedData>();
 
-		const UObject* DesiredObject = SharedData->GetObject(Binding);
-		if (DesiredObject != nullptr)
+		const UAnimNextAnimationGraph* DesiredAnimationGraph = SharedData->GetAnimationGraph(Binding);
+		if (DesiredAnimationGraph != nullptr)
 		{
-			const UAnimNextAnimationGraph* AnimationGraph = IGraphFactory::GetGraphFromObjectWithFallback(Context, Binding, DesiredObject);
-			if (AnimationGraph != nullptr)
+			TTraitBinding<IBlendStack> BlendStackTrait;
+			Binding.GetStackInterface(BlendStackTrait);
+
+			IBlendStack::FGraphRequestPtr ActiveGraphRequest;
+			BlendStackTrait.GetActiveGraph(Context, ActiveGraphRequest);
+			const bool bForceBlend = SharedData->GetbForceBlend(Binding);
+			if (bForceBlend || !ActiveGraphRequest || (DesiredAnimationGraph != ActiveGraphRequest->AnimationGraph))
 			{
-				TTraitBinding<IBlendStack> BlendStackTrait;
-				Binding.GetStackInterface(BlendStackTrait);
-
-				IBlendStack::FGraphRequestPtr ActiveGraphRequest;
-				BlendStackTrait.GetActiveGraph(Context, ActiveGraphRequest);
-				const bool bForceBlend = SharedData->GetbForceBlend(Binding);
-				if (bForceBlend || !ActiveGraphRequest || (DesiredObject != ActiveGraphRequest->FactoryObject))
-				{
-					FAnimNextDataInterfacePayload GraphPayload;
-					IPayloadFactory::GetPayloadForGraphWithFallback(Context, Binding, DesiredObject, AnimationGraph, GraphPayload);
-
-					IBlendStack::FGraphRequest GraphRequest;
-					GraphRequest.FactoryObject = DesiredObject;
-					GraphRequest.AnimationGraph = AnimationGraph;
-					GraphRequest.BlendArgs.BlendTime = SharedData->GetBlendTime(Binding);
-					GraphRequest.GraphPayload = MoveTemp(GraphPayload);
-
-					BlendStackTrait.PushGraph(Context, MoveTemp(GraphRequest));
-				}
+				IBlendStack::FGraphRequest GraphRequest;
+				GraphRequest.AnimationGraph = DesiredAnimationGraph;
+				GraphRequest.BlendArgs.BlendTime = SharedData->GetBlendTime(Binding);
+
+				BlendStackTrait.PushGraph(Context, MoveTemp(GraphRequest));
 			}
 		}
 	}
+
+	void FBlendStackRequesterTrait::OnBecomeRelevant(FUpdateTraversalContext& Context, const TTraitBinding<IUpdate>& Binding, const FTraitUpdateState& TraitState) const
+	{
+		const FSharedData* SharedData = Binding.GetSharedData<FSharedData>();
+
+		TTraitBinding<IBlendStack> BlendStackTrait;
+		Binding.GetStackInterface(BlendStackTrait);
+
+		const UAnimNextAnimationGraph* DesiredAnimationGraph = SharedData->GetAnimationGraph(Binding);
+		if (DesiredAnimationGraph != nullptr)
+		{
+			IBlendStack::FGraphRequest GraphRequest;
+			GraphRequest.AnimationGraph = DesiredAnimationGraph;
+			GraphRequest.BlendArgs.BlendTime = SharedData->GetBlendTime(Binding);
+
+			BlendStackTrait.PushGraph(Context, MoveTemp(GraphRequest));
+		}
+	}
 }
diff --git a/Engine/Plugins/Experimental/AnimNextAnimGraph/Source/AnimNextAnimGraph/Private/Traits/MakeGraphFromAsset.cpp b/Engine/Plugins/Experimental/AnimNextAnimGraph/Source/AnimNextAnimGraph/Private/Traits/MakeGraphFromAsset.cpp
deleted file mode 100644
diff --git a/Engine/Plugins/Experimental/AnimNextAnimGraph/Source/AnimNextAnimGraph/Private/Traits/MakeGraphFromAsset.h b/Engine/Plugins/Experimental/AnimNextAnimGraph/Source/AnimNextAnimGraph/Private/Traits/MakeGraphFromAsset.h
deleted file mode 100644
diff --git a/Engine/Plugins/Experimental/AnimNextAnimGraph/Source/AnimNextAnimGraph/Public/GraphInterfaces/AnimNextNativeDataInterface_AnimSequencePlayer.h b/Engine/Plugins/Experimental/AnimNextAnimGraph/Source/AnimNextAnimGraph/Public/GraphInterfaces/AnimNextNativeDataInterface_AnimSequencePlayer.h
deleted file mode 100644
diff --git a/Engine/Plugins/Experimental/AnimNextAnimGraph/Source/AnimNextAnimGraph/Public/PlayAnim/PlayAnimRequest.h b/Engine/Plugins/Experimental/AnimNextAnimGraph/Source/AnimNextAnimGraph/Public/PlayAnim/PlayAnimRequest.h
--- a/Engine/Plugins/Experimental/AnimNextAnimGraph/Source/AnimNextAnimGraph/Public/PlayAnim/PlayAnimRequest.h
+++ b/Engine/Plugins/Experimental/AnimNextAnimGraph/Source/AnimNextAnimGraph/Public/PlayAnim/PlayAnimRequest.h
@@ -5,7 +5,6 @@
 #include "CoreMinimal.h"
 
 #include "AlphaBlend.h"
-#include "AnimNextDataInterfacePayload.h"
 #include "TraitCore/TraitEvent.h"
 #include "TraitInterfaces/ITimeline.h"
 #include "PlayAnim/PlayAnimStatus.h"
@@ -60,6 +59,33 @@ struct FAnimNextPlayAnimBlendSettings
 	EAnimNextPlayAnimBlendMode BlendMode = EAnimNextPlayAnimBlendMode::Standard;
 };
 
+/**
+ * PlayAnim payload
+ *
+ * Encapsulates the data interface used to play a simple animation
+ */
+USTRUCT(BlueprintType)
+struct FAnimNextPlayAnimPayload
+{
+	GENERATED_BODY()
+
+	// The animation object to play with this request
+	UPROPERTY(EditAnywhere, BlueprintReadWrite, Category = "PlayAnim")
+	TObjectPtr<UAnimSequence> AnimationObject;
+
+	// The play rate of the request
+	UPROPERTY(EditAnywhere, BlueprintReadWrite, Category = "PlayAnim")
+	double PlayRate = 1.0f;
+
+	// The timeline start position of the request
+	UPROPERTY(EditAnywhere, BlueprintReadWrite, Category = "PlayAnim")
+	double StartPosition = 0.0f;
+
+	// Whether to loop the animation
+	UPROPERTY(EditAnywhere, BlueprintReadWrite, Category = "PlayAnim")
+	bool IsLooping = false;
+};
+
 /**
  * PlayAnim Request Arguments
  *
@@ -89,7 +115,7 @@ struct FAnimNextPlayAnimRequestArgs
 
 	// Payload that will be applied to the animation graph's variables via its data interfaces.
 	UPROPERTY()
-	FAnimNextDataInterfacePayload Payload;
+	FInstancedStruct Payload;
 };
 
 namespace UE::AnimNext
diff --git a/Engine/Plugins/Experimental/AnimNextAnimGraph/Source/AnimNextAnimGraph/Public/TraitInterfaces/IBlendStack.h b/Engine/Plugins/Experimental/AnimNextAnimGraph/Source/AnimNextAnimGraph/Public/TraitInterfaces/IBlendStack.h
--- a/Engine/Plugins/Experimental/AnimNextAnimGraph/Source/AnimNextAnimGraph/Public/TraitInterfaces/IBlendStack.h
+++ b/Engine/Plugins/Experimental/AnimNextAnimGraph/Source/AnimNextAnimGraph/Public/TraitInterfaces/IBlendStack.h
@@ -6,8 +6,6 @@
 #include "TraitCore/ITraitInterface.h"
 #include "TraitCore/TraitBinding.h"
 #include "AlphaBlend.h"
-#include "AnimNextDataInterfacePayload.h"
-#include "IPayloadFactory.h"
 #include "Graph/AnimNextAnimationGraph.h"
 #include "Graph/AnimNextGraphInstancePtr.h"
 
@@ -29,18 +27,14 @@ namespace UE::AnimNext
 			 * Consider making dynamic payload if we want to implement a special blend framework.
 			 */
 
-			// Payload applied to the AnimationGraph's variables via its data interfaces.
-			FAnimNextDataInterfacePayload GraphPayload;
+			// Payloads applied to the AnimationGraph's variables via its data interfaces.
+			TArray<FInstancedStruct> Payloads;
 
 			// Blend-in duration for this graph request 
 			FAlphaBlendArgs BlendArgs;
 
 			// The template graph to use for the new graph instance
 			TObjectPtr<const UAnimNextAnimationGraph> AnimationGraph;
-
-			// The object that was used for factory generation of the graph.
-			// Used as a key to allow identification of this request without inspecting payloads.
-			TObjectPtr<const UObject> FactoryObject;
 		};
 
 		typedef FGraphRequest* FGraphRequestPtr;
diff --git a/Engine/Plugins/Experimental/AnimNextAnimGraph/Source/AnimNextAnimGraph/Public/TraitInterfaces/IGraphFactory.h b/Engine/Plugins/Experimental/AnimNextAnimGraph/Source/AnimNextAnimGraph/Public/TraitInterfaces/IGraphFactory.h
deleted file mode 100644
diff --git a/Engine/Plugins/Experimental/AnimNextAnimGraph/Source/AnimNextAnimGraph/Public/TraitInterfaces/IPayloadFactory.h b/Engine/Plugins/Experimental/AnimNextAnimGraph/Source/AnimNextAnimGraph/Public/TraitInterfaces/IPayloadFactory.h
deleted file mode 100644
diff --git a/Engine/Plugins/Experimental/AnimNextAnimGraph/Source/AnimNextAnimGraphEditor/AnimNextAnimGraphEditor.Build.cs b/Engine/Plugins/Experimental/AnimNextAnimGraph/Source/AnimNextAnimGraphEditor/AnimNextAnimGraphEditor.Build.cs
deleted file mode 100644
diff --git a/Engine/Plugins/Experimental/AnimNextAnimGraph/Source/AnimNextAnimGraphEditor/Private/AnimNextAnimGraphEditorModule.cpp b/Engine/Plugins/Experimental/AnimNextAnimGraph/Source/AnimNextAnimGraphEditor/Private/AnimNextAnimGraphEditorModule.cpp
deleted file mode 100644
diff --git a/Engine/Plugins/Experimental/AnimNextChooser/AnimNextChooser.uplugin b/Engine/Plugins/Experimental/AnimNextChooser/AnimNextChooser.uplugin
--- a/Engine/Plugins/Experimental/AnimNextChooser/AnimNextChooser.uplugin
+++ b/Engine/Plugins/Experimental/AnimNextChooser/AnimNextChooser.uplugin
@@ -21,10 +21,6 @@
 			"Name": "AnimNext",
 			"Enabled": true
 		},
-		{
-			"Name": "AnimNextAnimGraph",
-			"Enabled": true
-		},
 		{
 			"Name": "Chooser",
 			"Enabled": true
diff --git a/Engine/Plugins/Experimental/AnimNextChooser/Source/AnimNextChooser/AnimNextChooser.Build.cs b/Engine/Plugins/Experimental/AnimNextChooser/Source/AnimNextChooser/AnimNextChooser.Build.cs
--- a/Engine/Plugins/Experimental/AnimNextChooser/Source/AnimNextChooser/AnimNextChooser.Build.cs
+++ b/Engine/Plugins/Experimental/AnimNextChooser/Source/AnimNextChooser/AnimNextChooser.Build.cs
@@ -6,7 +6,7 @@ namespace UnrealBuildTool.Rules
 	{
 		public AnimNextChooser(ReadOnlyTargetRules Target) : base(Target)
 		{
-			PrivateDependencyModuleNames.AddRange(
+			PublicDependencyModuleNames.AddRange(
 				new string[]
 				{
 					"Core",
@@ -16,7 +16,6 @@ namespace UnrealBuildTool.Rules
 					"Engine",
 					"AnimNext",
 					"Chooser",
-					"AnimNextAnimGraph",
 				}
 			);
 		}
diff --git a/Engine/Plugins/Experimental/AnimNextChooser/Source/AnimNextChooser/Private/MakeGraphFromChooser.cpp b/Engine/Plugins/Experimental/AnimNextChooser/Source/AnimNextChooser/Private/MakeGraphFromChooser.cpp
deleted file mode 100644
diff --git a/Engine/Plugins/Experimental/AnimNextChooser/Source/AnimNextChooser/Private/MakeGraphFromChooser.h b/Engine/Plugins/Experimental/AnimNextChooser/Source/AnimNextChooser/Private/MakeGraphFromChooser.h
deleted file mode 100644
diff --git a/Engine/Plugins/Experimental/AnimNextChooser/Source/AnimNextChooser/Private/Module.cpp b/Engine/Plugins/Experimental/AnimNextChooser/Source/AnimNextChooser/Private/Module.cpp
--- a/Engine/Plugins/Experimental/AnimNextChooser/Source/AnimNextChooser/Private/Module.cpp
+++ b/Engine/Plugins/Experimental/AnimNextChooser/Source/AnimNextChooser/Private/Module.cpp
@@ -19,12 +19,6 @@ class FModule : public IModuleInterface
 public:
 	virtual void StartupModule() override
 	{
-		static TPair<UClass*, FRigVMRegistry::ERegisterObjectOperation> const AllowedObjectTypes[] =
-		{
-			{ UChooserTable::StaticClass(), FRigVMRegistry::ERegisterObjectOperation::Class },
-		};
-
-		FRigVMRegistry::Get().RegisterObjectTypes(AllowedObjectTypes);
 	}
 
 	virtual void ShutdownModule() override
diff --git a/Engine/Plugins/Experimental/AnimNextPoseSearch/Source/AnimNextPoseSearch/Private/MotionMatchingTrait.cpp b/Engine/Plugins/Experimental/AnimNextPoseSearch/Source/AnimNextPoseSearch/Private/MotionMatchingTrait.cpp
--- a/Engine/Plugins/Experimental/AnimNextPoseSearch/Source/AnimNextPoseSearch/Private/MotionMatchingTrait.cpp
+++ b/Engine/Plugins/Experimental/AnimNextPoseSearch/Source/AnimNextPoseSearch/Private/MotionMatchingTrait.cpp
@@ -1,17 +1,13 @@
 // Copyright Epic Games, Inc. All Rights Reserved.
 
 #include "MotionMatchingTrait.h"
-
-#include "AnimNextAnimGraphSettings.h"
 #include "Animation/AnimSequence.h"
 #include "DataInterface/DataInterfaceStructAdapter.h"
 #include "IPoseHistory.h"
-#include "GraphInterfaces/AnimNextNativeDataInterface_AnimSequencePlayer.h"
+#include "PlayAnim/PlayAnimRequest.h"
 #include "PoseSearch/PoseSearchDatabase.h"
 #include "PoseSearch/PoseSearchDerivedData.h"
 #include "TraitInterfaces/IBlendStack.h"
-#include "TraitInterfaces/IGraphFactory.h"
-#include "TraitInterfaces/ITimeline.h"
 
 namespace UE::AnimNext
 {
@@ -156,28 +152,19 @@ namespace UE::AnimNext
 			{
 				if (const FPoseSearchDatabaseAnimationAssetBase* DatabaseAnimationAssetBase = MotionMatchingState.CurrentSearchResult.Database->GetDatabaseAnimationAsset<FPoseSearchDatabaseAnimationAssetBase>(*SearchIndexAsset))
 				{
-					const UObject* AnimationAsset = DatabaseAnimationAssetBase->GetAnimationAsset();
-					
-					UE::AnimNext::IBlendStack::FGraphRequest NewGraphRequest;
-					NewGraphRequest.BlendArgs = SharedData->BlendArguments;
-
-					const UAnimNextAnimationGraph* AnimationGraph = IGraphFactory::GetGraphFromObjectWithFallback(Context, Binding, AnimationAsset);
-					if(AnimationGraph != nullptr)
+					if (UAnimSequence* AnimSequence = Cast<UAnimSequence>(DatabaseAnimationAssetBase->GetAnimationAsset()))
 					{
-						// Create a default 'native data interface' payload via which we expect to communicate with the graph 
-						FInstancedStruct NewPayload;
-						NewPayload.InitializeAs<FAnimNextNativeDataInterface_AnimSequencePlayer>();
-						FAnimNextNativeDataInterface_AnimSequencePlayer& NewNativeInterface = NewPayload.GetMutable<FAnimNextNativeDataInterface_AnimSequencePlayer>();
-						NewNativeInterface.AnimSequence = Cast<UAnimSequence>(AnimationAsset);
-						NewNativeInterface.PlayRate = MotionMatchingState.WantedPlayRate;
-						NewNativeInterface.StartPosition = MotionMatchingState.CurrentSearchResult.AssetTime;
-						NewNativeInterface.Loop = SearchIndexAsset->IsLooping();
-
-						NewGraphRequest.GraphPayload.AddNative(MoveTemp(NewPayload));
-
-						IPayloadFactory::GetPayloadForGraphWithFallback(Context, Binding, AnimationAsset, AnimationGraph, NewGraphRequest.GraphPayload);
-
-						NewGraphRequest.AnimationGraph = AnimationGraph;
+						UE::AnimNext::IBlendStack::FGraphRequest NewGraphRequest;
+						NewGraphRequest.BlendArgs = SharedData->BlendArguments;
+						NewGraphRequest.AnimationGraph = SharedData->AnimationGraph;
+
+						FInstancedStruct& NewPayload = NewGraphRequest.Payloads.AddDefaulted_GetRef();
+						NewPayload.InitializeAs<FAnimNextPlayAnimPayload>();
+						FAnimNextPlayAnimPayload& NewPlayAnimPayload = NewPayload.GetMutable<FAnimNextPlayAnimPayload>();
+						NewPlayAnimPayload.AnimationObject = AnimSequence;
+						NewPlayAnimPayload.PlayRate = MotionMatchingState.WantedPlayRate;
+						NewPlayAnimPayload.StartPosition = MotionMatchingState.CurrentSearchResult.AssetTime;
+						NewPlayAnimPayload.IsLooping = SearchIndexAsset->IsLooping();
 
 						BlendStackBinding.PushGraph(Context, MoveTemp(NewGraphRequest));
 					}
diff --git a/Engine/Plugins/Experimental/AnimNextPoseSearch/Source/AnimNextPoseSearch/Private/MotionMatchingTrait.h b/Engine/Plugins/Experimental/AnimNextPoseSearch/Source/AnimNextPoseSearch/Private/MotionMatchingTrait.h
--- a/Engine/Plugins/Experimental/AnimNextPoseSearch/Source/AnimNextPoseSearch/Private/MotionMatchingTrait.h
+++ b/Engine/Plugins/Experimental/AnimNextPoseSearch/Source/AnimNextPoseSearch/Private/MotionMatchingTrait.h
@@ -14,6 +14,10 @@ struct FMotionMatchingTraitSharedData : public FAnimNextTraitSharedData
 {
 	GENERATED_BODY()
 
+	// Asset to use as a sub-graph
+	UPROPERTY(EditAnywhere, Category = "Default")
+	TObjectPtr<const UAnimNextAnimationGraph> AnimationGraph;
+
 	// Entry point in the Subgraph that we will use
 	UPROPERTY(EditAnywhere, Category = "Default", meta=(CustomWidget = "ParamName", AllowedParamType = "FAnimNextEntryPoint"))
 	FName EntryPoint;
@@ -57,6 +61,7 @@ struct FMotionMatchingTraitSharedData : public FAnimNextTraitSharedData
 
 	// Latent pin support boilerplate
 	#define TRAIT_LATENT_PROPERTIES_ENUMERATOR(GeneratorMacro) \
+		GeneratorMacro(AnimationGraph) \
 		GeneratorMacro(EntryPoint) \
 		GeneratorMacro(Databases) \
 		GeneratorMacro(BlendArguments) \
diff --git a/Engine/Plugins/Experimental/AnimNextStateTree/Source/AnimNextStateTree/AnimNextStateTree.Build.cs b/Engine/Plugins/Experimental/AnimNextStateTree/Source/AnimNextStateTree/AnimNextStateTree.Build.cs
--- a/Engine/Plugins/Experimental/AnimNextStateTree/Source/AnimNextStateTree/AnimNextStateTree.Build.cs
+++ b/Engine/Plugins/Experimental/AnimNextStateTree/Source/AnimNextStateTree/AnimNextStateTree.Build.cs
@@ -20,6 +20,7 @@ namespace UnrealBuildTool.Rules
 					"CoreUObject",
 					"StateTreeModule",
 					"Engine",
+					"Chooser",
 					"AnimNext", 
 				}
 			);
diff --git a/Engine/Plugins/Experimental/AnimNextStateTree/Source/AnimNextStateTree/Private/AnimNextStateTreeContext.cpp b/Engine/Plugins/Experimental/AnimNextStateTree/Source/AnimNextStateTree/Private/AnimNextStateTreeContext.cpp
--- a/Engine/Plugins/Experimental/AnimNextStateTree/Source/AnimNextStateTree/Private/AnimNextStateTreeContext.cpp
+++ b/Engine/Plugins/Experimental/AnimNextStateTree/Source/AnimNextStateTree/Private/AnimNextStateTreeContext.cpp
@@ -2,55 +2,47 @@
 
 #include "AnimNextStateTreeContext.h"
 
-#include "AnimNextAnimGraphSettings.h"
+#include "DataInterface/DataInterfaceStructAdapter.h"
 #include "TraitCore/TraitStackBinding.h"
 #include "TraitInterfaces/IBlendStack.h"
-#include "TraitInterfaces/IGraphFactory.h"
 #include "TraitInterfaces/ISmoothBlend.h"
 #include "TraitInterfaces/ITimeline.h"
 
-bool FAnimNextStateTreeTraitContext::PushAssetOntoBlendStack(TNonNullPtr<UObject> InAsset, const FAlphaBlendArgs& InBlendArguments, FAnimNextDataInterfacePayload&& InPayload) const
-{
-	using namespace UE::AnimNext;
-
-	if (Binding == nullptr || Context == nullptr)
-	{
-		return false;
-	}
-
-	TTraitBinding<IBlendStack> BlendStackBinding;
-	if(!Binding->GetStackInterface<IBlendStack>(BlendStackBinding))
-	{
-		return false;
-	}
-
-	// Get a graph to instantiate
-	const UAnimNextAnimationGraph* AnimationGraph = IGraphFactory::GetGraphFromObjectWithFallback(*Context, *Binding, InAsset);
-	if(AnimationGraph == nullptr)
+bool FAnimNextStateTreeTraitContext::PushAnimationGraphOntoBlendStack(TNonNullPtr<UAnimNextAnimationGraph> InAnimationGraph, const FAlphaBlendArgs& InBlendArguments, TArrayView<const FInstancedStruct> PayloadStructs) const
+{		
+	if (Binding && Context)
 	{
-		return false;
+		UE::AnimNext::TTraitBinding<UE::AnimNext::IBlendStack> BlendStackBinding;
+		if(Binding->GetInterface<UE::AnimNext::IBlendStack>(BlendStackBinding))
+		{
+			UE::AnimNext::IBlendStack::FGraphRequestPtr ActiveGraphRequest;
+			BlendStackBinding.GetActiveGraph(*Context, ActiveGraphRequest);
+
+			if (ActiveGraphRequest && ActiveGraphRequest->AnimationGraph == InAnimationGraph)
+			{
+				return true;
+			}
+			
+			UE::AnimNext::IBlendStack::FGraphRequest NewGraphRequest;
+			NewGraphRequest.BlendArgs = InBlendArguments;
+			NewGraphRequest.AnimationGraph = InAnimationGraph.Get();
+			NewGraphRequest.Payloads.Append(PayloadStructs);
+			
+			BlendStackBinding.PushGraph(*Context, MoveTemp(NewGraphRequest));
+
+			return true;
+		}
 	}
 
-	// Get a payload to instantiate to control the graph
-	IPayloadFactory::GetPayloadForGraphWithFallback(*Context, *Binding, InAsset, AnimationGraph, InPayload);
-
-	IBlendStack::FGraphRequest NewGraphRequest;
-	NewGraphRequest.BlendArgs = InBlendArguments;
-	NewGraphRequest.FactoryObject = InAsset;
-	NewGraphRequest.AnimationGraph = AnimationGraph;
-	NewGraphRequest.GraphPayload = MoveTemp(InPayload);
-
-	BlendStackBinding.PushGraph(*Context, MoveTemp(NewGraphRequest));
-
-	return true;
+	return false;
 }
 
-float FAnimNextStateTreeTraitContext::QueryPlaybackRatio(TNonNullPtr<UObject> InAsset) const
+float FAnimNextStateTreeTraitContext::QueryPlaybackRatio(TNonNullPtr<UAnimNextAnimationGraph> InAnimationGraph) const
 {
 	if (Binding && Context)
 	{
 		UE::AnimNext::TTraitBinding<UE::AnimNext::ITimeline> Timeline;
-		if (Binding->GetStackInterface<UE::AnimNext::ITimeline>(Timeline))
+		if (Binding->GetInterface<UE::AnimNext::ITimeline>(Timeline))
 		{
 			UE::AnimNext::FTimelineProgress Progress = Timeline.GetProgress(*Context);
 
diff --git a/Engine/Plugins/Experimental/AnimNextStateTree/Source/AnimNextStateTree/Private/AnimNextStateTreeGraphInstanceTask.cpp b/Engine/Plugins/Experimental/AnimNextStateTree/Source/AnimNextStateTree/Private/AnimNextStateTreeGraphInstanceTask.cpp
--- a/Engine/Plugins/Experimental/AnimNextStateTree/Source/AnimNextStateTree/Private/AnimNextStateTreeGraphInstanceTask.cpp
+++ b/Engine/Plugins/Experimental/AnimNextStateTree/Source/AnimNextStateTree/Private/AnimNextStateTreeGraphInstanceTask.cpp
@@ -1,16 +1,12 @@
 // Copyright Epic Games, Inc. All Rights Reserved.
 
 #include "Tasks/AnimNextStateTreeGraphInstanceTask.h"
-
-#include "AnimNextAnimGraphSettings.h"
 #include "StateTreeLinker.h"
 #include "StateTreeExecutionContext.h"
-#include "AnimNextDataInterfacePayload.h"
+#include "DataInterface/DataInterfaceStructAdapter.h"
 
 FAnimNextStateTreeGraphInstanceTask::FAnimNextStateTreeGraphInstanceTask()
 {
-	// Re-selecting the same state should not cause a re-trigger of EnterState()
-	bShouldStateChangeOnReselect = false;
 }
 
 bool FAnimNextStateTreeGraphInstanceTask::Link(FStateTreeLinker& Linker)
@@ -24,17 +20,12 @@ EStateTreeRunStatus FAnimNextStateTreeGraphInstanceTask::EnterState(FStateTreeEx
 	FAnimNextStateTreeTraitContext& ExecContext = Context.GetExternalData(TraitContextHandle);
 	FInstanceDataType& InstanceData = Context.GetInstanceData(*this);
 
-	// Reference task's default payload
-	FAnimNextDataInterfacePayload Payload;
-	Payload.AddStructView(InstanceData.Payload.GetMutableValue());
-
-	FInstancedStruct NewNativePayload;
-	NewNativePayload.InitializeAs(Context.GetMutableInstanceData()->GetMutableStorage().GetMutableGlobalParameters().GetScriptStruct(), Context.GetMutableInstanceData()->GetMutableStorage().GetMutableGlobalParameters().GetMemory());
-	Payload.AddNative(MoveTemp(NewNativePayload));
+	FInstancedStruct Payload;
+	Payload.InitializeAs(Context.GetMutableInstanceData()->GetMutableStorage().GetMutableGlobalParameters().GetScriptStruct(), Context.GetMutableInstanceData()->GetMutableStorage().GetMutableGlobalParameters().GetMemory());
 
-	if (ExecContext.PushAssetOntoBlendStack(InstanceData.Asset.Get(), InstanceData.BlendOptions, MoveTemp(Payload)))
+	if(ExecContext.PushAnimationGraphOntoBlendStack(InstanceData.AnimationGraph.Get(), InstanceData.BlendOptions, MakeConstArrayView({Payload})))
 	{
-		return EStateTreeRunStatus::Running;
+		return EStateTreeRunStatus::Running;	
 	}
 	
 	return EStateTreeRunStatus::Failed;
@@ -47,7 +38,7 @@ EStateTreeRunStatus FAnimNextStateTreeGraphInstanceTask::Tick(FStateTreeExecutio
 	{
 		FAnimNextStateTreeTraitContext& ExecContext = Context.GetExternalData(TraitContextHandle);
 
-		InstanceData.PlaybackRatio = ExecContext.QueryPlaybackRatio(InstanceData.Asset.Get());
+		InstanceData.PlaybackRatio = ExecContext.QueryPlaybackRatio(InstanceData.AnimationGraph.Get());
 
 		return EStateTreeRunStatus::Running;
 	}
@@ -59,22 +50,3 @@ void FAnimNextStateTreeGraphInstanceTask::ExitState(FStateTreeExecutionContext&
 {
 	FAnimNextStateTreeTaskBase::ExitState(Context, Transition);
 }
-
-#if WITH_EDITOR
-void FAnimNextStateTreeGraphInstanceTask::PostEditInstanceDataChangeChainProperty(const FPropertyChangedChainEvent& PropertyChangedEvent, FStateTreeDataView InstanceDataView)
-{
-	const FProperty* Property = PropertyChangedEvent.Property;
-	if(Property == nullptr)
-	{
-		return;
-	}
-
-	if(Property->GetFName() == GET_MEMBER_NAME_CHECKED(FAnimNextGraphInstanceTaskInstanceData, Asset))
-	{
-		// Repopulate payload property bag if the selected asset changes
-		FAnimNextGraphInstanceTaskInstanceData& InstanceData = InstanceDataView.GetMutable<FAnimNextGraphInstanceTaskInstanceData>();
-		const UAnimNextAnimationGraph* AnimationGraph = GetDefault<UAnimNextAnimGraphSettings>()->GetGraphFromObject(InstanceData.Asset);
-		GetDefault<UAnimNextAnimGraphSettings>()->GetNonNativePayloadFromGraph(InstanceData.Asset, AnimationGraph, InstanceData.Payload);
-	}
-}
-#endif
\ No newline at end of file
diff --git a/Engine/Plugins/Experimental/AnimNextStateTree/Source/AnimNextStateTree/Private/AnimStateTreeTrait.cpp b/Engine/Plugins/Experimental/AnimNextStateTree/Source/AnimNextStateTree/Private/AnimStateTreeTrait.cpp
--- a/Engine/Plugins/Experimental/AnimNextStateTree/Source/AnimNextStateTree/Private/AnimStateTreeTrait.cpp
+++ b/Engine/Plugins/Experimental/AnimNextStateTree/Source/AnimNextStateTree/Private/AnimStateTreeTrait.cpp
@@ -45,7 +45,7 @@ namespace UE::AnimNext
 			UObject* Owner = GetTransientPackage();
 			FStateTreeExecutionContext StateTreeExecutionContext(*Owner, *InstanceData->StateTree, InstanceData->InstanceData);
 
-			FAnimNextStateTreeTraitContext TraitContext(Context, &Binding);
+			FAnimNextStateTreeTraitContext TraitContext(Context, Binding.GetStack());
 			StateTreeExecutionContext.SetContextDataByName(UStateTreeAnimNextSchema::AnimStateTreeExecutionContextName, FStateTreeDataView(FAnimNextStateTreeTraitContext::StaticStruct(), reinterpret_cast<uint8*>(&TraitContext)));
 
 			FAnimNextGraphInstance* OwnerGraphInstance = &Binding.GetTraitPtr().GetNodeInstance()->GetOwner();
@@ -107,7 +107,7 @@ namespace UE::AnimNext
 			FAnimNextGraphInstance* OwnerGraphInstance = &Binding.GetTraitPtr().GetNodeInstance()->GetOwner();
 			if(OwnerGraphInstance && InstanceData->StateTree->GetDefaultParameters().GetPropertyBagStruct() == OwnerGraphInstance->GetVariables().GetPropertyBagStruct())
 			{
-				FAnimNextStateTreeTraitContext TraitContext(Context, &Binding);
+				FAnimNextStateTreeTraitContext TraitContext(Context, Binding.GetStack());
 				StateTreeExecutionContext.SetContextDataByName(UStateTreeAnimNextSchema::AnimStateTreeExecutionContextName, FStateTreeDataView(FAnimNextStateTreeTraitContext::StaticStruct(), reinterpret_cast<uint8*>(&TraitContext)));
 
 				const FInstancedPropertyBag& StateTreeParameters = InstanceData->StateTree->GetDefaultParameters();
diff --git a/Engine/Plugins/Experimental/AnimNextStateTree/Source/AnimNextStateTree/Public/AnimNextStateTreeContext.h b/Engine/Plugins/Experimental/AnimNextStateTree/Source/AnimNextStateTree/Public/AnimNextStateTreeContext.h
--- a/Engine/Plugins/Experimental/AnimNextStateTree/Source/AnimNextStateTree/Public/AnimNextStateTreeContext.h
+++ b/Engine/Plugins/Experimental/AnimNextStateTree/Source/AnimNextStateTree/Public/AnimNextStateTreeContext.h
@@ -9,14 +9,13 @@
 
 namespace UE::AnimNext
 {
-struct FTraitBinding;
+struct FTraitStackBinding;
 struct FExecutionContext;
 struct FStateTreeTrait;
 }
 
 class UAnimNextAnimationGraph;
 struct FAlphaBlendArgs;
-struct FAnimNextDataInterfacePayload;
 
 USTRUCT()
 struct FAnimNextStateTreeTraitContext
@@ -27,12 +26,12 @@ struct FAnimNextStateTreeTraitContext
 	
 	FAnimNextStateTreeTraitContext() {}
 
-	bool PushAssetOntoBlendStack(TNonNullPtr<UObject> InAsset, const FAlphaBlendArgs& InBlendArguments, FAnimNextDataInterfacePayload&& InPayload) const;
-	float QueryPlaybackRatio(TNonNullPtr<UObject> InAsset) const;
+	bool PushAnimationGraphOntoBlendStack(TNonNullPtr<UAnimNextAnimationGraph> InAnimationGraph, const FAlphaBlendArgs& InBlendArguments, const TArrayView<const FInstancedStruct> PayloadStructs) const;
+	float QueryPlaybackRatio(TNonNullPtr<UAnimNextAnimationGraph> InAnimationGraph) const;
 protected:
-	FAnimNextStateTreeTraitContext(UE::AnimNext::FExecutionContext& InContext, const UE::AnimNext::FTraitBinding* InBinding) : Context(&InContext), Binding(InBinding) {}
+	FAnimNextStateTreeTraitContext(UE::AnimNext::FExecutionContext& InContext, const UE::AnimNext::FTraitStackBinding* InBinding) : Context(&InContext), Binding(InBinding) {}
 protected:
 	UE::AnimNext::FExecutionContext* Context = nullptr;
-	const UE::AnimNext::FTraitBinding* Binding = nullptr;
+	const UE::AnimNext::FTraitStackBinding* Binding = nullptr;
 };
 
diff --git a/Engine/Plugins/Experimental/AnimNextStateTree/Source/AnimNextStateTree/Public/Tasks/AnimNextStateTreeGraphInstanceTask.h b/Engine/Plugins/Experimental/AnimNextStateTree/Source/AnimNextStateTree/Public/Tasks/AnimNextStateTreeGraphInstanceTask.h
--- a/Engine/Plugins/Experimental/AnimNextStateTree/Source/AnimNextStateTree/Public/Tasks/AnimNextStateTreeGraphInstanceTask.h
+++ b/Engine/Plugins/Experimental/AnimNextStateTree/Source/AnimNextStateTree/Public/Tasks/AnimNextStateTreeGraphInstanceTask.h
@@ -6,7 +6,6 @@
 #include "AnimNextStateTreeTypes.h"
 #include "Graph/AnimNextAnimationGraph.h"
 #include "AlphaBlend.h"
-#include "StructUtils/PropertyBag.h"
 
 #include "AnimNextStateTreeGraphInstanceTask.generated.h"
 
@@ -14,24 +13,16 @@ USTRUCT()
 struct ANIMNEXTSTATETREE_API FAnimNextGraphInstanceTaskInstanceData
 {
 	GENERATED_BODY()
+	
+	UPROPERTY(EditAnywhere, Category = Animation)
+	TObjectPtr<UAnimNextAnimationGraph> AnimationGraph = nullptr;
 
-	// The asset to instantiate
-	UPROPERTY(EditAnywhere, Category = Animation, meta=(GetAllowedClasses="/Script/AnimNextAnimGraph.AnimNextAnimGraphSettings:GetAllowedAssetClasses"))
-	TObjectPtr<UObject> Asset = nullptr;
-
-	// The payload to use for the asset when instanced
-	UPROPERTY(EditAnywhere, Category = Input, meta = (Optional, FixedLayout))
-	FInstancedPropertyBag Payload;
-
-	// Blend options for when the state is pushed
 	UPROPERTY(EditAnywhere, Category = Animation)
 	FAlphaBlendArgs BlendOptions;
 
-	// Whether this task should continue to tick once state is entered
 	UPROPERTY(EditAnywhere, Category = Animation)
 	bool bContinueTicking = true;
 
-	// Current playback ratio (debug)
 	UPROPERTY(VisibleAnywhere, Category = Animation)
 	float PlaybackRatio = 1.0f;
 };
@@ -52,9 +43,7 @@ protected:
 	virtual EStateTreeRunStatus EnterState(FStateTreeExecutionContext& Context, const FStateTreeTransitionResult& Transition) const override;
 	virtual EStateTreeRunStatus Tick(FStateTreeExecutionContext& Context, const float DeltaTime) const override;
 	virtual void ExitState(FStateTreeExecutionContext& Context, const FStateTreeTransitionResult& Transition) const override;
-#if WITH_EDITOR
-	virtual void PostEditInstanceDataChangeChainProperty(const FPropertyChangedChainEvent& PropertyChangedEvent, FStateTreeDataView InstanceDataView) override;
-#endif
+
 public:
 	TStateTreeExternalDataHandle<FAnimNextStateTreeTraitContext> TraitContextHandle;
 };
diff --git a/Engine/Plugins/Experimental/AnimNextStateTree/Source/AnimNextStateTreeEditor/AnimNextStateTreeEditor.Build.cs b/Engine/Plugins/Experimental/AnimNextStateTree/Source/AnimNextStateTreeEditor/AnimNextStateTreeEditor.Build.cs
--- a/Engine/Plugins/Experimental/AnimNextStateTree/Source/AnimNextStateTreeEditor/AnimNextStateTreeEditor.Build.cs
+++ b/Engine/Plugins/Experimental/AnimNextStateTree/Source/AnimNextStateTreeEditor/AnimNextStateTreeEditor.Build.cs
@@ -18,9 +18,7 @@ namespace UnrealBuildTool.Rules
 				    "UnrealEd",
 				    "AssetDefinition",
 				    "AnimNextAnimGraph",
-				    "RigVM",
-				    "StructUtilsEditor",
-				    "PropertyEditor",
+				    "RigVM"
 			    });
 
 			PublicDependencyModuleNames.AddRange(
