AnimNext: Removed dummy struct used for property-export

Replaced with direct call to UScriptStruct::ExportText

diff --git a/Engine/Plugins/Experimental/AnimNext/Source/AnimNext/Internal/Graph/RigDecorator_AnimNextCppTrait.h b/Engine/Plugins/Experimental/AnimNext/Source/AnimNext/Internal/Graph/RigDecorator_AnimNextCppTrait.h
--- a/Engine/Plugins/Experimental/AnimNext/Source/AnimNext/Internal/Graph/RigDecorator_AnimNextCppTrait.h
+++ b/Engine/Plugins/Experimental/AnimNext/Source/AnimNext/Internal/Graph/RigDecorator_AnimNextCppTrait.h
@@ -29,12 +29,3 @@ struct ANIMNEXT_API FRigDecorator_AnimNextCppDecorator : public FRigVMTrait
 	const UE::AnimNext::FTrait* GetTrait() const;
 #endif
 };
-
-USTRUCT(BlueprintType)
-struct FAnimNextCppDecoratorWrapper
-{
-	GENERATED_BODY()
-
-	UPROPERTY()
-	FRigDecorator_AnimNextCppDecorator CppDecorator;
-};
diff --git a/Engine/Plugins/Experimental/AnimNext/Source/AnimNextUncookedOnly/Private/AnimNextController.cpp b/Engine/Plugins/Experimental/AnimNext/Source/AnimNextUncookedOnly/Private/AnimNextController.cpp
--- a/Engine/Plugins/Experimental/AnimNext/Source/AnimNextUncookedOnly/Private/AnimNextController.cpp
+++ b/Engine/Plugins/Experimental/AnimNext/Source/AnimNextUncookedOnly/Private/AnimNextController.cpp
@@ -63,8 +63,7 @@ FName UAnimNextController::AddTraitByName(FName InNodeName, FName InNewTraitType
 			return NAME_None;	// This trait isn't supported on this node
 		}
 
-		const FProperty* Prop = FAnimNextCppDecoratorWrapper::StaticStruct()->FindPropertyByName(GET_MEMBER_NAME_STRING_CHECKED(FAnimNextCppDecoratorWrapper, CppDecorator));
-		Prop->ExportText_Direct(DefaultValue, &CppDecoratorStructInstance, &DefaultCppDecoratorStructInstance, nullptr, PPF_SerializedAsImportText);
+		FRigDecorator_AnimNextCppDecorator::StaticStruct()->ExportText(DefaultValue, &CppDecoratorStructInstance, &DefaultCppDecoratorStructInstance, nullptr, PPF_SerializedAsImportText, nullptr);
 	}
 
 	// Avoid multiple VM recompilations for internal operations
diff --git a/Engine/Plugins/Experimental/AnimNext/Source/AnimNextUncookedOnly/Private/AnimNextEdGraphNode.cpp b/Engine/Plugins/Experimental/AnimNext/Source/AnimNextUncookedOnly/Private/AnimNextEdGraphNode.cpp
--- a/Engine/Plugins/Experimental/AnimNext/Source/AnimNextUncookedOnly/Private/AnimNextEdGraphNode.cpp
+++ b/Engine/Plugins/Experimental/AnimNext/Source/AnimNextUncookedOnly/Private/AnimNextEdGraphNode.cpp
@@ -152,8 +152,7 @@ void UAnimNextEdGraphNode::BuildAddTraitContextMenu(UToolMenu* SubMenu)
 				continue;	// This trait isn't supported on this node
 			}
 
-			const FProperty* Prop = FAnimNextCppDecoratorWrapper::StaticStruct()->FindPropertyByName(GET_MEMBER_NAME_STRING_CHECKED(FAnimNextCppDecoratorWrapper, CppDecorator));
-			Prop->ExportText_Direct(DefaultValue, &CppDecoratorStructInstance, &DefaultCppDecoratorStructInstance, nullptr, PPF_SerializedAsImportText);
+			FRigDecorator_AnimNextCppDecorator::StaticStruct()->ExportText(DefaultValue, &CppDecoratorStructInstance, &DefaultCppDecoratorStructInstance, nullptr, PPF_SerializedAsImportText, nullptr);
 		}
 
 		FString DisplayNameMetadata;
diff --git a/Engine/Plugins/Experimental/AnimNext/Source/AnimNextUncookedOnly/Private/Graph/AnimNextAnimationGraph_EditorData.cpp b/Engine/Plugins/Experimental/AnimNext/Source/AnimNextUncookedOnly/Private/Graph/AnimNextAnimationGraph_EditorData.cpp
--- a/Engine/Plugins/Experimental/AnimNext/Source/AnimNextUncookedOnly/Private/Graph/AnimNextAnimationGraph_EditorData.cpp
+++ b/Engine/Plugins/Experimental/AnimNext/Source/AnimNextUncookedOnly/Private/Graph/AnimNextAnimationGraph_EditorData.cpp
@@ -274,10 +274,7 @@ namespace UE::AnimNext::UncookedOnly::Private
 				FRigDecorator_AnimNextCppDecorator CppDecoratorStructInstance;
 				CppDecoratorStructInstance.DecoratorSharedDataStruct = Trait->GetTraitSharedDataStruct();
 
-				const FProperty* Prop = FAnimNextCppDecoratorWrapper::StaticStruct()->FindPropertyByName(GET_MEMBER_NAME_STRING_CHECKED(FAnimNextCppDecoratorWrapper, CppDecorator));
-				check(Prop != nullptr);
-
-				Prop->ExportText_Direct(DefaultValue, &CppDecoratorStructInstance, &DefaultCppDecoratorStructInstance, nullptr, PPF_SerializedAsImportText);
+				FRigDecorator_AnimNextCppDecorator::StaticStruct()->ExportText(DefaultValue, &CppDecoratorStructInstance, &DefaultCppDecoratorStructInstance, nullptr, PPF_SerializedAsImportText, nullptr);
 			}
 
 			const FName ReferencePoseDecoratorName = VMController->AddTrait(VMReferencePoseNode->GetFName(), *CppDecoratorStruct->GetPathName(), TEXT("ReferencePose"), DefaultValue, INDEX_NONE, false, false);
diff --git a/Engine/Plugins/Experimental/AnimNextAnimGraph/Source/AnimNextAnimGraphTestSuite/Private/AnimNextAnimGraphTraitGraphTest.cpp b/Engine/Plugins/Experimental/AnimNextAnimGraph/Source/AnimNextAnimGraphTestSuite/Private/AnimNextAnimGraphTraitGraphTest.cpp
--- a/Engine/Plugins/Experimental/AnimNextAnimGraph/Source/AnimNextAnimGraphTestSuite/Private/AnimNextAnimGraphTraitGraphTest.cpp
+++ b/Engine/Plugins/Experimental/AnimNextAnimGraph/Source/AnimNextAnimGraphTestSuite/Private/AnimNextAnimGraphTraitGraphTest.cpp
@@ -136,10 +136,7 @@ bool FAnimationAnimNextEditorTest_GraphAddTrait::RunTest(const FString& InParame
 
 			UE_RETURN_ON_ERROR(CppDecoratorStructInstance.CanBeAddedToNode(TraitStackNode, nullptr), "FAnimationAnimNextEditorTest_GraphAddTrait -> Trait cannot be added to trait stack node");
 
-			const FProperty* Prop = FAnimNextCppDecoratorWrapper::StaticStruct()->FindPropertyByName(GET_MEMBER_NAME_STRING_CHECKED(FAnimNextCppDecoratorWrapper, CppDecorator));
-			UE_RETURN_ON_ERROR(Prop != nullptr, "FAnimationAnimNextEditorTest_GraphAddTrait -> Failed to find wrapper property");
-
-			Prop->ExportText_Direct(DefaultValue, &CppDecoratorStructInstance, &DefaultCppDecoratorStructInstance, nullptr, PPF_None);
+			FRigDecorator_AnimNextCppDecorator::StaticStruct()->ExportText(DefaultValue, &CppDecoratorStructInstance, &DefaultCppDecoratorStructInstance, nullptr, PPF_None, nullptr);
 		}
 
 		FString DisplayNameMetadata;
@@ -430,10 +427,7 @@ bool FAnimationAnimNextRuntimeTest_GraphExecute::RunTest(const FString& InParame
 
 			UE_RETURN_ON_ERROR(CppDecoratorStructInstance.CanBeAddedToNode(DecoratorStackNode, nullptr), "FAnimationAnimNextRuntimeTest_GraphExecute -> Trait cannot be added to trait stack node");
 
-			const FProperty* Prop = FAnimNextCppDecoratorWrapper::StaticStruct()->FindPropertyByName(GET_MEMBER_NAME_STRING_CHECKED(FAnimNextCppDecoratorWrapper, CppDecorator));
-			UE_RETURN_ON_ERROR(Prop != nullptr, "FAnimationAnimNextRuntimeTest_GraphExecute -> Failed to find wrapper property");
-
-			Prop->ExportText_Direct(DefaultValue, &CppDecoratorStructInstance, &DefaultCppDecoratorStructInstance, nullptr, PPF_None);
+			FRigDecorator_AnimNextCppDecorator::StaticStruct()->ExportText(DefaultValue, &CppDecoratorStructInstance, &DefaultCppDecoratorStructInstance, nullptr, PPF_None, nullptr);
 		}
 
 		FString DisplayNameMetadata;
@@ -540,10 +534,7 @@ bool FAnimationAnimNextRuntimeTest_GraphExecuteLatent::RunTest(const FString& In
 
 			UE_RETURN_ON_ERROR(CppDecoratorStructInstance.CanBeAddedToNode(DecoratorStackNode, nullptr), "FAnimationAnimNextRuntimeTest_GraphExecuteLatent -> Trait cannot be added to trait stack node");
 
-			const FProperty* Prop = FAnimNextCppDecoratorWrapper::StaticStruct()->FindPropertyByName(GET_MEMBER_NAME_STRING_CHECKED(FAnimNextCppDecoratorWrapper, CppDecorator));
-			UE_RETURN_ON_ERROR(Prop != nullptr, "FAnimationAnimNextRuntimeTest_GraphExecuteLatent -> Failed to find wrapper property");
-
-			Prop->ExportText_Direct(DefaultValue, &CppDecoratorStructInstance, &DefaultCppDecoratorStructInstance, nullptr, PPF_None);
+			FRigDecorator_AnimNextCppDecorator::StaticStruct()->ExportText(DefaultValue, &CppDecoratorStructInstance, &DefaultCppDecoratorStructInstance, nullptr, PPF_None, nullptr);
 		}
 
 		FString DisplayNameMetadata;
