Add ability to feed a LOD pose to a Run Graph rig unit

Added an input pose trait that reads a LOD pose from a e.g. get variable node and pushes the pose onto the pose stack. The input pose trait can then be connected to any other node in an anim graph and be further used.

diff --git a/Engine/Plugins/Experimental/AnimNext/Source/AnimNext/Private/EvaluationVM/Tasks/PushPose.cpp b/Engine/Plugins/Experimental/AnimNext/Source/AnimNext/Private/EvaluationVM/Tasks/PushPose.cpp
new file mode 100644
--- /dev/null
+++ b/Engine/Plugins/Experimental/AnimNext/Source/AnimNext/Private/EvaluationVM/Tasks/PushPose.cpp
@@ -0,0 +1,33 @@
+// Copyright Epic Games, Inc. All Rights Reserved.
+
+#include "EvaluationVM/Tasks/PushPose.h"
+#include "EvaluationVM/EvaluationVM.h"
+
+FAnimNextPushPoseTask::FAnimNextPushPoseTask(const FAnimNextGraphLODPose* InGraphPose)
+	: GraphPose(InGraphPose)
+{
+}
+
+FAnimNextPushPoseTask FAnimNextPushPoseTask::Make(const FAnimNextGraphLODPose* InGraphPose)
+{
+	return FAnimNextPushPoseTask(InGraphPose);
+}
+
+FAnimNextPushPoseTask& FAnimNextPushPoseTask::operator=(const FAnimNextPushPoseTask& Other)
+{
+	GraphPose = Other.GraphPose;
+	return *this;
+}
+
+void FAnimNextPushPoseTask::Execute(UE::AnimNext::FEvaluationVM& VM) const
+{
+	using namespace UE::AnimNext;
+
+	FKeyframeState Keyframe = VM.MakeUninitializedKeyframe(GraphPose->LODPose.IsAdditive());
+
+	Keyframe.Pose.CopyFrom(GraphPose->LODPose);
+	Keyframe.Curves.CopyFrom(GraphPose->Curves);
+	Keyframe.Attributes.CopyFrom(GraphPose->Attributes);
+
+	VM.PushValue(KEYFRAME_STACK_NAME, MakeUnique<FKeyframeState>(MoveTemp(Keyframe)));
+}
\ No newline at end of file
diff --git a/Engine/Plugins/Experimental/AnimNext/Source/AnimNext/Public/EvaluationVM/Tasks/PushPose.h b/Engine/Plugins/Experimental/AnimNext/Source/AnimNext/Public/EvaluationVM/Tasks/PushPose.h
new file mode 100644
--- /dev/null
+++ b/Engine/Plugins/Experimental/AnimNext/Source/AnimNext/Public/EvaluationVM/Tasks/PushPose.h
@@ -0,0 +1,28 @@
+// Copyright Epic Games, Inc. All Rights Reserved.
+
+#pragma once
+
+#include "CoreMinimal.h"
+#include "EvaluationVM/EvaluationTask.h"
+#include "Graph/AnimNext_LODPose.h"
+
+#include "PushPose.generated.h"
+
+USTRUCT()
+struct ANIMNEXT_API FAnimNextPushPoseTask : public FAnimNextEvaluationTask
+{
+	GENERATED_BODY()
+
+	DECLARE_ANIM_EVALUATION_TASK(FAnimNextPushPoseTask)
+
+	FAnimNextPushPoseTask() = default;
+	FAnimNextPushPoseTask(const FAnimNextGraphLODPose* InGraphPose);
+
+	static FAnimNextPushPoseTask Make(const FAnimNextGraphLODPose* InGraphPose);
+
+	virtual void Execute(UE::AnimNext::FEvaluationVM& VM) const override;
+
+	FAnimNextPushPoseTask& operator=(const FAnimNextPushPoseTask& Other);
+
+	const FAnimNextGraphLODPose* GraphPose = nullptr;
+};
\ No newline at end of file
diff --git a/Engine/Plugins/Experimental/AnimNextAnimGraph/Source/AnimNextAnimGraph/Private/Traits/InputPoseTrait.cpp b/Engine/Plugins/Experimental/AnimNextAnimGraph/Source/AnimNextAnimGraph/Private/Traits/InputPoseTrait.cpp
new file mode 100644
--- /dev/null
+++ b/Engine/Plugins/Experimental/AnimNextAnimGraph/Source/AnimNextAnimGraph/Private/Traits/InputPoseTrait.cpp
@@ -0,0 +1,28 @@
+// Copyright Epic Games, Inc. All Rights Reserved.
+
+#include "Traits/InputPoseTrait.h"
+#include "TraitCore/ExecutionContext.h"
+#include "EvaluationVM/Tasks/PushPose.h"
+
+namespace UE::AnimNext
+{
+	AUTO_REGISTER_ANIM_TRAIT(FInputPoseTrait)
+
+		// Trait implementation boilerplate
+#define TRAIT_INTERFACE_ENUMERATOR(GeneratorMacro) \
+		GeneratorMacro(IEvaluate) \
+
+		GENERATE_ANIM_TRAIT_IMPLEMENTATION(FInputPoseTrait, TRAIT_INTERFACE_ENUMERATOR, NULL_ANIM_TRAIT_INTERFACE_ENUMERATOR, NULL_ANIM_TRAIT_EVENT_ENUMERATOR)
+#undef TRAIT_INTERFACE_ENUMERATOR
+
+	void FInputPoseTrait::PreEvaluate(FEvaluateTraversalContext& Context, const TTraitBinding<IEvaluate>& Binding) const
+	{
+		const FSharedData* SharedData = Binding.GetSharedData<FSharedData>();
+		const FAnimNextGraphLODPose& InputPose = SharedData->GetInput(Binding);
+
+		if (InputPose.LODPose.IsValid())
+		{
+			Context.AppendTask(FAnimNextPushPoseTask::Make(&InputPose));
+		}
+	}
+}
\ No newline at end of file
diff --git a/Engine/Plugins/Experimental/AnimNextAnimGraph/Source/AnimNextAnimGraph/Private/Traits/InputPoseTrait.h b/Engine/Plugins/Experimental/AnimNextAnimGraph/Source/AnimNextAnimGraph/Private/Traits/InputPoseTrait.h
new file mode 100644
--- /dev/null
+++ b/Engine/Plugins/Experimental/AnimNextAnimGraph/Source/AnimNextAnimGraph/Private/Traits/InputPoseTrait.h
@@ -0,0 +1,40 @@
+// Copyright Epic Games, Inc. All Rights Reserved.
+
+#pragma once
+
+#include "CoreMinimal.h"
+
+#include "TraitCore/Trait.h"
+#include "TraitInterfaces/IEvaluate.h"
+#include "Graph/AnimNext_LODPose.h"
+
+#include "InputPoseTrait.generated.h"
+
+USTRUCT(meta = (DisplayName = "Input Pose"))
+struct FAnimNextInputPoseTraitSharedData : public FAnimNextTraitSharedData
+{
+	GENERATED_BODY()
+
+	UPROPERTY()
+	FAnimNextGraphLODPose Input;
+
+	// Latent pin support boilerplate
+#define TRAIT_LATENT_PROPERTIES_ENUMERATOR(GeneratorMacro) \
+		GeneratorMacro(Input)
+
+	GENERATE_TRAIT_LATENT_PROPERTIES(FAnimNextInputPoseTraitSharedData, TRAIT_LATENT_PROPERTIES_ENUMERATOR)
+#undef TRAIT_LATENT_PROPERTIES_ENUMERATOR
+};
+
+namespace UE::AnimNext
+{
+	struct FInputPoseTrait : FBaseTrait, IEvaluate
+	{
+		DECLARE_ANIM_TRAIT(FInputPoseTrait, 0xb2974c37, FBaseTrait)
+
+		using FSharedData = FAnimNextInputPoseTraitSharedData;
+
+		// IEvaluate impl
+		virtual void PreEvaluate(FEvaluateTraversalContext& Context, const TTraitBinding<IEvaluate>& Binding) const override;
+	};
+}
\ No newline at end of file
