AnimNext: Move remaining animation graph code to AnimNextAnimGraph plugin

diff --git a/Engine/Plugins/Experimental/AnimNext/Source/AnimNextUncookedOnly/Private/AnimNextEdGraphNode.cpp b/Engine/Plugins/Experimental/AnimNext/Source/AnimNextUncookedOnly/Private/AnimNextEdGraphNode.cpp
--- a/Engine/Plugins/Experimental/AnimNext/Source/AnimNextUncookedOnly/Private/AnimNextEdGraphNode.cpp
+++ b/Engine/Plugins/Experimental/AnimNext/Source/AnimNextUncookedOnly/Private/AnimNextEdGraphNode.cpp
@@ -1,258 +1,3 @@
 ﻿// Copyright Epic Games, Inc. All Rights Reserved.
 
 #include "AnimNextEdGraphNode.h"
-
-#include "ContentBrowserModule.h"
-#include "IContentBrowserSingleton.h"
-#include "TraitCore/TraitHandle.h"
-#include "Graph/RigUnit_AnimNextTraitStack.h"
-#include "Graph/RigDecorator_AnimNextCppTrait.h"
-#include "RigVMModel/RigVMController.h"
-#include "TraitCore/TraitRegistry.h"
-#include "ToolMenu.h"
-#include "UncookedOnlyUtils.h"
-#include "Framework/Application/SlateApplication.h"
-#include "Graph/AnimNextAnimationGraph.h"
-#include "Graph/RigUnit_AnimNextRunAnimationGraph.h"
-#include "Graph/RigVMTrait_AnimNextPublicVariables.h"
-
-#define LOCTEXT_NAMESPACE "AnimNextEdGraphNode"
-
-namespace UE::AnimNext::Editor::Private
-{
-	static const FLazyName VariablesTraitBaseName = TEXT("Variables");
-}
-
-void UAnimNextEdGraphNode::GetNodeContextMenuActions(UToolMenu* Menu, class UGraphNodeContextMenuContext* Context) const
-{
-	using namespace UE::AnimNext::Editor::Private;
-	
-	URigVMEdGraphNode::GetNodeContextMenuActions(Menu, Context);
-
-	UAnimNextEdGraphNode* NonConstThis = const_cast<UAnimNextEdGraphNode*>(this);
-
-	if (IsTraitStack())
-	{
-		FToolMenuSection& Section = Menu->AddSection("AnimNextTraitNodeActions", LOCTEXT("AnimNextTraitNodeActionsMenuHeader", "Traits"));
-
-		Section.AddSubMenu(
-			TEXT("AddTraitMenu"),
-			LOCTEXT("AddTraitMenu", "Add Trait"),
-			LOCTEXT("AddTraitMenuTooltip", "Add the chosen trait to currently selected node"),
-			FNewToolMenuDelegate::CreateUObject(NonConstThis, &UAnimNextEdGraphNode::BuildAddTraitContextMenu));
-	}
-	else if(IsRunGraphNode())
-	{
-		FToolMenuSection& Section = Menu->AddSection("AnimNextRunAnimGraphNodeActions", LOCTEXT("AnimNextAnimGraphNodeActionsMenuHeader", "Animation Graph"));
-
-		URigVMController* VMController = GetController();
-		URigVMNode* VMNode = GetModelNode();
-		URigVMPin* VMPin = Context->Pin != nullptr ? FindModelPinFromGraphPin(Context->Pin) : nullptr;
-
-		if(VMPin != nullptr && VMNode->FindTrait(VMPin))
-		{
-			Section.AddMenuEntry(
-				TEXT("RemoveExposedVariables"),
-				LOCTEXT("RemoveExposedVariablesMenu", "Remove Exposed Variables"),
-				LOCTEXT("RemoveExposeVariablesMenuTooltip", "Remove the exposed variable trait from this node"),
-				FSlateIcon(),
-				FUIAction(
-					FExecuteAction::CreateLambda([VMController, VMNode, Name = VMPin->GetFName()]()
-					{
-						VMController->RemoveTrait(VMNode, Name);
-					})
-				));
-		}
-		else
-		{
-			Section.AddSubMenu(
-				TEXT("ExposeVariables"),
-				LOCTEXT("ExposeVariablesMenu", "Expose Variables"),
-				LOCTEXT("ExposeVariablesMenuTooltip", "Expose the variables of a selected animation graph as pins on this node"),
-				FNewToolMenuDelegate::CreateUObject(NonConstThis, &UAnimNextEdGraphNode::BuildExposeVariablesContextMenu));
-		}
-	}
-}
-
-void UAnimNextEdGraphNode::ConfigurePin(UEdGraphPin* EdGraphPin, const URigVMPin* ModelPin) const
-{
-	Super::ConfigurePin(EdGraphPin, ModelPin);
-
-	// Trait handles always remain as a RigVM input pins so that we can still link things to them even if they are hidden
-	// We handle visibility for those explicitly here
-	const bool bIsInputPin = ModelPin->GetDirection() == ERigVMPinDirection::Input;
-	const bool bIsTraitHandle = ModelPin->GetCPPTypeObject() == FAnimNextTraitHandle::StaticStruct();
-	if (bIsInputPin && bIsTraitHandle)
-	{
-		if (const URigVMPin* DecoratorPin = ModelPin->GetParentPin())
-		{
-			if (DecoratorPin->IsTraitPin())
-			{
-				check(DecoratorPin->GetScriptStruct() == FRigDecorator_AnimNextCppDecorator::StaticStruct());
-
-				TSharedPtr<FStructOnScope> DecoratorScope = DecoratorPin->GetTraitInstance();
-				const FRigDecorator_AnimNextCppDecorator* VMDecorator = (const FRigDecorator_AnimNextCppDecorator*)DecoratorScope->GetStructMemory();
-
-				const UScriptStruct* TraitStruct = VMDecorator->GetTraitSharedDataStruct();
-				check(TraitStruct != nullptr);
-
-				const FProperty* PinProperty = TraitStruct->FindPropertyByName(ModelPin->GetFName());
-				EdGraphPin->bHidden = PinProperty->HasMetaData(FRigVMStruct::HiddenMetaName);
-			}
-		}
-	}
-}
-
-bool UAnimNextEdGraphNode::IsTraitStack() const
-{
-	if (const URigVMUnitNode* VMNode = Cast<URigVMUnitNode>(GetModelNode()))
-	{
-		const UScriptStruct* ScriptStruct = VMNode->GetScriptStruct();
-		return ScriptStruct == FRigUnit_AnimNextTraitStack::StaticStruct();
-	}
-
-	return false;
-}
-
-bool UAnimNextEdGraphNode::IsRunGraphNode() const
-{
-	if (const URigVMUnitNode* VMNode = Cast<URigVMUnitNode>(GetModelNode()))
-	{
-		const UScriptStruct* ScriptStruct = VMNode->GetScriptStruct();
-		return ScriptStruct == FRigUnit_AnimNextRunAnimationGraph::StaticStruct();
-	}
-
-	return false;
-}
-
-void UAnimNextEdGraphNode::BuildAddTraitContextMenu(UToolMenu* SubMenu)
-{
-	using namespace UE::AnimNext;
-
-	const FTraitRegistry& TraitRegistry = FTraitRegistry::Get();
-	TArray<const FTrait*> Traits = TraitRegistry.GetTraits();
-
-	URigVMController* VMController = GetController();
-	URigVMNode* VMNode = GetModelNode();
-
-	const UScriptStruct* CppDecoratorStruct = FRigDecorator_AnimNextCppDecorator::StaticStruct();
-
-	for (const FTrait* Trait : Traits)
-	{
-		UScriptStruct* ScriptStruct = Trait->GetTraitSharedDataStruct();
-
-		FString DefaultValue;
-		{
-			const FRigDecorator_AnimNextCppDecorator DefaultCppDecoratorStructInstance;
-			FRigDecorator_AnimNextCppDecorator CppDecoratorStructInstance;
-			CppDecoratorStructInstance.DecoratorSharedDataStruct = ScriptStruct;
-
-			if (!CppDecoratorStructInstance.CanBeAddedToNode(VMNode, nullptr))
-			{
-				continue;	// This trait isn't supported on this node
-			}
-
-			FRigDecorator_AnimNextCppDecorator::StaticStruct()->ExportText(DefaultValue, &CppDecoratorStructInstance, &DefaultCppDecoratorStructInstance, nullptr, PPF_SerializedAsImportText, nullptr);
-		}
-
-		FString DisplayNameMetadata;
-		ScriptStruct->GetStringMetaDataHierarchical(FRigVMStruct::DisplayNameMetaName, &DisplayNameMetadata);
-		const FString DisplayName = DisplayNameMetadata.IsEmpty() ? Trait->GetTraitName() : DisplayNameMetadata;
-
-		const FText ToolTip = ScriptStruct->GetToolTipText();
-
-		FToolMenuEntry TraitEntry = FToolMenuEntry::InitMenuEntry(
-			*Trait->GetTraitName(),
-			FText::FromString(DisplayName),
-			ToolTip,
-			FSlateIcon(),
-			FUIAction(FExecuteAction::CreateLambda(
-				[this, Trait, VMController, VMNode, CppDecoratorStruct, DefaultValue, DisplayName]()
-				{
-					VMController->AddTrait(
-						VMNode->GetFName(),
-						*CppDecoratorStruct->GetPathName(),
-						*DisplayName,
-						DefaultValue, INDEX_NONE, true, true);
-				})
-			)
-		);
-
-		SubMenu->AddMenuEntry(NAME_None, TraitEntry);
-	}
-}
-
-void UAnimNextEdGraphNode::BuildExposeVariablesContextMenu(UToolMenu* SubMenu)
-{
-	using namespace UE::AnimNext::Editor::Private;
-
-	URigVMController* VMController = GetController();
-	URigVMNode* VMNode = GetModelNode();
-
-	FContentBrowserModule& ContentBrowserModule = FModuleManager::Get().LoadModuleChecked<FContentBrowserModule>(TEXT("ContentBrowser"));
-
-	FAssetPickerConfig AssetPickerConfig;
-	AssetPickerConfig.Filter.ClassPaths.Add(UAnimNextDataInterface::StaticClass()->GetClassPathName());
-	AssetPickerConfig.Filter.bRecursiveClasses = true;
-	AssetPickerConfig.InitialAssetViewType = EAssetViewType::List;
-	AssetPickerConfig.AssetShowWarningText = LOCTEXT("NoAssetsWithPublicVariablesMessage", "No animation graphs with public variables found");
-	AssetPickerConfig.OnAssetSelected = FOnAssetSelected::CreateLambda([VMController, VMNode](const FAssetData& InAssetData)
-	{
-		FSlateApplication::Get().DismissAllMenus();
-
-		FString DefaultValue;
-		FRigVMTrait_AnimNextPublicVariables DefaultTrait;
-		FRigVMTrait_AnimNextPublicVariables NewTrait;
-		UAnimNextDataInterface* Asset = CastChecked<UAnimNextDataInterface>(InAssetData.GetAsset());
-		NewTrait.InternalAsset = Asset;
-		const FInstancedPropertyBag& PublicVariableDefaults = Asset->GetPublicVariableDefaults();
-		TConstArrayView<FPropertyBagPropertyDesc> Descs = PublicVariableDefaults.GetPropertyBagStruct()->GetPropertyDescs();
-		NewTrait.InternalVariableNames.Reserve(Descs.Num());
-		for(const FPropertyBagPropertyDesc& Desc : Descs)
-		{
-			NewTrait.InternalVariableNames.Add(Desc.Name);
-		}
-		FRigVMTrait_AnimNextPublicVariables::StaticStruct()->ExportText(DefaultValue, &NewTrait, &DefaultTrait, nullptr, PPF_SerializedAsImportText, nullptr);
-
-		const FName ValidTraitName = URigVMSchema::GetUniqueName(VariablesTraitBaseName, [VMNode](const FName& InName) {
-			return VMNode->FindPin(InName.ToString()) == nullptr;
-		}, false, false);
-		VMController->AddTrait(VMNode, FRigVMTrait_AnimNextPublicVariables::StaticStruct(), ValidTraitName, DefaultValue);
-	});
-
-	AssetPickerConfig.OnShouldFilterAsset = FOnShouldFilterAsset::CreateLambda([this](const FAssetData& InAssetData)
-	{
-		// Filter to only show assets with public variables
-		FAnimNextAssetRegistryExports Exports;
-		if(UE::AnimNext::UncookedOnly::FUtils::GetExportedVariablesForAsset(InAssetData, Exports))
-		{
-			for(const FAnimNextAssetRegistryExportedVariable& Export : Exports.Variables)
-			{
-				if((Export.GetFlags() & EAnimNextExportedVariableFlags::Public) != EAnimNextExportedVariableFlags::NoFlags)
-				{
-					return false;
-				}
-			}
-		}
-		return true;
-	});
-
-	FToolMenuEntry Entry = FToolMenuEntry::InitWidget(
-		TEXT("AnimationGraphPicker"),
-		SNew(SBox)
-		.WidthOverride(300.0f)
-		.HeightOverride(400.0f)
-		[
-			ContentBrowserModule.Get().CreateAssetPicker(AssetPickerConfig)
-		],
-		FText::GetEmpty(),
-		true,
-		false,
-		false,
-		LOCTEXT("AnimationGraphPickerTooltip", "Choose an animation graph with public variables to expose")
-	);
-	
-	SubMenu->AddMenuEntry(NAME_None, Entry);
-}
-
-#undef LOCTEXT_NAMESPACE
diff --git a/Engine/Plugins/Experimental/AnimNext/Source/AnimNextUncookedOnly/Private/AnimNextRigVMAssetEditorData.cpp b/Engine/Plugins/Experimental/AnimNext/Source/AnimNextUncookedOnly/Private/AnimNextRigVMAssetEditorData.cpp
--- a/Engine/Plugins/Experimental/AnimNext/Source/AnimNextUncookedOnly/Private/AnimNextRigVMAssetEditorData.cpp
+++ b/Engine/Plugins/Experimental/AnimNext/Source/AnimNextUncookedOnly/Private/AnimNextRigVMAssetEditorData.cpp
@@ -1,8 +1,6 @@
 // Copyright Epic Games, Inc. All Rights Reserved.
 
 #include "AnimNextRigVMAssetEditorData.h"
-
-#include "AnimNextController.h"
 #include "AnimNextEdGraph.h"
 #include "AnimNextEdGraphSchema.h"
 #include "AnimNextRigVMAsset.h"
@@ -17,11 +15,9 @@
 #include "ObjectTools.h"
 #include "UncookedOnlyUtils.h"
 #include "DataInterface/AnimNextDataInterface_EditorData.h"
-#include "Entries/AnimNextAnimationGraphEntry.h"
 #include "Entries/AnimNextDataInterfaceEntry.h"
 #include "Entries/AnimNextEventGraphEntry.h"
 #include "Entries/AnimNextVariableEntry.h"
-#include "Graph/AnimNextAnimationGraphSchema.h"
 #include "Graph/RigUnit_AnimNextBeginExecution.h"
 #include "Misc/TransactionObjectEvent.h"
 #include "Module/AnimNextEventGraphSchema.h"
@@ -158,6 +154,11 @@ void UAnimNextRigVMAssetEditorData::Initialize(bool bRecompileVM)
 		Entry->Initialize(this);
 	}
 
+	InitializeAssetUserData();
+}
+
+void UAnimNextRigVMAssetEditorData::InitializeAssetUserData()
+{
 	if (IInterface_AssetUserData* OuterUserData = Cast<IInterface_AssetUserData>(GetOuter()))
 	{
 		if(!OuterUserData->HasAssetUserDataOfClass(GetAssetUserDataClass()))
@@ -790,7 +791,7 @@ void UAnimNextRigVMAssetEditorData::RecompileVM()
 			return;
 		}
 
-		UAnimNextController* Controller = CastChecked<UAnimNextController>(VMClient->GetOrCreateController(AllGraphs[0]));
+		URigVMController* Controller = VMClient->GetOrCreateController(AllGraphs[0]);
 
 		URigVMCompiler* Compiler = URigVMCompiler::StaticClass()->GetDefaultObject<URigVMCompiler>();
 		Compiler->Compile(Settings, AllGraphs, Controller, Asset->VM, Asset->ExtendedExecuteContext, Asset->GetExternalVariables(), &PinToOperandMap);
@@ -1593,81 +1594,6 @@ UAnimNextEventGraphEntry* UAnimNextRigVMAssetEditorData::AddEventGraph(FName InN
 	return NewEntry;
 }
 
-UAnimNextAnimationGraphEntry* UAnimNextRigVMAssetLibrary::AddAnimationGraph(UAnimNextRigVMAsset* InAsset, FName InName, bool bSetupUndoRedo, bool bPrintPythonCommand)
-{
-	return UE::AnimNext::UncookedOnly::FUtils::GetEditorData(InAsset)->AddAnimationGraph(InName, bSetupUndoRedo, bPrintPythonCommand);
-}
-
-UAnimNextAnimationGraphEntry* UAnimNextRigVMAssetEditorData::AddAnimationGraph(FName InName, bool bSetupUndoRedo, bool bPrintPythonCommand)
-{
-	if(InName == NAME_None)
-	{
-		ReportError(TEXT("UAnimNextRigVMAssetEditorData::AddAnimationGraph: Invalid graph name supplied."));
-		return nullptr;
-	}
-
-	if(!GetEntryClasses().Contains(UAnimNextAnimationGraphEntry::StaticClass()) || !CanAddNewEntry(UAnimNextAnimationGraphEntry::StaticClass()))
-	{
-		ReportError(TEXT("UAnimNextRigVMAssetEditorData::AddAnimationGraph: Cannot add an animation graph to this asset - entry is not allowed."));
-		return nullptr;
-	}
-
-	// Check for duplicate name
-	FName NewGraphName = InName;
-	auto DuplicateNamePredicate = [&NewGraphName](const UAnimNextRigVMAssetEntry* InEntry)
-	{
-		return InEntry->GetEntryName() == NewGraphName;
-	};
-
-	bool bAlreadyExists = Entries.ContainsByPredicate(DuplicateNamePredicate);
-	int32 NameNumber = InName.GetNumber() + 1;
-	while(bAlreadyExists)
-	{
-		NewGraphName = FName(InName, NameNumber++);
-		bAlreadyExists =  Entries.ContainsByPredicate(DuplicateNamePredicate);
-	}
-
-	UAnimNextAnimationGraphEntry* NewEntry = CreateNewSubEntry<UAnimNextAnimationGraphEntry>(this);
-	NewEntry->GraphName = NewGraphName;
-	NewEntry->Initialize(this);
-
-	if(bSetupUndoRedo)
-	{
-		NewEntry->Modify();
-		Modify();
-	}
-
-	AddEntryInternal(NewEntry);
-
-	// Add new graph
-	{
-		TGuardValue<bool> EnablePythonPrint(bSuspendPythonMessagesForRigVMClient, !bPrintPythonCommand);
-		TGuardValue<bool> DisableAutoCompile(bAutoRecompileVM, false);
-		// Editor data has to be the graph outer, or RigVM unique name generator will not work
-		URigVMGraph* NewRigVMGraphModel = RigVMClient.CreateModel(URigVMGraph::StaticClass()->GetFName(), UAnimNextAnimationGraphSchema::StaticClass(), bSetupUndoRedo, this);
-		if (ensure(NewRigVMGraphModel))
-		{
-			// Then, to avoid the graph losing ref due to external package, set the same package as the Entry
-			if (!NewRigVMGraphModel->HasAnyFlags(RF_Transient))
-			{
-				NewRigVMGraphModel->SetExternalPackage(CastChecked<UObject>(NewEntry)->GetExternalPackage());
-			}
-			ensure(NewRigVMGraphModel);
-			NewEntry->Graph = NewRigVMGraphModel;
-
-			RefreshExternalModels();
-			RigVMClient.AddModel(NewRigVMGraphModel, true);
-			URigVMController* Controller = RigVMClient.GetController(NewRigVMGraphModel);
-			UE::AnimNext::UncookedOnly::FUtils::SetupAnimGraph(NewEntry->GetEntryName(), Controller);
-		}
-	}
-
-	CustomizeNewAssetEntry(NewEntry);
-
-	BroadcastModified(EAnimNextEditorDataNotifType::EntryAdded, NewEntry);
-
-	return NewEntry;
-}
 
 UAnimNextDataInterfaceEntry* UAnimNextRigVMAssetLibrary::AddDataInterface(UAnimNextRigVMAsset* InAsset, UAnimNextDataInterface* InDataInterface, bool bSetupUndoRedo, bool bPrintPythonCommand)
 {
diff --git a/Engine/Plugins/Experimental/AnimNext/Source/AnimNextUncookedOnly/Private/Module/AnimNextModule_EditorData.cpp b/Engine/Plugins/Experimental/AnimNext/Source/AnimNextUncookedOnly/Private/Module/AnimNextModule_EditorData.cpp
--- a/Engine/Plugins/Experimental/AnimNext/Source/AnimNextUncookedOnly/Private/Module/AnimNextModule_EditorData.cpp
+++ b/Engine/Plugins/Experimental/AnimNext/Source/AnimNextUncookedOnly/Private/Module/AnimNextModule_EditorData.cpp
@@ -4,17 +4,11 @@
 
 #include "ExternalPackageHelper.h"
 #include "UncookedOnlyUtils.h"
-#include "AssetRegistry/IAssetRegistry.h"
 #include "Module/AnimNextModule.h"
-#include "Entries/AnimNextAnimationGraphEntry.h"
 #include "Entries/AnimNextEventGraphEntry.h"
-#include "AnimNextEventGraphSchema.h"
 #include "Entries/AnimNextDataInterfaceEntry.h"
 #include "Entries/AnimNextVariableEntry.h"
-#include "Graph/AnimNextAnimationGraphSchema.h"
-#include "RigVMModel/Nodes/RigVMAggregateNode.h"
 #include "UObject/AssetRegistryTagsContext.h"
-#include "UObject/FortniteMainBranchObjectVersion.h"
 #include "UObject/LinkerLoad.h"
 #include "Variables/AnimNextUniversalObjectLocatorBindingData.h"
 
@@ -32,138 +26,6 @@ void UAnimNextModule_EditorData::RecompileVM()
 #endif
 }
 
-void UAnimNextModule_EditorData::PostLoad()
-{
-	Super::PostLoad();
-
-	auto FindEntryForRigVMGraph = [this](URigVMGraph* InRigVMGraph)
-	{
-		UAnimNextRigVMAssetEntry* FoundEntry = nullptr;
-		for(UAnimNextRigVMAssetEntry* Entry : Entries)
-		{
-			if(IAnimNextRigVMGraphInterface* GraphEntry = Cast<IAnimNextRigVMGraphInterface>(Entry))
-			{
-				if(InRigVMGraph == GraphEntry->GetRigVMGraph())
-				{
-					FoundEntry = Entry;
-					break;
-				}
-			}
-		}
-		return FoundEntry;
-	};
-
-	if(GetLinkerCustomVersion(FFortniteMainBranchObjectVersion::GUID) < FFortniteMainBranchObjectVersion::AnimNextCombineGraphContexts)
-	{
-		// Must preload entries so their data is populated or we cannot find the appropriate entries for graphs
-		for(UAnimNextRigVMAssetEntry* Entry : Entries) 
-		{
-			Entry->GetLinker()->Preload(Entry);
-		}
-
-		TArray<URigVMGraph*> AllModels = RigVMClient.GetAllModels(false, true);
-		for(URigVMGraph* Graph : AllModels)
-		{
-			Graph->SetExecuteContextStruct(FAnimNextExecuteContext::StaticStruct());
-			if(UAnimNextRigVMAssetEntry* FoundEntry = FindEntryForRigVMGraph(Graph))
-			{
-				if(FoundEntry->IsA(UAnimNextAnimationGraphEntry::StaticClass()))
-				{
-					Graph->SetSchemaClass(UAnimNextAnimationGraphSchema::StaticClass());
-				}
-				else
-				{
-					Graph->SetSchemaClass(UAnimNextEventGraphSchema::StaticClass());
-				}
-			}
-			else
-			{
-				Graph->SetSchemaClass(UAnimNextAnimationGraphSchema::StaticClass());
-			}
-		}
-	}
-
-	if(GetLinkerCustomVersion(FFortniteMainBranchObjectVersion::GUID) < FFortniteMainBranchObjectVersion::AnimNextMoveGraphsToEntries)
-	{
-		// Must preload entries so their data is populated or we cannot find the appropriate entries for graphs
-		for(UAnimNextRigVMAssetEntry* Entry : Entries) 
-		{
-			Entry->GetLinker()->Preload(Entry);
-		}
-		
-		for(TObjectPtr<UAnimNextEdGraph> Graph : Graphs_DEPRECATED)
-		{
-			URigVMGraph* FoundRigVMGraph = GetRigVMGraphForEditorObject(Graph);
-			if(FoundRigVMGraph)
-			{
-				if(UAnimNextRigVMAssetEntry* FoundEntry = FindEntryForRigVMGraph(FoundRigVMGraph))
-				{
-					if(UAnimNextAnimationGraphEntry* AnimationGraphEntry = Cast<UAnimNextAnimationGraphEntry>(FoundEntry))
-					{
-						AnimationGraphEntry->EdGraph = Graph;
-					}
-					else if(UAnimNextEventGraphEntry* EventGraphEntry = Cast<UAnimNextEventGraphEntry>(FoundEntry))
-					{
-						EventGraphEntry->EdGraph = Graph;
-					}
-
-					Graph->Rename(nullptr, FoundEntry, REN_DoNotDirty | REN_DontCreateRedirectors | REN_NonTransactional);
-					Graph->Initialize(this);
-				}
-			}
-		}
-
-		// We used to add a default model that is no longer needed
-		URigVMGraph* DefaultModel = RigVMClient.GetDefaultModel();
-		if(DefaultModel && DefaultModel->GetName() == TEXT("RigVMGraph"))
-		{
-			bool bFound = false;
-			for(UAnimNextRigVMAssetEntry* Entry : Entries)
-			{
-				if(UAnimNextEventGraphEntry* EventGraphEntry = Cast<UAnimNextEventGraphEntry>(Entry))
-				{
-					if(DefaultModel == static_cast<IAnimNextRigVMGraphInterface*>(EventGraphEntry)->GetRigVMGraph())
-					{
-						bFound = true;
-						break;
-					}
-				}
-			}
-
-			if(!bFound)
-			{
-				TGuardValue<bool> DisablePythonPrint(bSuspendPythonMessagesForRigVMClient, false);
-				TGuardValue<bool> DisableAutoCompile(bAutoRecompileVM, false);
-				RigVMClient.RemoveModel(DefaultModel->GetNodePath(), false);
-			}
-		}
-
-		RecompileVM();
-	}
-
-	if(GetLinkerCustomVersion(FFortniteMainBranchObjectVersion::GUID) < FFortniteMainBranchObjectVersion::AnimNextGraphAccessSpecifiers)
-	{
-		// Must preload entries so their data is populated as we will be modifying them
-		for(UAnimNextRigVMAssetEntry* Entry : Entries) 
-		{
-			Entry->GetLinker()->Preload(Entry);
-		}
-
-		// Force older assets to all have public symbols so they work as-is. Newer assets need user intervention as entries default to private
-		for(UAnimNextRigVMAssetEntry* Entry : Entries)
-		{
-			if(UAnimNextAnimationGraphEntry* AnimationGraphEntry = Cast<UAnimNextAnimationGraphEntry>(Entry))
-			{
-				AnimationGraphEntry->Access = EAnimNextExportAccessSpecifier::Public;
-			}
-			else if(UAnimNextVariableEntry* ParameterEntry = Cast<UAnimNextVariableEntry>(Entry))
-			{
-				ParameterEntry->Access = EAnimNextExportAccessSpecifier::Public;
-			}
-		}
-	}
-}
-
 TConstArrayView<TSubclassOf<UAnimNextRigVMAssetEntry>> UAnimNextModule_EditorData::GetEntryClasses() const
 {
 	static const TSubclassOf<UAnimNextRigVMAssetEntry> Classes[] =
diff --git a/Engine/Plugins/Experimental/AnimNext/Source/AnimNextUncookedOnly/Private/Module/AnimNextModule_EventGraphSchema.cpp b/Engine/Plugins/Experimental/AnimNext/Source/AnimNextUncookedOnly/Private/Module/AnimNextModule_EventGraphSchema.cpp
--- a/Engine/Plugins/Experimental/AnimNext/Source/AnimNextUncookedOnly/Private/Module/AnimNextModule_EventGraphSchema.cpp
+++ b/Engine/Plugins/Experimental/AnimNext/Source/AnimNextUncookedOnly/Private/Module/AnimNextModule_EventGraphSchema.cpp
@@ -1,10 +1,7 @@
 ﻿// Copyright Epic Games, Inc. All Rights Reserved.
 
 #include "AnimNextEventGraphSchema.h"
-
 #include "AnimNextExecuteContext.h"
-#include "Graph/RigUnit_AnimNextBase.h"
-#include "Graph/RigUnit_AnimNextTraitStack.h"
 
 bool UAnimNextEventGraphSchema::SupportsUnitFunction(URigVMController* InController, const FRigVMFunction* InUnitFunction) const
 {
@@ -12,8 +9,10 @@ bool UAnimNextEventGraphSchema::SupportsUnitFunction(URigVMController* InControl
 	{
 		if(FunctionExecuteContextStruct == FAnimNextExecuteContext::StaticStruct())
 		{
-			// Disallow trait stacks in event graphs 
-			if(InUnitFunction->Struct && InUnitFunction->Struct->IsChildOf(FRigUnit_AnimNextTraitStack::StaticStruct()))
+			// Disallow trait stacks in event graphs
+			// TODO: This should probably be using the regular ExecuteContext filtering (i.e. trait stacks should use a different execute context)
+			static const UScriptStruct* TraitStackStruct = FindObjectChecked<UScriptStruct>(nullptr, TEXT("/Script/AnimNextAnimGraph.RigUnit_AnimNextTraitStack"));
+			if(InUnitFunction->Struct && InUnitFunction->Struct->IsChildOf(TraitStackStruct))
 			{
 				return false;
 			}
diff --git a/Engine/Plugins/Experimental/AnimNext/Source/AnimNextUncookedOnly/Private/UncookedOnlyUtils.cpp b/Engine/Plugins/Experimental/AnimNext/Source/AnimNextUncookedOnly/Private/UncookedOnlyUtils.cpp
--- a/Engine/Plugins/Experimental/AnimNext/Source/AnimNextUncookedOnly/Private/UncookedOnlyUtils.cpp
+++ b/Engine/Plugins/Experimental/AnimNext/Source/AnimNextUncookedOnly/Private/UncookedOnlyUtils.cpp
@@ -6,9 +6,7 @@
 #include "AssetRegistry/AssetRegistryModule.h"
 #include "AssetRegistry/IAssetRegistry.h"
 #include "Module/AnimNextModule.h"
-#include "AnimNextController.h"
 #include "Module/AnimNextModule_EditorData.h"
-#include "Graph/RigUnit_AnimNextGraphRoot.h"
 #include "Serialization/MemoryReader.h"
 #include "RigVMCore/RigVM.h"
 #include "AnimNextRigVMAsset.h"
@@ -16,9 +14,6 @@
 #include "Entries/AnimNextRigVMAssetEntry.h"
 #include "AnimNextUncookedOnlyModule.h"
 #include "IAnimNextRigVMExportInterface.h"
-#include "GameFramework/Character.h"
-#include "Graph/AnimNextGraphEntryPoint.h"
-#include "Graph/AnimNextAnimationGraphSchema.h"
 #include "Logging/StructuredLog.h"
 #include "AnimNextAssetWorkspaceAssetUserData.h"
 #include "IWorkspaceEditor.h"
@@ -31,6 +26,7 @@
 #include "Entries/AnimNextDataInterfaceEntry.h"
 #include "DataInterface/AnimNextDataInterface.h"
 #include "DataInterface/AnimNextDataInterface_EditorData.h"
+#include "Logging/MessageLog.h"
 
 #define LOCTEXT_NAMESPACE "AnimNextUncookedOnlyUtils"
 
@@ -342,7 +338,7 @@ void FUtils::CompileVariableBindings(const FRigVMCompileSettings& InSettings, UA
 
 	FRigVMClient* VMClient = EditorData->GetRigVMClient();
 	URigVMController* Controller = VMClient->GetOrCreateController(BindingGraph);
-	URigVMNode* ExecuteBindingsNode = Controller->AddUnitNode(FRigUnit_AnimNextExecuteBindings::StaticStruct(), FRigUnit_AnimNextExecuteBindings::GetMethodName(), FVector2D::ZeroVector, FString(), false);
+	URigVMNode* ExecuteBindingsNode = Controller->AddUnitNode(FRigUnit_AnimNextExecuteBindings::StaticStruct(), FRigVMStruct::ExecuteName, FVector2D::ZeroVector, FString(), false);
 	if(ExecuteBindingsNode == nullptr)
 	{
 		InSettings.ReportError(TEXT("Could not spawn Execute Bindings node"));
@@ -358,7 +354,7 @@ void FUtils::CompileVariableBindings(const FRigVMCompileSettings& InSettings, UA
 
 	if(bHasPublicVariablesToCopy)
 	{
-		URigVMNode* CopyProxyVariablesNode = Controller->AddUnitNode(FRigUnit_CopyModuleProxyVariables::StaticStruct(), FRigUnit_CopyModuleProxyVariables::GetMethodName(), FVector2D(200, 0.0f), FString(), false);
+		URigVMNode* CopyProxyVariablesNode = Controller->AddUnitNode(FRigUnit_CopyModuleProxyVariables::StaticStruct(), FRigVMStruct::ExecuteName, FVector2D(200, 0.0f), FString(), false);
 		if(CopyProxyVariablesNode == nullptr)
 		{
 			InSettings.ReportError(TEXT("Could not spawn Copy Module Proxy Variables node"));
@@ -669,33 +665,6 @@ FRigVMTemplateArgumentType FUtils::GetRigVMArgTypeFromParamType(const FAnimNextP
 	return ArgType;
 }
 
-void FUtils::SetupAnimGraph(const FName EntryName, URigVMController* InController)
-{
-	// Clear the graph
-	InController->RemoveNodes(InController->GetGraph()->GetNodes());
-
-	// Add root node
-	URigVMUnitNode* MainEntryPointNode = InController->AddUnitNode(FRigUnit_AnimNextGraphRoot::StaticStruct(), FRigUnit_AnimNextGraphRoot::EventName, FVector2D(-400.0f, 0.0f), FString(), false);
-	if(MainEntryPointNode == nullptr)
-	{
-		return;
-	}
-
-	URigVMPin* BeginExecutePin = MainEntryPointNode->FindPin(GET_MEMBER_NAME_STRING_CHECKED(FRigUnit_AnimNextGraphRoot, Result));
-	if(BeginExecutePin == nullptr)
-	{
-		return;
-	}
-
-	URigVMPin* EntryPointPin = MainEntryPointNode->FindPin(GET_MEMBER_NAME_STRING_CHECKED(FRigUnit_AnimNextGraphRoot, EntryPoint));
-	if(EntryPointPin == nullptr)
-	{
-		return;
-	}
-
-	InController->SetPinDefaultValue(EntryPointPin->GetPinPath(), EntryName.ToString());
-}
-
 void FUtils::SetupEventGraph(URigVMController* InController, UScriptStruct* InEventStruct)
 {
 	// Clear the graph
diff --git a/Engine/Plugins/Experimental/AnimNextAnimGraph/Source/AnimNextAnimGraph/AnimNextAnimGraph.Build.cs b/Engine/Plugins/Experimental/AnimNextAnimGraph/Source/AnimNextAnimGraph/AnimNextAnimGraph.Build.cs
--- a/Engine/Plugins/Experimental/AnimNextAnimGraph/Source/AnimNextAnimGraph/AnimNextAnimGraph.Build.cs
+++ b/Engine/Plugins/Experimental/AnimNextAnimGraph/Source/AnimNextAnimGraph/AnimNextAnimGraph.Build.cs
@@ -6,6 +6,13 @@ namespace UnrealBuildTool.Rules
 	{
 		public AnimNextAnimGraph(ReadOnlyTargetRules Target) : base(Target)
 		{
+			PrivateDependencyModuleNames.AddRange(
+				new string[]
+				{
+					"AssetRegistry",
+				}
+			);
+
 			PublicDependencyModuleNames.AddRange(
 				new string[]
 				{
@@ -16,6 +23,16 @@ namespace UnrealBuildTool.Rules
 					"AnimNext",
 				}
 			);
+
+			if (Target.bBuildEditor == true)
+			{
+				PrivateDependencyModuleNames.AddRange(
+					new string[]
+					{
+						"RigVMDeveloper",
+					}
+				);
+			}
 		}
 	}
 }
\ No newline at end of file
diff --git a/Engine/Plugins/Experimental/AnimNextAnimGraph/Source/AnimNextAnimGraph/Internal/Graph/AnimNextAnimationGraph.h b/Engine/Plugins/Experimental/AnimNextAnimGraph/Source/AnimNextAnimGraph/Internal/Graph/AnimNextAnimationGraph.h
new file mode 100644
--- /dev/null
+++ b/Engine/Plugins/Experimental/AnimNextAnimGraph/Source/AnimNextAnimGraph/Internal/Graph/AnimNextAnimationGraph.h
@@ -0,0 +1,156 @@
+﻿// Copyright Epic Games, Inc. All Rights Reserved.
+
+#pragma once
+
+#include "CoreMinimal.h"
+#include "Graph/AnimNextGraphEntryPoint.h"
+#include "DataInterface/AnimNextDataInterface.h"
+#include "Graph/RigUnit_AnimNextGraphRoot.h"
+#include "RigVMCore/RigVM.h"
+#include "TraitCore/TraitPtr.h"
+#include "TraitCore/TraitHandle.h"
+#include "TraitCore/EntryPointHandle.h"
+#include "Graph/RigUnit_AnimNextGraphEvaluator.h"
+#include "RigVMHost.h"
+
+#include "AnimNextAnimationGraph.generated.h"
+
+class UEdGraph;
+class UAnimNextModule;
+class UAnimGraphNode_AnimNextGraph;
+struct FAnimNode_AnimNextGraph;
+struct FRigUnit_AnimNextGraphEvaluator;
+struct FAnimNextGraphInstance;
+struct FAnimNextScheduleGraphTask;
+struct FAnimNextEditorParam;
+struct FAnimNextParam;
+struct FAnimNextModuleInstance;
+
+namespace UE::AnimNext
+{
+	struct FContext;
+	struct FExecutionContext;
+	class FAnimNextModuleImpl;
+	struct FTestUtils;
+	struct FParametersProxy;
+	struct FPlayAnimSlotTrait;
+	struct FAnimationGraphComponent;
+}
+
+namespace UE::AnimNext::UncookedOnly
+{
+	struct FUtils;
+}
+
+namespace UE::AnimNext::Editor
+{
+	class FModuleEditor;
+	class FVariableCustomization;
+}
+
+namespace UE::AnimNext::AnimGraph
+{
+	class FAnimNextAnimGraphModule;
+}
+
+namespace UE::AnimNext::Graph
+{
+	extern ANIMNEXTANIMGRAPH_API const FName EntryPointName;
+	extern ANIMNEXTANIMGRAPH_API const FName ResultName;
+}
+
+// A user-created collection of animation logic & data
+UCLASS(BlueprintType)
+class ANIMNEXTANIMGRAPH_API UAnimNextAnimationGraph : public UAnimNextDataInterface
+{
+	GENERATED_BODY()
+
+public:
+	UAnimNextAnimationGraph(const FObjectInitializer& ObjectInitializer);
+
+	// UObject interface
+	virtual void Serialize(FArchive& Ar) override;
+
+	// Allocates an instance of the graph with an optional module and/or parent graph instance
+	// @param	InModuleInstance			The module instance to use
+	// @param	InParentGraphInstance		The parent graph instance to use
+	// @param	InEntryPoint				The entry point to use. If this is NAME_None then the default entry point for this graph is used
+	// @return the allocated instance. This can be invalid if the entry point was not found or the graph's root trait is invalid.
+	TSharedPtr<FAnimNextGraphInstance> AllocateInstance(FAnimNextModuleInstance* InModuleInstance = nullptr, FAnimNextGraphInstance* InParentGraphInstance = nullptr, FName InEntryPoint = NAME_None) const;
+
+private:
+	// Loads the graph data from the provided archive buffer and returns true on success, false otherwise
+	bool LoadFromArchiveBuffer(const TArray<uint8>& SharedDataArchiveBuffer);
+
+#if WITH_EDITORONLY_DATA
+	// During graph compilation, if we have existing graph instances, we freeze them by releasing their memory before thawing them
+	// Freezing is a partial release of resources that retains the necessary information to re-create things safely
+	void FreezeGraphInstances();
+
+	// During graph compilation, once compilation is done we thaw existing graph instances to reallocate their memory
+	void ThawGraphInstances();
+#endif
+	
+	friend class UAnimNextAnimationGraphFactory;
+	friend class UAnimNextAnimationGraph_EditorData;
+	friend class UAnimNextVariableEntry;
+	friend struct UE::AnimNext::UncookedOnly::FUtils;
+	friend class UE::AnimNext::Editor::FModuleEditor;
+	friend struct UE::AnimNext::FTestUtils;
+	friend FAnimNextGraphInstance;
+	friend class UAnimGraphNode_AnimNextGraph;
+	friend UE::AnimNext::FExecutionContext;
+	friend struct FAnimNextScheduleGraphTask;
+	friend UE::AnimNext::AnimGraph::FAnimNextAnimGraphModule;
+	friend class UE::AnimNext::Editor::FVariableCustomization;
+	friend struct UE::AnimNext::FParametersProxy;
+	friend struct UE::AnimNext::FPlayAnimSlotTrait;
+	friend struct UE::AnimNext::FAnimationGraphComponent;
+
+protected:
+#if WITH_EDITORONLY_DATA
+	mutable FCriticalSection GraphInstancesLock;
+
+	// This is a list of live graph instances that have been allocated, used in the editor to reset instances when we re-compile/live edit
+	mutable TSet<FAnimNextGraphInstance*> GraphInstances;
+#endif
+
+	// This is the execute method definition used by a graph to evaluate latent pins
+	UPROPERTY()
+	FAnimNextGraphEvaluatorExecuteDefinition ExecuteDefinition;
+
+	// Data for each entry point in this graph
+	UPROPERTY()
+	TArray<FAnimNextGraphEntryPoint> EntryPoints;
+
+	// This is a resolved handle to the root trait in our graph, for each entry point 
+	TMap<FName, FAnimNextTraitHandle> ResolvedRootTraitHandles;
+
+	// This is an index into EntryPoints, for each entry point
+	TMap<FName, int32> ResolvedEntryPoints;
+
+	// This is the graph shared data used by the trait system, the output of FTraitReader
+	// We de-serialize manually into this buffer from the archive buffer, this is never saved on disk
+	TArray<uint8> SharedDataBuffer;
+
+	// This is a list of all referenced UObjects in the graph shared data
+	// We collect all the references here to make it quick and easy for the GC to query them
+	// It means that object references in the graph shared data are not visited at runtime by the GC (they are immutable)
+	// The shared data serialization archive stores indices to these to perform UObject serialization
+	UPROPERTY()
+	TArray<TObjectPtr<UObject>> GraphReferencedObjects;
+
+	// The entry point that this graph defaults to using
+	UPROPERTY(EditAnywhere, Category = "Graph")
+	FName DefaultEntryPoint = FRigUnit_AnimNextGraphRoot::DefaultEntryPoint;
+
+	// Default state for this graph
+	UPROPERTY()
+	FAnimNextGraphState DefaultState;
+
+#if WITH_EDITORONLY_DATA
+	// This buffer holds the output of the FTraitWriter post compilation
+	// We serialize it manually and it is discarded at runtime
+	TArray<uint8> SharedDataArchiveBuffer;
+#endif
+};
diff --git a/Engine/Plugins/Experimental/AnimNextAnimGraph/Source/AnimNextAnimGraph/Internal/Graph/AnimNextGraphContextData.h b/Engine/Plugins/Experimental/AnimNextAnimGraph/Source/AnimNextAnimGraph/Internal/Graph/AnimNextGraphContextData.h
new file mode 100644
--- /dev/null
+++ b/Engine/Plugins/Experimental/AnimNextAnimGraph/Source/AnimNextAnimGraph/Internal/Graph/AnimNextGraphContextData.h
@@ -0,0 +1,42 @@
+// Copyright Epic Games, Inc. All Rights Reserved.
+
+#pragma once
+
+#include "CoreMinimal.h"
+#include "Module/AnimNextModuleContextData.h"
+#include "Graph/AnimNextGraphInstance.h"
+#include "AnimNextGraphContextData.generated.h"
+
+namespace UE::AnimNext
+{
+	struct FLatentPropertyHandle;
+}
+
+USTRUCT()
+struct FAnimNextGraphContextData : public FAnimNextModuleContextData
+{
+	GENERATED_BODY()
+
+	FAnimNextGraphContextData() = default;
+
+	FAnimNextGraphContextData(FAnimNextModuleInstance* InModuleInstance, const FAnimNextGraphInstance* InInstance, const TConstArrayView<UE::AnimNext::FLatentPropertyHandle>& InLatentHandles, void* InDestinationBasePtr, bool bInIsFrozen)
+		: FAnimNextModuleContextData(InModuleInstance, InInstance)
+		, LatentHandles(InLatentHandles)
+		, DestinationBasePtr(InDestinationBasePtr)
+		, bIsFrozen(bInIsFrozen)
+	{
+	}
+
+	const FAnimNextGraphInstance& GetGraphInstance() const { return static_cast<const FAnimNextGraphInstance&>(GetDataInterfaceInstance()); }
+	const TConstArrayView<UE::AnimNext::FLatentPropertyHandle>& GetLatentHandles() const { return LatentHandles; }
+	void* GetDestinationBasePtr() const { return DestinationBasePtr; }
+	bool IsFrozen() const { return bIsFrozen; }
+
+private:
+	TConstArrayView<UE::AnimNext::FLatentPropertyHandle> LatentHandles;
+	void* DestinationBasePtr = nullptr;
+	bool bIsFrozen = false;
+
+	friend struct FAnimNextGraphInstance;
+	friend struct FAnimNextExecuteContext;
+};
diff --git a/Engine/Plugins/Experimental/AnimNextAnimGraph/Source/AnimNextAnimGraph/Internal/Graph/RigDecorator_AnimNextCppTrait.h b/Engine/Plugins/Experimental/AnimNextAnimGraph/Source/AnimNextAnimGraph/Internal/Graph/RigDecorator_AnimNextCppTrait.h
new file mode 100644
--- /dev/null
+++ b/Engine/Plugins/Experimental/AnimNextAnimGraph/Source/AnimNextAnimGraph/Internal/Graph/RigDecorator_AnimNextCppTrait.h
@@ -0,0 +1,31 @@
+// Copyright Epic Games, Inc. All Rights Reserved.
+
+#pragma once
+
+#include "RigVMCore/RigVMTrait.h"
+
+#include "RigDecorator_AnimNextCppTrait.generated.h"
+
+namespace UE::AnimNext { struct FTrait; }
+
+/**
+ * AnimNext RigDecorator for all C++ traits.
+ * The trait shared data UScriptStruct determines which properties are exposed.
+ */
+USTRUCT(BlueprintType)
+struct ANIMNEXTANIMGRAPH_API FRigDecorator_AnimNextCppDecorator : public FRigVMTrait
+{
+	GENERATED_BODY()
+
+	// The struct the trait exposes with its shared data. Each one of its properties will be added as a pin.
+	UPROPERTY(meta = (Hidden))
+	TObjectPtr<UScriptStruct> DecoratorSharedDataStruct = nullptr;
+
+#if WITH_EDITOR
+	virtual void GetProgrammaticPins(URigVMController* InController, int32 InParentPinIndex, const FString& InDefaultValue, struct FRigVMPinInfoArray& OutPinArray) const override;
+	
+	virtual UScriptStruct* GetTraitSharedDataStruct() const override { return DecoratorSharedDataStruct; }
+
+	const UE::AnimNext::FTrait* GetTrait() const;
+#endif
+};
diff --git a/Engine/Plugins/Experimental/AnimNextAnimGraph/Source/AnimNextAnimGraph/Internal/Graph/RigUnit_AnimNextRunAnimationGraph.h b/Engine/Plugins/Experimental/AnimNextAnimGraph/Source/AnimNextAnimGraph/Internal/Graph/RigUnit_AnimNextRunAnimationGraph.h
new file mode 100644
--- /dev/null
+++ b/Engine/Plugins/Experimental/AnimNextAnimGraph/Source/AnimNextAnimGraph/Internal/Graph/RigUnit_AnimNextRunAnimationGraph.h
@@ -0,0 +1,44 @@
+// Copyright Epic Games, Inc. All Rights Reserved.
+
+#pragma once
+
+#include "AnimNextExecuteContext.h"
+#include "Graph/AnimNext_LODPose.h"
+#include "Graph/RigUnit_AnimNextBase.h"
+#include "Graph/AnimNextGraphInstancePtr.h"
+
+#include "RigUnit_AnimNextRunAnimationGraph.generated.h"
+
+/** Runs an animation graph */
+USTRUCT(meta=(DisplayName="Run Graph", Category="Animation Graph", NodeColor="0, 1, 1", Keywords="Trait,Stack"))
+struct ANIMNEXTANIMGRAPH_API FRigUnit_AnimNextRunAnimationGraph : public FRigUnit_AnimNextBase
+{
+	GENERATED_BODY()
+
+	RIGVM_METHOD()
+	void Execute();
+
+	// Graph to run
+	UPROPERTY(EditAnywhere, Category = "Graph", meta = (Input))
+	TObjectPtr<UAnimNextAnimationGraph> Graph;
+
+	// Instance used to hold graph state
+	UPROPERTY(EditAnywhere, Category = "Graph", meta = (Input, Output))
+	FAnimNextGraphInstancePtr Instance;
+
+	// LOD to run the graph at
+	UPROPERTY(EditAnywhere, Category = "Graph", meta = (Input))
+	int32 LOD = 0;
+
+	// Reference pose for the graph
+	UPROPERTY(EditAnywhere, Category = "Graph", meta = (Input))
+	FAnimNextGraphReferencePose ReferencePose;
+
+	// Pose result
+	UPROPERTY(EditAnywhere, Category = "Graph", meta = (Output))
+	FAnimNextGraphLODPose Result;
+
+	// The execution result
+	UPROPERTY(EditAnywhere, DisplayName = "Execute", Category = "BeginExecution", meta = (Input, Output))
+	FAnimNextExecuteContext ExecuteContext;
+};
diff --git a/Engine/Plugins/Experimental/AnimNextAnimGraph/Source/AnimNextAnimGraph/Internal/TraitCore/NodeTemplateBuilder.h b/Engine/Plugins/Experimental/AnimNextAnimGraph/Source/AnimNextAnimGraph/Internal/TraitCore/NodeTemplateBuilder.h
new file mode 100644
--- /dev/null
+++ b/Engine/Plugins/Experimental/AnimNextAnimGraph/Source/AnimNextAnimGraph/Internal/TraitCore/NodeTemplateBuilder.h
@@ -0,0 +1,43 @@
+// Copyright Epic Games, Inc. All Rights Reserved.
+
+#pragma once
+
+#include "CoreMinimal.h"
+#include "TraitCore/TraitUID.h"
+
+namespace UE::AnimNext
+{
+	struct FNodeTemplate;
+
+	/**
+	  * FNodeTemplateBuilder
+	  * 
+	  * Utility to help construct node templates.
+	  */
+	struct ANIMNEXTANIMGRAPH_API FNodeTemplateBuilder
+	{
+		FNodeTemplateBuilder() = default;
+
+		// Adds the specified trait type to the node template
+		void AddTrait(FTraitUID TraitUID);
+
+		// Returns a node template for the provided list of traits
+		// The node template will be built into the provided buffer and a pointer to it is returned
+		FNodeTemplate* BuildNodeTemplate(TArray<uint8>& NodeTemplateBuffer) const;
+
+		// Returns a node template for the provided list of traits
+		// The node template will be built into the provided buffer and a pointer to it is returned
+		static FNodeTemplate* BuildNodeTemplate(const TArray<FTraitUID>& InTraitUIDs, TArray<uint8>& NodeTemplateBuffer);
+
+		// Resets the node template builder
+		void Reset();
+
+	private:
+		static uint32 GetNodeTemplateUID(const TArray<FTraitUID>& InTraitUIDs);
+		static void AppendTemplateTrait(
+			const TArray<FTraitUID>& InTraitUIDs, int32 TraitIndex,
+			TArray<uint8>& NodeTemplateBuffer);
+
+		TArray<FTraitUID> TraitUIDs;	// The list of traits to use when building the node template
+	};
+}
diff --git a/Engine/Plugins/Experimental/AnimNextAnimGraph/Source/AnimNextAnimGraph/Private/AnimNextAnimGraphModule.cpp b/Engine/Plugins/Experimental/AnimNextAnimGraph/Source/AnimNextAnimGraph/Private/AnimNextAnimGraphModule.cpp
--- a/Engine/Plugins/Experimental/AnimNextAnimGraph/Source/AnimNextAnimGraph/Private/AnimNextAnimGraphModule.cpp
+++ b/Engine/Plugins/Experimental/AnimNextAnimGraph/Source/AnimNextAnimGraph/Private/AnimNextAnimGraphModule.cpp
@@ -3,18 +3,27 @@
 #include "AnimNextAnimGraphModule.h"
 #include "AnimNextAnimGraphSettings.h"
 #include "Animation/BlendProfile.h"
-#include "Curves/CurveFloat.h"
 #include "Module/AnimNextModule.h"
 #include "Modules/ModuleInterface.h"
 #include "Modules/ModuleManager.h"
 #include "RigVMCore/RigVMRegistry.h"
 #include "Animation/AnimSequence.h"
-#include "IAnimNextModuleInterface.h"
-#include "TraitInterfaces/IEvaluate.h"
 #include "TraitInterfaces/IUpdate.h"
-#include "EvaluationVM/EvaluationVM.h"
-#include "Graph/AnimNext_LODPose.h"
 #include "Graph/AnimNextAnimationGraph.h"
+#include "Graph/AnimNext_LODPose.h"
+#include "Graph/RigVMTrait_AnimNextPublicVariables.h"
+#include "TraitCore/NodeTemplateRegistry.h"
+#include "TraitCore/TraitInterfaceRegistry.h"
+#include "TraitCore/TraitRegistry.h"
+
+#if WITH_ANIMNEXT_CONSOLE_COMMANDS
+#include "HAL/IConsoleManager.h"
+#include "UObject/UObjectIterator.h"
+
+#include "TraitCore/NodeDescription.h"
+#include "TraitCore/NodeTemplate.h"
+#include "TraitCore/TraitTemplate.h"
+#endif
 
 namespace UE::AnimNext::AnimGraph
 {
@@ -40,71 +49,233 @@ void FAnimNextAnimGraphModule::StartupModule()
 		Settings->LoadConfig();
 	}
 
+	FRigVMRegistry& RigVMRegistry = FRigVMRegistry::Get();
 	static TPair<UClass*, FRigVMRegistry::ERegisterObjectOperation> const AllowedObjectTypes[] =
 	{
-		{ UAnimSequence::StaticClass(), FRigVMRegistry::ERegisterObjectOperation::Class },
-		{ UScriptStruct::StaticClass(), FRigVMRegistry::ERegisterObjectOperation::Class },
-		{ UBlendProfile::StaticClass(), FRigVMRegistry::ERegisterObjectOperation::Class },
-		{ UCurveFloat::StaticClass(), FRigVMRegistry::ERegisterObjectOperation::Class },
-		{ UAnimNextModule::StaticClass(), FRigVMRegistry::ERegisterObjectOperation::Class },
+		{ UAnimNextAnimationGraph::StaticClass(), FRigVMRegistry::ERegisterObjectOperation::Class },
 	};
 
-	FRigVMRegistry::Get().RegisterObjectTypes(AllowedObjectTypes);
+	RigVMRegistry.RegisterObjectTypes(AllowedObjectTypes);
+	
+	static UScriptStruct* const AllowedStructTypes[] =
+	{
+		FAnimNextGraphReferencePose::StaticStruct(),
+		FRigVMTrait_AnimNextPublicVariables::StaticStruct()
+	};
+
+	RigVMRegistry.RegisterStructTypes(AllowedStructTypes);
+
+	FTraitRegistry::Init();
+	FTraitInterfaceRegistry::Init();
+	FNodeTemplateRegistry::Init();
 
-	IAnimNextModuleInterface::Get().RegisterAnimNextAnimGraph(*this);
+#if WITH_ANIMNEXT_CONSOLE_COMMANDS
+	if (!IsRunningCommandlet())
+	{
+		ConsoleCommands.Add(IConsoleManager::Get().RegisterConsoleCommand(
+			TEXT("AnimNext.ListNodeTemplates"),
+			TEXT("Dumps statistics about node templates to the log."),
+			FConsoleCommandWithArgsDelegate::CreateRaw(this, &FAnimNextAnimGraphModule::ListNodeTemplates),
+			ECVF_Default
+		));
+		ConsoleCommands.Add(IConsoleManager::Get().RegisterConsoleCommand(
+			TEXT("AnimNext.Module"),
+			TEXT("Dumps statistics about modules to the log."),
+			FConsoleCommandWithArgsDelegate::CreateRaw(this, &FAnimNextAnimGraphModule::ListAnimationGraphs),
+			ECVF_Default
+		));
+	}
+#endif
 }
 
 void FAnimNextAnimGraphModule::ShutdownModule()
 {
-	IAnimNextModuleInterface::Get().UnregisterAnimNextAnimGraph();
+#if WITH_ANIMNEXT_CONSOLE_COMMANDS
+	for (IConsoleObject* Cmd : ConsoleCommands)
+	{
+		IConsoleManager::Get().UnregisterConsoleObject(Cmd);
+	}
+	ConsoleCommands.Empty();
+#endif
+
+	FNodeTemplateRegistry::Destroy();
+	FTraitInterfaceRegistry::Destroy();
+	FTraitRegistry::Destroy();
+
 	LoadedGraphs.Reset();
 }
 
-void FAnimNextAnimGraphModule::UpdateGraph(FAnimNextGraphInstance& GraphInstance, float DeltaTime, UE::AnimNext::FTraitEventList& InputEventList, FTraitEventList& OutputEventList) const
+void FAnimNextAnimGraphModule::AddReferencedObjects(FReferenceCollector& Collector)
+{
+	Collector.AddReferencedObjects(LoadedGraphs);
+}
+
+FString FAnimNextAnimGraphModule::GetReferencerName() const
 {
-	UE::AnimNext::UpdateGraph(GraphInstance, DeltaTime, InputEventList, OutputEventList);
+	return TEXT("AnimNextAnimGraphModule");
 }
 
-void FAnimNextAnimGraphModule::EvaluateGraph(FAnimNextGraphInstance& GraphInstance, const UE::AnimNext::FReferencePose& RefPose, int32 GraphLODLevel, FAnimNextGraphLODPose& OutputPose) const
+#if WITH_ANIMNEXT_CONSOLE_COMMANDS
+void FAnimNextAnimGraphModule::ListNodeTemplates(const TArray<FString>& Args)
 {
-	const FEvaluationProgram EvaluationProgram = UE::AnimNext::EvaluateGraph(GraphInstance);
+	// Turn off log times to make diff-ing easier
+	TGuardValue<ELogTimes::Type> DisableLogTimes(GPrintLogTimes, ELogTimes::None);
+
+	// Make sure to log everything
+	const ELogVerbosity::Type OldVerbosity = LogAnimation.GetVerbosity();
+	LogAnimation.SetVerbosity(ELogVerbosity::All);
 
-	FEvaluationVM EvaluationVM(EEvaluationFlags::All, RefPose, GraphLODLevel);
-	bool bHasValidOutput = false;
+	const FNodeTemplateRegistry& NodeTemplateRegistry = FNodeTemplateRegistry::Get();
+	const FTraitRegistry& TraitRegistry = FTraitRegistry::Get();
 
-	if (!EvaluationProgram.IsEmpty())
+	UE_LOG(LogAnimation, Log, TEXT("===== AnimNext Node Templates ====="));
+	UE_LOG(LogAnimation, Log, TEXT("Template Buffer Size: %u bytes"), NodeTemplateRegistry.TemplateBuffer.GetAllocatedSize());
+
+	for (auto It = NodeTemplateRegistry.TemplateUIDToHandleMap.CreateConstIterator(); It; ++It)
 	{
-		EvaluationProgram.Execute(EvaluationVM);
+		const FNodeTemplateRegistryHandle Handle = It.Value();
+		const FNodeTemplate* NodeTemplate = NodeTemplateRegistry.Find(Handle);
+
+		const uint32 NumTraits = NodeTemplate->GetNumTraits();
 
-		TUniquePtr<FKeyframeState> EvaluatedKeyframe;
-		if (EvaluationVM.PopValue(KEYFRAME_STACK_NAME, EvaluatedKeyframe))
+		UE_LOG(LogAnimation, Log, TEXT("[%x] has %u traits ..."), NodeTemplate->GetUID(), NumTraits);
+		UE_LOG(LogAnimation, Log, TEXT("    Template Size: %u bytes"), NodeTemplate->GetNodeTemplateSize());
+		UE_LOG(LogAnimation, Log, TEXT("    Shared Data Size: %u bytes"), NodeTemplate->GetNodeSharedDataSize());
+		UE_LOG(LogAnimation, Log, TEXT("    Instance Data Size: %u bytes"), NodeTemplate->GetNodeInstanceDataSize());
+		UE_LOG(LogAnimation, Log, TEXT("    Traits ..."));
+
+		const FTraitTemplate* TraitTemplates = NodeTemplate->GetTraits();
+		for (uint32 TraitIndex = 0; TraitIndex < NumTraits; ++TraitIndex)
 		{
-			OutputPose.LODPose.CopyFrom(EvaluatedKeyframe->Pose);
-			OutputPose.Curves.CopyFrom(EvaluatedKeyframe->Curves);
-			OutputPose.Attributes.CopyFrom(EvaluatedKeyframe->Attributes);
-			bHasValidOutput = true;
+			const FTraitTemplate* TraitTemplate = TraitTemplates + TraitIndex;
+			const FTrait* Trait = TraitRegistry.Find(TraitTemplate->GetRegistryHandle());
+			const FString TraitName = Trait != nullptr ? Trait->GetTraitName() : TEXT("<Unknown>");
+
+			const uint32 NextTraitIndex = TraitIndex + 1;
+			const uint32 EndOfNextTraitSharedData = NextTraitIndex < NumTraits ? TraitTemplates[NextTraitIndex].GetNodeSharedOffset() : NodeTemplate->GetNodeSharedDataSize();
+			const uint32 TraitSharedDataSize = EndOfNextTraitSharedData - TraitTemplate->GetNodeSharedOffset();
+
+			const uint32 EndOfNextTraitInstanceData = NextTraitIndex < NumTraits ? TraitTemplates[NextTraitIndex].GetNodeInstanceOffset() : NodeTemplate->GetNodeInstanceDataSize();
+			const uint32 TraitInstanceDataSize = EndOfNextTraitInstanceData - TraitTemplate->GetNodeInstanceOffset();
+
+			UE_LOG(LogAnimation, Log, TEXT("            %u: [%x] %s (%s)"), TraitIndex, TraitTemplate->GetUID().GetUID(), *TraitName, TraitTemplate->GetMode() == ETraitMode::Base ? TEXT("Base") : TEXT("Additive"));
+			UE_LOG(LogAnimation, Log, TEXT("                Shared Data: [Offset: %u bytes, Size: %u bytes]"), TraitTemplate->GetNodeSharedOffset(), TraitSharedDataSize);
+			if (TraitTemplate->HasLatentProperties() && Trait != nullptr)
+			{
+				UE_LOG(LogAnimation, Log, TEXT("                Shared Data Latent Property Handles: [Offset: %u bytes, Count: %u]"), TraitTemplate->GetNodeSharedLatentPropertyHandlesOffset(), Trait->GetNumLatentTraitProperties());
+			}
+			UE_LOG(LogAnimation, Log, TEXT("                Instance Data: [Offset: %u bytes, Size: %u bytes]"), TraitTemplate->GetNodeInstanceOffset(), TraitInstanceDataSize);
 		}
 	}
 
-	if (!bHasValidOutput)
-	{
-		// We need to output a valid pose, generate one
-		FKeyframeState ReferenceKeyframe = EvaluationVM.MakeReferenceKeyframe(false);
-		OutputPose.LODPose.CopyFrom(ReferenceKeyframe.Pose);
-		OutputPose.Curves.CopyFrom(ReferenceKeyframe.Curves);
-		OutputPose.Attributes.CopyFrom(ReferenceKeyframe.Attributes);
-	}
+	LogAnimation.SetVerbosity(OldVerbosity);
 }
 
-void FAnimNextAnimGraphModule::AddReferencedObjects(FReferenceCollector& Collector)
+void FAnimNextAnimGraphModule::ListAnimationGraphs(const TArray<FString>& Args)
 {
-	Collector.AddReferencedObjects(LoadedGraphs);
-}
+	// Turn off log times to make diff-ing easier
+	TGuardValue<ELogTimes::Type> DisableLogTimes(GPrintLogTimes, ELogTimes::None);
 
-FString FAnimNextAnimGraphModule::GetReferencerName() const
-{
-	return TEXT("AnimNextAnimGraphModule");
+	// Make sure to log everything
+	const ELogVerbosity::Type OldVerbosity = LogAnimation.GetVerbosity();
+	LogAnimation.SetVerbosity(ELogVerbosity::All);
+
+	TArray<const UAnimNextAnimationGraph*> AnimationGraphs;
+
+	for (TObjectIterator<UAnimNextAnimationGraph> It; It; ++It)
+	{
+		AnimationGraphs.Add(*It);
+	}
+
+	struct FCompareObjectNames
+	{
+		FORCEINLINE bool operator()(const UAnimNextAnimationGraph& Lhs, const UAnimNextAnimationGraph& Rhs) const
+		{
+			return Lhs.GetPathName().Compare(Rhs.GetPathName()) < 0;
+		}
+	};
+	AnimationGraphs.Sort(FCompareObjectNames());
+
+	const FNodeTemplateRegistry& NodeTemplateRegistry = FNodeTemplateRegistry::Get();
+	const FTraitRegistry& TraitRegistry = FTraitRegistry::Get();
+	const bool bDetailedOutput = true;
+
+	UE_LOG(LogAnimation, Log, TEXT("===== AnimNext Modules ====="));
+	UE_LOG(LogAnimation, Log, TEXT("Num Graphs: %u"), AnimationGraphs.Num());
+
+	for (const UAnimNextAnimationGraph* AnimationGraph : AnimationGraphs)
+	{
+		uint32 TotalInstanceSize = 0;
+		uint32 NumNodes = 0;
+		{
+			// We always have a node at offset 0
+			int32 NodeOffset = 0;
+
+			while (NodeOffset < AnimationGraph->SharedDataBuffer.Num())
+			{
+				const FNodeDescription* NodeDesc = reinterpret_cast<const FNodeDescription*>(&AnimationGraph->SharedDataBuffer[NodeOffset]);
+
+				TotalInstanceSize += NodeDesc->GetNodeInstanceDataSize();
+				NumNodes++;
+
+				const FNodeTemplate* NodeTemplate = NodeTemplateRegistry.Find(NodeDesc->GetTemplateHandle());
+				NodeOffset += NodeTemplate->GetNodeSharedDataSize();
+			}
+		}
+
+		UE_LOG(LogAnimation, Log, TEXT("    %s ..."), *AnimationGraph->GetPathName());
+		UE_LOG(LogAnimation, Log, TEXT("        Shared Data Size: %.2f KB"), double(AnimationGraph->SharedDataBuffer.Num()) / 1024.0);
+		UE_LOG(LogAnimation, Log, TEXT("        Max Instance Data Size: %.2f KB"), double(TotalInstanceSize) / 1024.0);
+		UE_LOG(LogAnimation, Log, TEXT("        Num Nodes: %u"), NumNodes);
+
+		if (bDetailedOutput)
+		{
+			// We always have a node at offset 0
+			int32 NodeOffset = 0;
+
+			while (NodeOffset < AnimationGraph->SharedDataBuffer.Num())
+			{
+				const FNodeDescription* NodeDesc = reinterpret_cast<const FNodeDescription*>(&AnimationGraph->SharedDataBuffer[NodeOffset]);
+				const FNodeTemplate* NodeTemplate = NodeTemplateRegistry.Find(NodeDesc->GetTemplateHandle());
+
+				const uint32 NumTraits = NodeTemplate->GetNumTraits();
+
+				UE_LOG(LogAnimation, Log, TEXT("        Node %u: [Template %x with %u traits]"), NodeDesc->GetUID().GetNodeIndex(), NodeTemplate->GetUID(), NumTraits);
+				UE_LOG(LogAnimation, Log, TEXT("            Shared Data: [Offset: %u bytes, Size: %u bytes]"), NodeOffset, NodeTemplate->GetNodeSharedDataSize());
+				UE_LOG(LogAnimation, Log, TEXT("            Instance Data Size: %u bytes"), NodeDesc->GetNodeInstanceDataSize());
+				UE_LOG(LogAnimation, Log, TEXT("            Traits ..."));
+
+				const FTraitTemplate* TraitTemplates = NodeTemplate->GetTraits();
+				for (uint32 TraitIndex = 0; TraitIndex < NumTraits; ++TraitIndex)
+				{
+					const FTraitTemplate* TraitTemplate = TraitTemplates + TraitIndex;
+					const FTrait* Trait = TraitRegistry.Find(TraitTemplate->GetRegistryHandle());
+					const FString TraitName = Trait != nullptr ? Trait->GetTraitName() : TEXT("<Unknown>");
+
+					const uint32 NextTraitIndex = TraitIndex + 1;
+					const uint32 EndOfNextTraitSharedData = NextTraitIndex < NumTraits ? TraitTemplates[NextTraitIndex].GetNodeSharedOffset() : NodeTemplate->GetNodeSharedDataSize();
+					const uint32 TraitSharedDataSize = EndOfNextTraitSharedData - TraitTemplate->GetNodeSharedOffset();
+
+					const uint32 EndOfNextTraitInstanceData = NextTraitIndex < NumTraits ? TraitTemplates[NextTraitIndex].GetNodeInstanceOffset() : NodeTemplate->GetNodeInstanceDataSize();
+					const uint32 TraitInstanceDataSize = EndOfNextTraitInstanceData - TraitTemplate->GetNodeInstanceOffset();
+
+					UE_LOG(LogAnimation, Log, TEXT("                    %u: [%x] %s (%s)"), TraitIndex, TraitTemplate->GetUID().GetUID(), *TraitName, TraitTemplate->GetMode() == ETraitMode::Base ? TEXT("Base") : TEXT("Additive"));
+					UE_LOG(LogAnimation, Log, TEXT("                        Shared Data: [Offset: %u bytes, Size: %u bytes]"), TraitTemplate->GetNodeSharedOffset(), TraitSharedDataSize);
+					if (TraitTemplate->HasLatentProperties() && Trait != nullptr)
+					{
+						UE_LOG(LogAnimation, Log, TEXT("                        Shared Data Latent Property Handles: [Offset: %u bytes, Count: %u]"), TraitTemplate->GetNodeSharedLatentPropertyHandlesOffset(), Trait->GetNumLatentTraitProperties());
+					}
+					UE_LOG(LogAnimation, Log, TEXT("                        Instance Data: [Offset: %u bytes, Size: %u bytes]"), TraitTemplate->GetNodeInstanceOffset(), TraitInstanceDataSize);
+				}
+
+				NodeOffset += NodeTemplate->GetNodeSharedDataSize();
+			}
+		}
+	}
+
+	LogAnimation.SetVerbosity(OldVerbosity);
 }
+#endif
 
 }
 
diff --git a/Engine/Plugins/Experimental/AnimNextAnimGraph/Source/AnimNextAnimGraph/Private/AnimNextAnimGraphModule.h b/Engine/Plugins/Experimental/AnimNextAnimGraph/Source/AnimNextAnimGraph/Private/AnimNextAnimGraphModule.h
--- a/Engine/Plugins/Experimental/AnimNextAnimGraph/Source/AnimNextAnimGraph/Private/AnimNextAnimGraphModule.h
+++ b/Engine/Plugins/Experimental/AnimNextAnimGraph/Source/AnimNextAnimGraph/Private/AnimNextAnimGraphModule.h
@@ -3,8 +3,11 @@
 #pragma once
 
 #include "CoreMinimal.h"
-#include "IAnimNextModuleInterface.h"
 #include "Modules/ModuleInterface.h"
+#include "UObject/GCObject.h"
+
+// Enable console commands only in development builds when logging is enabled
+#define WITH_ANIMNEXT_CONSOLE_COMMANDS (!UE_BUILD_SHIPPING && !NO_LOGGING)
 
 class UAnimNextAnimationGraph;
 class UAnimNextAnimGraphSettings;
@@ -12,16 +15,12 @@ class UAnimNextAnimGraphSettings;
 namespace UE::AnimNext::AnimGraph
 {
 
-class FAnimNextAnimGraphModule : public IModuleInterface, public IAnimNextAnimGraph, public FGCObject
+class FAnimNextAnimGraphModule : public IModuleInterface, public FGCObject
 {
 	// IModuleInterface interface
 	virtual void StartupModule() override;
 	virtual void ShutdownModule() override;
 
-	// IAnimNextAnimGraph interface
-	virtual void UpdateGraph(FAnimNextGraphInstance& GraphInstance, float DeltaTime, UE::AnimNext::FTraitEventList& InputEventList, FTraitEventList& OutputEventList) const override;
-	virtual void EvaluateGraph(FAnimNextGraphInstance& GraphInstance, const UE::AnimNext::FReferencePose& RefPose, int32 GraphLODLevel, FAnimNextGraphLODPose& OutputPose) const override;
-
 	// FGCObject interface
 	virtual void AddReferencedObjects(FReferenceCollector& Collector) override;
 	virtual FString GetReferencerName() const override;
@@ -30,7 +29,14 @@ class FAnimNextAnimGraphModule : public IModuleInterface, public IAnimNextAnimGr
 	
 	// References to loaded default graphs to prevent GC.
 	// They must exist here rather than on the UAnimNextAnimGraphSettings CDO to avoid breaking DisregardForGC assumptions
-	TArray<TObjectPtr<const UAnimNextAnimationGraph>> LoadedGraphs;	
+	TArray<TObjectPtr<const UAnimNextAnimationGraph>> LoadedGraphs;
+
+#if WITH_ANIMNEXT_CONSOLE_COMMANDS
+	void ListNodeTemplates(const TArray<FString>& Args);
+	void ListAnimationGraphs(const TArray<FString>& Args);
+
+	TArray<IConsoleObject*> ConsoleCommands;
+#endif
 };
 
 }
diff --git a/Engine/Plugins/Experimental/AnimNextAnimGraph/Source/AnimNextAnimGraph/Private/Graph/AnimNextAnimationGraph.cpp b/Engine/Plugins/Experimental/AnimNextAnimGraph/Source/AnimNextAnimGraph/Private/Graph/AnimNextAnimationGraph.cpp
new file mode 100644
--- /dev/null
+++ b/Engine/Plugins/Experimental/AnimNextAnimGraph/Source/AnimNextAnimGraph/Private/Graph/AnimNextAnimationGraph.cpp
@@ -0,0 +1,215 @@
+﻿// Copyright Epic Games, Inc. All Rights Reserved.
+
+#include "Graph/AnimNextAnimationGraph.h"
+#include "RigVMCore/RigVMExecuteContext.h"
+#include "Graph/RigUnit_AnimNextShimRoot.h"
+#include "UObject/AssetRegistryTagsContext.h"
+#include "UObject/ObjectSaveContext.h"
+#include "TraitCore/TraitReader.h"
+#include "TraitCore/ExecutionContext.h"
+#include "Graph/AnimNext_LODPose.h"
+#include "Graph/AnimNextGraphInstance.h"
+#include "Serialization/MemoryReader.h"
+#include "AnimNextAnimGraphStats.h"
+#include "Graph/AnimNextGraphEntryPoint.h"
+#include "UObject/LinkerLoad.h"
+#include "UObject/ObjectResource.h"
+
+#if WITH_EDITOR	
+#include "Engine/ExternalAssetDependencyGatherer.h"
+REGISTER_ASSETDEPENDENCY_GATHERER(FExternalAssetDependencyGatherer, UAnimNextAnimationGraph);
+#endif // WITH_EDITOR
+
+DEFINE_STAT(STAT_AnimNext_Graph_AllocateInstance);
+
+#include UE_INLINE_GENERATED_CPP_BY_NAME(AnimNextAnimationGraph)
+
+UAnimNextAnimationGraph::UAnimNextAnimationGraph(const FObjectInitializer& ObjectInitializer)
+	: Super(ObjectInitializer)
+{
+	ExtendedExecuteContext.SetContextPublicDataStruct(FAnimNextExecuteContext::StaticStruct());
+}
+
+TSharedPtr<FAnimNextGraphInstance> UAnimNextAnimationGraph::AllocateInstance(FAnimNextModuleInstance* InModuleInstance, FAnimNextGraphInstance* InParentGraphInstance, FName InEntryPoint) const
+{
+	SCOPE_CYCLE_COUNTER(STAT_AnimNext_Graph_AllocateInstance);
+
+	const FName EntryPoint = (InEntryPoint == NAME_None) ? DefaultEntryPoint : InEntryPoint;
+	const FAnimNextTraitHandle ResolvedRootTraitHandle = ResolvedRootTraitHandles.FindRef(EntryPoint);
+	if (!ResolvedRootTraitHandle.IsValid())
+	{
+		return nullptr;
+	}
+
+	TSharedPtr<FAnimNextGraphInstance> InstanceImpl = MakeShared<FAnimNextGraphInstance>();
+	InstanceImpl->DataInterface = this;
+	InstanceImpl->ModuleInstance = InModuleInstance;
+	InstanceImpl->ParentGraphInstance = InParentGraphInstance;
+	InstanceImpl->EntryPoint = EntryPoint;
+
+	// If we have a parent graph, use its root since we share the same root, otherwise if we have no parent, we are the root
+	InstanceImpl->RootGraphInstance = InParentGraphInstance != nullptr ? InParentGraphInstance->GetRootGraphInstance() : InstanceImpl.Get();
+
+	InstanceImpl->Variables = VariableDefaults;
+	InstanceImpl->ExtendedExecuteContext = ExtendedExecuteContext;
+	if(VariableDefaults.GetPropertyBagStruct() != nullptr)
+	{
+		InstanceImpl->PublicVariablesState = FAnimNextGraphInstance::EPublicVariablesState::Unbound;
+
+		// Setup external variables memory ptrs manually as we dont follow the pattern of owning multiple URigVMHosts like control rig.
+		// InitializeVM() is called, but only sets up handles for the defaults in the module, not for an instance
+		const int32 NumVariables = InstanceImpl->Variables.GetNumPropertiesInBag();
+		TArray<FRigVMExternalVariableRuntimeData> ExternalVariableRuntimeData;
+		ExternalVariableRuntimeData.Reserve(NumVariables);
+		TConstArrayView<FPropertyBagPropertyDesc> Descs = InstanceImpl->Variables.GetPropertyBagStruct()->GetPropertyDescs();
+		uint8* BasePtr = InstanceImpl->Variables.GetMutableValue().GetMemory();
+		for(int32 VariableIndex = 0; VariableIndex < NumVariables; ++VariableIndex)
+		{
+			ExternalVariableRuntimeData.Emplace(Descs[VariableIndex].CachedProperty->ContainerPtrToValuePtr<uint8>(BasePtr));
+		}
+		InstanceImpl->ExtendedExecuteContext.ExternalVariableRuntimeData = MoveTemp(ExternalVariableRuntimeData);
+	}
+	else
+	{
+		InstanceImpl->PublicVariablesState = FAnimNextGraphInstance::EPublicVariablesState::None;
+	}
+
+	// Now initialize the 'instance', cache memory handles etc. in the context
+	VM->InitializeInstance(InstanceImpl->ExtendedExecuteContext);
+
+	{
+		UE::AnimNext::FExecutionContext Context(*InstanceImpl.Get());
+		InstanceImpl->GraphInstancePtr = Context.AllocateNodeInstance(*InstanceImpl, ResolvedRootTraitHandle);
+	}
+
+	if (!InstanceImpl->IsValid())
+	{
+		// We failed to allocate our instance, reset the ptr
+		InstanceImpl.Reset();
+	}
+
+#if WITH_EDITORONLY_DATA
+	if (InstanceImpl.IsValid() && InstanceImpl->IsValid())
+	{
+		FScopeLock Lock(&GraphInstancesLock);
+		check(!GraphInstances.Contains(InstanceImpl.Get()));
+		GraphInstances.Add(InstanceImpl.Get());
+	}
+#endif
+
+	return InstanceImpl;
+}
+
+void UAnimNextAnimationGraph::Serialize(FArchive& Ar)
+{
+	Super::Serialize(Ar);
+
+	if (Ar.IsLoading())
+	{
+		if(Ar.CustomVer(FFortniteMainBranchObjectVersion::GUID) < FFortniteMainBranchObjectVersion::AnimNextCombineParameterBlocksAndGraphs)
+		{
+			// Skip over shared archive buffer if we are loading from an older version
+			if (const FLinkerLoad* Linker = GetLinker())
+			{
+				const int32 LinkerIndex = GetLinkerIndex();
+				const FObjectExport& Export = Linker->ExportMap[LinkerIndex];
+				Ar.Seek(Export.SerialOffset + Export.SerialSize);
+			}
+		}
+		else
+		{
+			int32 SharedDataArchiveBufferSize = 0;
+			Ar << SharedDataArchiveBufferSize;
+
+#if !WITH_EDITORONLY_DATA
+			// When editor data isn't present, we don't persist the archive buffer as it is only needed on load
+			// to populate the graph shared data
+			TArray<uint8> SharedDataArchiveBuffer;
+#endif
+
+			SharedDataArchiveBuffer.SetNumUninitialized(SharedDataArchiveBufferSize);
+			Ar.Serialize(SharedDataArchiveBuffer.GetData(), SharedDataArchiveBufferSize);
+
+			if (Ar.IsLoadingFromCookedPackage())
+			{
+				// If we are cooked, we populate our graph shared data otherwise in the editor we'll compile on load
+				// and re-populate everything then to account for changes in code/content
+				LoadFromArchiveBuffer(SharedDataArchiveBuffer);
+			}
+		}
+	}
+	else if (Ar.IsSaving())
+	{
+#if WITH_EDITORONLY_DATA
+		// We only save the archive buffer, if code changes we'll be able to de-serialize from it when
+		// building the runtime buffer
+		// This allows us to have editor only/non-shipping only properties that are stripped out on load
+		int32 SharedDataArchiveBufferSize = SharedDataArchiveBuffer.Num();
+		Ar << SharedDataArchiveBufferSize;
+		Ar.Serialize(SharedDataArchiveBuffer.GetData(), SharedDataArchiveBufferSize);
+#endif
+	}
+	else
+	{
+		// Counting, etc
+		Ar << SharedDataBuffer;
+
+#if WITH_EDITORONLY_DATA
+		Ar << SharedDataArchiveBuffer;
+#endif
+	}
+}
+
+bool UAnimNextAnimationGraph::LoadFromArchiveBuffer(const TArray<uint8>& InSharedDataArchiveBuffer)
+{
+	using namespace UE::AnimNext;
+
+	// Reconstruct our graph shared data
+	FMemoryReader GraphSharedDataArchive(InSharedDataArchiveBuffer);
+	FTraitReader TraitReader(GraphReferencedObjects, GraphSharedDataArchive);
+
+	const FTraitReader::EErrorState ErrorState = TraitReader.ReadGraph(SharedDataBuffer);
+	if (ErrorState == FTraitReader::EErrorState::None)
+	{
+		for(int32 EntryPointIndex = 0; EntryPointIndex < EntryPoints.Num(); ++EntryPointIndex)
+		{
+			const FAnimNextGraphEntryPoint& EntryPoint = EntryPoints[EntryPointIndex];
+			ResolvedRootTraitHandles.Add(EntryPoint.EntryPointName, TraitReader.ResolveEntryPointHandle(EntryPoint.RootTraitHandle));
+			ResolvedEntryPoints.Add(EntryPoint.EntryPointName, EntryPointIndex);
+		}
+
+		// Make sure our execute method is registered
+		FRigUnit_AnimNextGraphEvaluator::RegisterExecuteMethod(ExecuteDefinition);
+		return true;
+	}
+	else
+	{
+		SharedDataBuffer.Empty(0);
+		ResolvedRootTraitHandles.Add(DefaultEntryPoint, FAnimNextTraitHandle());
+		return false;
+	}
+}
+
+#if WITH_EDITORONLY_DATA
+void UAnimNextAnimationGraph::FreezeGraphInstances()
+{
+	FScopeLock Lock(&GraphInstancesLock);
+
+	TSet<FAnimNextGraphInstance*> GraphInstancesCopy = GraphInstances;
+	for (FAnimNextGraphInstance* GraphInstance : GraphInstancesCopy)
+	{
+		GraphInstance->Freeze();
+	}
+}
+
+void UAnimNextAnimationGraph::ThawGraphInstances()
+{
+	FScopeLock Lock(&GraphInstancesLock);
+
+	TSet<FAnimNextGraphInstance*> GraphInstancesCopy = GraphInstances;
+	for (FAnimNextGraphInstance* GraphInstance : GraphInstancesCopy)
+	{
+		GraphInstance->Thaw();
+	}
+}
+#endif
diff --git a/Engine/Plugins/Experimental/AnimNextAnimGraph/Source/AnimNextAnimGraph/Private/Graph/AnimNextGraphInstance.cpp b/Engine/Plugins/Experimental/AnimNextAnimGraph/Source/AnimNextAnimGraph/Private/Graph/AnimNextGraphInstance.cpp
new file mode 100644
--- /dev/null
+++ b/Engine/Plugins/Experimental/AnimNextAnimGraph/Source/AnimNextAnimGraph/Private/Graph/AnimNextGraphInstance.cpp
@@ -0,0 +1,401 @@
+// Copyright Epic Games, Inc. All Rights Reserved.
+
+#include "Graph/AnimNextGraphInstance.h"
+
+#include "AnimNextAnimGraphStats.h"
+#include "DataInterface/AnimNextDataInterfaceHost.h"
+#include "DataInterface/DataInterfaceStructAdapter.h"
+#include "Graph/AnimNextAnimationGraph.h"
+#include "Graph/AnimNextGraphContextData.h"
+#include "TraitCore/ExecutionContext.h"
+#include "Graph/GC_GraphInstanceComponent.h"
+#include "Graph/RigUnit_AnimNextShimRoot.h"
+#include "Graph/RigVMTrait_AnimNextPublicVariables.h"
+#include "Logging/StructuredLog.h"
+#include "Module/AnimNextModule.h"
+#include "Module/AnimNextModuleInstance.h"
+
+DEFINE_STAT(STAT_AnimNext_Graph_RigVM);
+
+FAnimNextGraphInstance::FAnimNextGraphInstance()
+{
+#if WITH_EDITORONLY_DATA
+	UAnimNextModule::OnModuleCompiled().AddRaw(this, &FAnimNextGraphInstance::OnModuleCompiled);
+#endif
+}
+
+FAnimNextGraphInstance::~FAnimNextGraphInstance()
+{
+	Release();
+}
+
+void FAnimNextGraphInstance::Release()
+{
+#if WITH_EDITORONLY_DATA
+	UAnimNextModule::OnModuleCompiled().RemoveAll(this);
+
+	if(const UAnimNextAnimationGraph* Graph = GetAnimationGraph())
+	{
+		FScopeLock Lock(&Graph->GraphInstancesLock);
+		Graph->GraphInstances.Remove(this);
+	}
+#endif
+
+	if (!GraphInstancePtr.IsValid())
+	{
+		return;
+	}
+
+	GraphInstancePtr.Reset();
+	ModuleInstance = nullptr;
+	ParentGraphInstance = nullptr;
+	RootGraphInstance = nullptr;
+	ExtendedExecuteContext.Reset();
+	Components.Empty();
+	DataInterface = nullptr;
+}
+
+bool FAnimNextGraphInstance::IsValid() const
+{
+	return GraphInstancePtr.IsValid();
+}
+
+const UAnimNextAnimationGraph* FAnimNextGraphInstance::GetAnimationGraph() const
+{
+	return CastChecked<UAnimNextAnimationGraph>(DataInterface);
+}
+
+FName FAnimNextGraphInstance::GetEntryPoint() const
+{
+	return EntryPoint;
+}
+
+UE::AnimNext::FWeakTraitPtr FAnimNextGraphInstance::GetGraphRootPtr() const
+{
+	return GraphInstancePtr;
+}
+
+FAnimNextModuleInstance* FAnimNextGraphInstance::GetModuleInstance() const
+{
+	return ModuleInstance;
+}
+
+FAnimNextGraphInstance* FAnimNextGraphInstance::GetParentGraphInstance() const
+{
+	return ParentGraphInstance;
+}
+
+FAnimNextGraphInstance* FAnimNextGraphInstance::GetRootGraphInstance() const
+{
+	return RootGraphInstance;
+}
+
+bool FAnimNextGraphInstance::UsesAnimationGraph(const UAnimNextAnimationGraph* InAnimationGraph) const
+{
+	return GetAnimationGraph() == InAnimationGraph;
+}
+
+bool FAnimNextGraphInstance::UsesEntryPoint(FName InEntryPoint) const
+{
+	if(const UAnimNextAnimationGraph* AnimationGraph = GetAnimationGraph())
+	{
+		if(InEntryPoint == NAME_None)
+		{
+			return EntryPoint == AnimationGraph->DefaultEntryPoint;
+		}
+
+		return InEntryPoint == EntryPoint;
+	}
+	return false;
+}
+
+bool FAnimNextGraphInstance::IsRoot() const
+{
+	return this == RootGraphInstance;
+}
+
+bool FAnimNextGraphInstance::HasUpdated() const
+{
+	return bHasUpdatedOnce;
+}
+
+void FAnimNextGraphInstance::AddStructReferencedObjects(FReferenceCollector& Collector)
+{
+	if (!IsRoot())
+	{
+		return;	// If we aren't the root graph instance, we don't own the components
+	}
+
+	if (const UE::AnimNext::FGCGraphInstanceComponent* Component = TryGetComponent<UE::AnimNext::FGCGraphInstanceComponent>())
+	{
+		Component->AddReferencedObjects(Collector);
+	}
+}
+
+UE::AnimNext::FGraphInstanceComponent* FAnimNextGraphInstance::TryGetComponent(int32 ComponentNameHash, FName ComponentName) const
+{
+	if (const TSharedPtr<UE::AnimNext::FGraphInstanceComponent>* Component = RootGraphInstance->Components.FindByHash(ComponentNameHash, ComponentName))
+	{
+		return Component->Get();
+	}
+
+	return nullptr;
+}
+
+UE::AnimNext::FGraphInstanceComponent& FAnimNextGraphInstance::AddComponent(int32 ComponentNameHash, FName ComponentName, TSharedPtr<UE::AnimNext::FGraphInstanceComponent>&& Component)
+{
+	return *RootGraphInstance->Components.AddByHash(ComponentNameHash, ComponentName, MoveTemp(Component)).Get();
+}
+
+GraphInstanceComponentMapType::TConstIterator FAnimNextGraphInstance::GetComponentIterator() const
+{
+	return RootGraphInstance->Components.CreateConstIterator();
+}
+
+void FAnimNextGraphInstance::Update()
+{
+	bHasUpdatedOnce = true;
+}
+
+FAnimNextDataInterfaceInstance* FAnimNextGraphInstance::GetHost() const
+{
+	if(ParentGraphInstance != nullptr)
+	{
+		return ParentGraphInstance;
+	}
+
+	return ModuleInstance;
+}
+
+void FAnimNextGraphInstance::ExecuteLatentPins(const TConstArrayView<UE::AnimNext::FLatentPropertyHandle>& LatentHandles, void* DestinationBasePtr, bool bIsFrozen)
+{
+	SCOPE_CYCLE_COUNTER(STAT_AnimNext_Graph_RigVM);
+
+	if (!IsValid())
+	{
+		return;
+	}
+
+	if (URigVM* VM = GetAnimationGraph()->RigVM)
+	{
+		FAnimNextExecuteContext& AnimNextContext = ExtendedExecuteContext.GetPublicDataSafe<FAnimNextExecuteContext>();
+
+		// Insert our context data for the scope of execution
+		FAnimNextGraphContextData ContextData(ModuleInstance, this, LatentHandles, DestinationBasePtr, bIsFrozen);
+		UE::AnimNext::FScopedExecuteContextData ContextDataScope(AnimNextContext, ContextData);
+
+		VM->ExecuteVM(ExtendedExecuteContext, FRigUnit_AnimNextShimRoot::EventName);
+	}
+}
+
+#if WITH_EDITORONLY_DATA
+void FAnimNextGraphInstance::Freeze()
+{
+	if (!IsValid())
+	{
+		return;
+	}
+
+	GraphInstancePtr.Reset();
+	ExtendedExecuteContext.Reset();
+	Components.Empty();
+	PublicVariablesState = PublicVariablesState == EPublicVariablesState::Bound ? EPublicVariablesState::Unbound : EPublicVariablesState::None;
+	bHasUpdatedOnce = false;
+}
+
+void FAnimNextGraphInstance::Thaw()
+{
+	if (const UAnimNextAnimationGraph* AnimationGraph = GetAnimationGraph())
+	{
+		Variables.MigrateToNewBagInstance(AnimationGraph->VariableDefaults);
+
+		ExtendedExecuteContext = AnimationGraph->ExtendedExecuteContext;
+
+		{
+			UE::AnimNext::FExecutionContext Context(*this);
+			if(const FAnimNextTraitHandle* FoundHandle = AnimationGraph->ResolvedRootTraitHandles.Find(EntryPoint))
+			{
+				GraphInstancePtr = Context.AllocateNodeInstance(*this, *FoundHandle);
+			}
+		}
+
+		if (!IsValid())
+		{
+			// We failed to allocate our instance, clear everything
+			Release();
+		}
+	}
+}
+
+void FAnimNextGraphInstance::OnModuleCompiled(UAnimNextModule* InModule)
+{
+	// If we are hosted directly by a module, invalidate and mark our bindings as needing update. They will be lazily re-bound the next time we run.
+	if(ModuleInstance && InModule == ModuleInstance->GetModule() && ParentGraphInstance == nullptr)
+	{
+		UnbindPublicVariables();
+	}
+}
+
+#endif
+
+// Helper for binding to both FAnimNextDataInterfaceInstance and IDataInterfaceHost without an abstraction between the two
+template<typename HostType>
+bool FAnimNextGraphInstance::BindToHostHelper(const HostType& InHost, bool bInAutoBind)
+{
+	bool bPublicVariablesBound = false;
+
+	const UAnimNextAnimationGraph* AnimationGraph = GetAnimationGraph();
+	const UPropertyBag* PropertyBag = Variables.GetPropertyBagStruct();
+	for(const FAnimNextImplementedDataInterface& ImplementedInterface : AnimationGraph->GetImplementedInterfaces())
+	{
+		if(bInAutoBind && !ImplementedInterface.bAutoBindToHost)
+		{
+			continue;
+		}
+
+		const FAnimNextImplementedDataInterface* HostImplementedInterface = InHost.GetDataInterface()->FindImplementedInterface(ImplementedInterface.DataInterface);
+		if(HostImplementedInterface == nullptr)
+		{
+			// Host does not implement this interface, so skip
+			continue;
+		}
+
+		if(HostImplementedInterface->NumVariables != ImplementedInterface.NumVariables)
+		{
+			UE_LOGFMT(LogAnimation, Error, "BindToHost: Mismatched interface variables: '{Name}' ({Count}) vs Host '{Host}' ({HostCount})", AnimationGraph->GetFName(), ImplementedInterface.NumVariables, InHost.GetDataInterfaceName(), HostImplementedInterface->NumVariables);
+			continue;
+		}
+
+		int32 VariableIndex = ImplementedInterface.VariableIndex;
+		int32 HostVariableIndex = HostImplementedInterface->VariableIndex;
+		const int32 EndVariableIndex = ImplementedInterface.VariableIndex + ImplementedInterface.NumVariables;
+		for(; VariableIndex < EndVariableIndex; ++VariableIndex, ++HostVariableIndex)
+		{
+			const FPropertyBagPropertyDesc& Desc = PropertyBag->GetPropertyDescs()[VariableIndex];
+
+			uint8* HostMemory = InHost.GetMemoryForVariable(HostVariableIndex, Desc.Name, Desc.CachedProperty);
+			if(HostMemory == nullptr)
+			{
+				continue;
+			}
+
+			ExtendedExecuteContext.ExternalVariableRuntimeData[VariableIndex].Memory = HostMemory;
+			bPublicVariablesBound = true;
+		}
+	}
+
+	return bPublicVariablesBound;
+}
+
+void FAnimNextGraphInstance::BindPublicVariables(TArrayView<FStructView> InHostStructs)
+{
+	using namespace UE::AnimNext;
+
+	const UAnimNextAnimationGraph* AnimationGraph = GetAnimationGraph();
+	if (AnimationGraph == nullptr)
+	{
+		return;
+	}
+
+	const int32 NumStructs = InHostStructs.Num();
+	if(NumStructs > 0)
+	{
+		// stack allocating uninitialized arrays to contain FDataInterfaceStructAdapter(s)
+		TArrayView<FDataInterfaceStructAdapter> Adapters((FDataInterfaceStructAdapter*)FMemory_Alloca(NumStructs * sizeof(FDataInterfaceStructAdapter)), NumStructs);
+		TArrayView<IDataInterfaceHost*> PayloadVariableBindings((IDataInterfaceHost**)FMemory_Alloca(NumStructs * sizeof(IDataInterfaceHost*)), NumStructs);
+
+		// initializing the FDataInterfaceStructAdapter(s) with their respective HostStruct
+		for (int32 StructIndex = 0; StructIndex < NumStructs; ++StructIndex)
+		{
+			FStructView& StructView = InHostStructs[StructIndex];
+			PayloadVariableBindings[StructIndex] = new (&Adapters[StructIndex]) FDataInterfaceStructAdapter(AnimationGraph, StructView);
+		}
+
+		// Inner binding 
+		BindPublicVariables(PayloadVariableBindings);
+
+		// uninitializing the FDataInterfaceStructAdapter(s)
+		for (int32 StructIndex = 0; StructIndex < NumStructs; ++StructIndex)
+		{
+			Adapters[StructIndex].~FDataInterfaceStructAdapter();
+		}
+
+		// abandoning no longer initialized stack allocated arrays to their demise
+	}
+}
+
+void FAnimNextGraphInstance::BindPublicVariables(TConstArrayView<UE::AnimNext::IDataInterfaceHost*> InHosts)
+{
+	using namespace UE::AnimNext;
+
+	const UAnimNextAnimationGraph* AnimationGraph = GetAnimationGraph();
+	if (AnimationGraph == nullptr)
+	{
+		return;
+	}
+
+	if(PublicVariablesState == EPublicVariablesState::Bound)
+	{
+		return;
+	}
+
+	const UPropertyBag* PropertyBag = Variables.GetPropertyBagStruct();
+	if(PropertyBag == nullptr)
+	{
+		// Nothing to bind
+		PublicVariablesState = EPublicVariablesState::None;
+		return;
+	}
+
+	bool bPublicVariablesBound = false;
+
+	// First apply any automatic bindings to this instance's host
+	if(FAnimNextDataInterfaceInstance* InstanceHost = GetHost())
+	{
+		bPublicVariablesBound |= BindToHostHelper(*InstanceHost, true);
+	}
+
+	// Next bind to any supplied host interfaces (bInAutoBind = false)
+	for(const IDataInterfaceHost* HostInterface : InHosts)
+	{
+		check(HostInterface);
+		bPublicVariablesBound |= BindToHostHelper(*HostInterface, false);
+	}
+
+	if(bPublicVariablesBound)
+	{
+		// Re-initialize memory handles
+		AnimationGraph->RigVM->InitializeInstance(ExtendedExecuteContext, /* bCopyMemory = */false);
+	}
+
+	PublicVariablesState = EPublicVariablesState::Bound;
+}
+
+void FAnimNextGraphInstance::UnbindPublicVariables()
+{
+	const UAnimNextAnimationGraph* AnimationGraph = GetAnimationGraph();
+	if (AnimationGraph == nullptr)
+	{
+		return;
+	}
+
+	if(PublicVariablesState != EPublicVariablesState::Bound)
+	{
+		return;
+	}
+
+	// Reset external variable ptrs to point to internal public vars
+	const int32 NumVariables = Variables.GetNumPropertiesInBag();
+	TConstArrayView<FPropertyBagPropertyDesc> Descs = Variables.GetPropertyBagStruct()->GetPropertyDescs();
+	uint8* BasePtr = Variables.GetMutableValue().GetMemory();
+	for(int32 VariableIndex = 0; VariableIndex < NumVariables; ++VariableIndex)
+	{
+		ExtendedExecuteContext.ExternalVariableRuntimeData[VariableIndex].Memory = Descs[VariableIndex].CachedProperty->ContainerPtrToValuePtr<uint8>(BasePtr);
+	}
+
+	// Re-initialize memory handles
+	AnimationGraph->RigVM->InitializeInstance(ExtendedExecuteContext, /* bCopyMemory = */false);
+
+	PublicVariablesState = EPublicVariablesState::Unbound;
+}
+
diff --git a/Engine/Plugins/Experimental/AnimNextAnimGraph/Source/AnimNextAnimGraph/Private/Graph/GC_GraphInstanceComponent.h b/Engine/Plugins/Experimental/AnimNextAnimGraph/Source/AnimNextAnimGraph/Private/Graph/GC_GraphInstanceComponent.h
new file mode 100644
--- /dev/null
+++ b/Engine/Plugins/Experimental/AnimNextAnimGraph/Source/AnimNextAnimGraph/Private/Graph/GC_GraphInstanceComponent.h
@@ -0,0 +1,38 @@
+// Copyright Epic Games, Inc. All Rights Reserved.
+
+#pragma once
+
+#include "CoreMinimal.h"
+#include "TraitCore/TraitPtr.h"
+#include "Graph/GraphInstanceComponent.h"
+
+class FReferenceCollector;
+
+namespace UE::AnimNext
+{
+	/**
+	 * FGCGraphInstanceComponent
+	 *
+	 * This component maintains the necessary state to garbage collection.
+	 */
+	struct FGCGraphInstanceComponent : public FGraphInstanceComponent
+	{
+		DECLARE_ANIM_GRAPH_INSTANCE_COMPONENT(FGCGraphInstanceComponent)
+
+		explicit FGCGraphInstanceComponent(FAnimNextGraphInstance& InOwnerInstance);
+
+		// Registers the provided trait with the GC system
+		// Once registered, IGarbageCollection::AddReferencedObjects will be called on it during GC
+		void Register(const FWeakTraitPtr& InTraitPtr);
+
+		// Unregisters the provided trait from the GC system
+		void Unregister(const FWeakTraitPtr& InTraitPtr);
+
+		// Called during garbage collection to collect strong object references
+		void AddReferencedObjects(FReferenceCollector& Collector) const;
+
+	private:
+		// List of trait handles that contain UObject references and implement IGarbageCollection
+		TArray<FWeakTraitPtr> TraitsWithReferences;
+	};
+}
diff --git a/Engine/Plugins/Experimental/AnimNextAnimGraph/Source/AnimNextAnimGraph/Private/Graph/RigUnit_AnimNextRunAnimationGraph.cpp b/Engine/Plugins/Experimental/AnimNextAnimGraph/Source/AnimNextAnimGraph/Private/Graph/RigUnit_AnimNextRunAnimationGraph.cpp
new file mode 100644
--- /dev/null
+++ b/Engine/Plugins/Experimental/AnimNextAnimGraph/Source/AnimNextAnimGraph/Private/Graph/RigUnit_AnimNextRunAnimationGraph.cpp
@@ -0,0 +1,133 @@
+// Copyright Epic Games, Inc. All Rights Reserved.
+
+#include "Graph/RigUnit_AnimNextRunAnimationGraph.h"
+
+#include "Units/RigUnitContext.h"
+#include "AnimNextAnimGraphStats.h"
+#include "TraitInterfaces/IEvaluate.h"
+#include "TraitInterfaces/IUpdate.h"
+#include "Graph/AnimationGraphComponent.h"
+#include "Module/AnimNextModuleInstance.h"
+#include "Graph/AnimNextAnimationGraph.h"
+#include "Graph/RigVMTrait_AnimNextPublicVariables.h"
+#include "TraitCore/TraitEventList.h"
+
+#include UE_INLINE_GENERATED_CPP_BY_NAME(RigUnit_AnimNextRunAnimationGraph)
+
+DEFINE_STAT(STAT_AnimNext_Run_Graph);
+
+FRigUnit_AnimNextRunAnimationGraph_Execute()
+{
+	SCOPE_CYCLE_COUNTER(STAT_AnimNext_Run_Graph);
+
+	using namespace UE::AnimNext;
+
+	const FAnimNextModuleContextData& ModuleContextData = ExecuteContext.GetContextData<FAnimNextModuleContextData>();
+	FAnimNextModuleInstance& ModuleInstance = ModuleContextData.GetModuleInstance();
+
+	if (!ReferencePose.ReferencePose.IsValid())
+	{
+		return;
+	}
+
+	FAnimationGraphComponent& AnimationGraphComponent = ModuleInstance.GetComponent<FAnimationGraphComponent>();
+
+	if(Graph == nullptr)
+	{
+		AnimationGraphComponent.ReleaseInstance(Instance);
+		return;
+	}
+
+	// Release the instance if the graph has changed
+	if(Instance.IsValid() && !Instance.UsesAnimationGraph(Graph))
+	{
+		AnimationGraphComponent.ReleaseInstance(Instance);
+	}
+
+	// Lazily (re-)allocate graph instance if required
+	if(!Instance.IsValid())
+	{
+		Instance = AnimationGraphComponent.AllocateInstance(Graph);
+	}
+
+	if(!Instance.IsValid())
+	{
+		return;
+	}
+
+	if(Instance.RequiresPublicVariableBinding())
+	{
+		// Bind any public variables we find
+		TArray<FPublicVariablesTraitToDataInterfaceHostAdapter, TInlineAllocator<4>> Adapters;
+		Adapters.Reserve(ExecuteContext.GetTraits().Num());
+		TArray<IDataInterfaceHost*, TInlineAllocator<4>> Hosts;
+		Hosts.Reserve(ExecuteContext.GetTraits().Num());
+		for(const FRigVMTraitScope& TraitScope : ExecuteContext.GetTraits())
+		{
+			const FRigVMTrait_AnimNextPublicVariables* VariablesTrait = TraitScope.GetTrait<FRigVMTrait_AnimNextPublicVariables>();
+			if(VariablesTrait == nullptr)
+			{
+				continue;
+			}
+
+			FPublicVariablesTraitToDataInterfaceHostAdapter& Adapter = Adapters.Emplace_GetRef(*VariablesTrait, TraitScope);
+			Hosts.Add(&Adapter);
+		}
+
+		Instance.BindPublicVariables(Hosts);
+	}
+
+	const UE::AnimNext::FReferencePose& RefPose = ReferencePose.ReferencePose.GetRef<UE::AnimNext::FReferencePose>();
+
+	// TODO: Currently forcing additive flag to false here
+	if (Result.LODPose.ShouldPrepareForLOD(RefPose, LOD, false))
+	{
+		Result.LODPose.PrepareForLOD(RefPose, LOD, true, false);
+	}
+
+	ensure(Result.LODPose.LODLevel == LOD);
+
+	// Every graph in a schedule will see the same input events (if they were queued before the schedule started)
+	FTraitEventList InputEventList;
+	FTraitEventList OutputEventList;
+
+	// A schedule can contain multiple graphs, we copy the input event list since it might be appended to during our update
+	{
+		FReadScopeLock ReadLock(ModuleInstance.EventListLock);
+		InputEventList = ModuleInstance.InputEventList;
+	}
+
+	// Track how many input events we started with, we'll append the new ones
+	const int32 NumOriginalInputEvents = InputEventList.Num();
+
+	// Internally we use memstack allocation, so we need a mark here
+	FMemStack& MemStack = FMemStack::Get();
+	FMemMark MemMark(MemStack);
+
+	// We allocate a dummy buffer to trigger the allocation of a large chunk if this is the first mark
+	// This reduces churn internally by avoiding a chunk to be repeatedly allocated and freed as we push/pop marks
+	MemStack.Alloc(size_t(FPageAllocator::SmallPageSize) + 1, 16);
+
+	UpdateGraph(*Instance.GetImpl(), ExecuteContext.GetDeltaTime(), InputEventList, OutputEventList);
+	EvaluateGraph(*Instance.GetImpl(), RefPose, LOD, Result);
+
+	// We might have appended new input/output events, append them
+	{
+		const int32 NumInputEvents = InputEventList.Num();
+
+		FWriteScopeLock WriteLock(ModuleInstance.EventListLock);
+
+		// Append the new input events
+		for (int32 EventIndex = NumOriginalInputEvents; EventIndex < NumInputEvents; ++EventIndex)
+		{
+			FAnimNextTraitEventPtr& Event = InputEventList[EventIndex];
+			if (Event->IsValid())
+			{
+				ModuleInstance.InputEventList.Push(MoveTemp(Event));
+			}
+		}
+
+		// Append our output events
+		ModuleInstance.OutputEventList.Append(OutputEventList);
+	}
+}
diff --git a/Engine/Plugins/Experimental/AnimNextAnimGraph/Source/AnimNextAnimGraph/Private/TraitInterfaces/IEvaluate.cpp b/Engine/Plugins/Experimental/AnimNextAnimGraph/Source/AnimNextAnimGraph/Private/TraitInterfaces/IEvaluate.cpp
--- a/Engine/Plugins/Experimental/AnimNextAnimGraph/Source/AnimNextAnimGraph/Private/TraitInterfaces/IEvaluate.cpp
+++ b/Engine/Plugins/Experimental/AnimNextAnimGraph/Source/AnimNextAnimGraph/Private/TraitInterfaces/IEvaluate.cpp
@@ -5,7 +5,10 @@
 #include "TraitCore/ExecutionContext.h"
 #include "TraitInterfaces/IHierarchy.h"
 #include "AnimNextAnimGraphStats.h"
+#include "EvaluationVM/EvaluationFlags.h"
+#include "EvaluationVM/EvaluationVM.h"
 #include "Graph/AnimNextGraphInstance.h"
+#include "Graph/AnimNext_LODPose.h"
 
 DEFINE_STAT(STAT_AnimNext_EvaluateGraph);
 
@@ -219,4 +222,35 @@ namespace UE::AnimNext
 
 		return EvaluationProgram;
 	}
+
+	void EvaluateGraph(FAnimNextGraphInstance& GraphInstance, const FReferencePose& RefPose, int32 GraphLODLevel, FAnimNextGraphLODPose& OutputPose)
+	{
+		const FEvaluationProgram EvaluationProgram = EvaluateGraph(GraphInstance);
+
+		FEvaluationVM EvaluationVM(EEvaluationFlags::All, RefPose, GraphLODLevel);
+		bool bHasValidOutput = false;
+
+		if (!EvaluationProgram.IsEmpty())
+		{
+			EvaluationProgram.Execute(EvaluationVM);
+
+			TUniquePtr<FKeyframeState> EvaluatedKeyframe;
+			if (EvaluationVM.PopValue(KEYFRAME_STACK_NAME, EvaluatedKeyframe))
+			{
+				OutputPose.LODPose.CopyFrom(EvaluatedKeyframe->Pose);
+				OutputPose.Curves.CopyFrom(EvaluatedKeyframe->Curves);
+				OutputPose.Attributes.CopyFrom(EvaluatedKeyframe->Attributes);
+				bHasValidOutput = true;
+			}
+		}
+
+		if (!bHasValidOutput)
+		{
+			// We need to output a valid pose, generate one
+			FKeyframeState ReferenceKeyframe = EvaluationVM.MakeReferenceKeyframe(false);
+			OutputPose.LODPose.CopyFrom(ReferenceKeyframe.Pose);
+			OutputPose.Curves.CopyFrom(ReferenceKeyframe.Curves);
+			OutputPose.Attributes.CopyFrom(ReferenceKeyframe.Attributes);
+		}
+	}
 }
diff --git a/Engine/Plugins/Experimental/AnimNextAnimGraph/Source/AnimNextAnimGraph/Private/TraitInterfaces/IGarbageColllection.cpp b/Engine/Plugins/Experimental/AnimNextAnimGraph/Source/AnimNextAnimGraph/Private/TraitInterfaces/IGarbageColllection.cpp
new file mode 100644
--- /dev/null
+++ b/Engine/Plugins/Experimental/AnimNextAnimGraph/Source/AnimNextAnimGraph/Private/TraitInterfaces/IGarbageColllection.cpp
@@ -0,0 +1,45 @@
+// Copyright Epic Games, Inc. All Rights Reserved.
+
+#include "TraitInterfaces/IGarbageCollection.h"
+
+#include "TraitCore/ExecutionContext.h"
+#include "Graph/GC_GraphInstanceComponent.h"
+
+namespace UE::AnimNext
+{
+	AUTO_REGISTER_ANIM_TRAIT_INTERFACE(IGarbageCollection)
+
+#if WITH_EDITOR
+	const FText& IGarbageCollection::GetDisplayName() const
+	{
+		static FText InterfaceName = NSLOCTEXT("TraitInterfaces", "TraitInterface_IGarbageCollection_Name", "Garbage Collection");
+		return InterfaceName;
+	}
+	const FText& IGarbageCollection::GetDisplayShortName() const
+	{
+		static FText InterfaceShortName = NSLOCTEXT("TraitInterfaces", "TraitInterface_IGarbageCollection_ShortName", "GC");
+		return InterfaceShortName;
+	}
+#endif // WITH_EDITOR
+
+	void IGarbageCollection::RegisterWithGC(const FExecutionContext& Context, const FTraitBinding& Binding)
+	{
+		FGCGraphInstanceComponent& Component = Context.GetComponent<FGCGraphInstanceComponent>();
+		Component.Register(Binding.GetTraitPtr());
+	}
+
+	void IGarbageCollection::UnregisterWithGC(const FExecutionContext& Context, const FTraitBinding& Binding)
+	{
+		FGCGraphInstanceComponent& Component = Context.GetComponent<FGCGraphInstanceComponent>();
+		Component.Unregister(Binding.GetTraitPtr());
+	}
+
+	void IGarbageCollection::AddReferencedObjects(const FExecutionContext& Context, const TTraitBinding<IGarbageCollection>& Binding, FReferenceCollector& Collector) const
+	{
+		TTraitBinding<IGarbageCollection> SuperBinding;
+		if (Binding.GetStackInterfaceSuper(SuperBinding))
+		{
+			SuperBinding.AddReferencedObjects(Context, Collector);
+		}
+	}
+}
diff --git a/Engine/Plugins/Experimental/AnimNextAnimGraph/Source/AnimNextAnimGraph/Private/TraitInterfaces/IUpdate.cpp b/Engine/Plugins/Experimental/AnimNextAnimGraph/Source/AnimNextAnimGraph/Private/TraitInterfaces/IUpdate.cpp
--- a/Engine/Plugins/Experimental/AnimNextAnimGraph/Source/AnimNextAnimGraph/Private/TraitInterfaces/IUpdate.cpp
+++ b/Engine/Plugins/Experimental/AnimNextAnimGraph/Source/AnimNextAnimGraph/Private/TraitInterfaces/IUpdate.cpp
@@ -6,6 +6,7 @@
 #include "Graph/GraphInstanceComponent.h"
 #include "AnimNextAnimGraphStats.h"
 #include "Graph/AnimNextGraphInstance.h"
+#include "TraitCore/TraitEventRaising.h"
 
 DEFINE_STAT(STAT_AnimNext_UpdateGraph);
 
diff --git a/Engine/Plugins/Experimental/AnimNextAnimGraph/Source/AnimNextAnimGraph/Public/AnimNextAnimGraphStats.h b/Engine/Plugins/Experimental/AnimNextAnimGraph/Source/AnimNextAnimGraph/Public/AnimNextAnimGraphStats.h
--- a/Engine/Plugins/Experimental/AnimNextAnimGraph/Source/AnimNextAnimGraph/Public/AnimNextAnimGraphStats.h
+++ b/Engine/Plugins/Experimental/AnimNextAnimGraph/Source/AnimNextAnimGraph/Public/AnimNextAnimGraphStats.h
@@ -7,3 +7,10 @@
 
 DECLARE_CYCLE_STAT_EXTERN(TEXT("AnimNext: Update Graph"), STAT_AnimNext_UpdateGraph, STATGROUP_AnimNext, ANIMNEXTANIMGRAPH_API);
 DECLARE_CYCLE_STAT_EXTERN(TEXT("AnimNext: Evaluate Graph"), STAT_AnimNext_EvaluateGraph, STATGROUP_AnimNext, ANIMNEXTANIMGRAPH_API);
+
+DECLARE_CYCLE_STAT_EXTERN(TEXT("AnimNext: Run Graph"), STAT_AnimNext_Run_Graph, STATGROUP_AnimNext, ANIMNEXTANIMGRAPH_API);
+
+DECLARE_CYCLE_STAT_EXTERN(TEXT("AnimNext: Allocate Graph Instance"), STAT_AnimNext_Graph_AllocateInstance, STATGROUP_AnimNext, ANIMNEXTANIMGRAPH_API);
+DECLARE_CYCLE_STAT_EXTERN(TEXT("AnimNext: Graph RigVM"), STAT_AnimNext_Graph_RigVM, STATGROUP_AnimNext, ANIMNEXTANIMGRAPH_API);
+DECLARE_CYCLE_STAT_EXTERN(TEXT("AnimNext: Execute Evaluation Program"), STAT_AnimNext_EvaluationProgram_Execute, STATGROUP_AnimNext, ANIMNEXTANIMGRAPH_API);
+
