AnimNext: Move remaining animation graph code to AnimNextAnimGraph plugin

diff --git a/Engine/Plugins/Experimental/AnimNextAnimGraph/Source/AnimNextAnimGraph/Public/Graph/AnimNextGraphInstance.h b/Engine/Plugins/Experimental/AnimNextAnimGraph/Source/AnimNextAnimGraph/Public/Graph/AnimNextGraphInstance.h
new file mode 100644
--- /dev/null
+++ b/Engine/Plugins/Experimental/AnimNextAnimGraph/Source/AnimNextAnimGraph/Public/Graph/AnimNextGraphInstance.h
@@ -0,0 +1,247 @@
+// Copyright Epic Games, Inc. All Rights Reserved.
+
+#pragma once
+
+#include "CoreMinimal.h"
+#include "DataInterface/AnimNextDataInterfaceInstance.h"
+#include "HAL/CriticalSection.h"
+#include "TraitCore/TraitEvent.h"
+#include "TraitCore/TraitEventList.h"
+#include "TraitCore/TraitPtr.h"
+#include "RigVMCore/RigVMExecuteContext.h"
+
+#include "AnimNextGraphInstance.generated.h"
+
+class FReferenceCollector;
+
+struct FAnimNextGraphInstancePtr;
+struct FRigUnit_AnimNextGraphEvaluator;
+class UAnimNextAnimationGraph;
+class UAnimNextModule;
+struct FAnimNextModuleInstance;
+class FRigVMTraitScope;
+struct FRigVMExtendedExecuteContext;
+struct FAnimNextExecuteContext;
+
+namespace UE::AnimNext
+{
+	class IDataInterfaceHost;
+	struct FExecutionContext;
+	struct FGraphInstanceComponent;
+	struct FLatentPropertyHandle;
+	struct FTraitStackBinding;
+	struct FPlayAnimSlotTrait;
+	struct FBlendStackCoreTrait;
+}
+
+using GraphInstanceComponentMapType = TMap<FName, TSharedPtr<UE::AnimNext::FGraphInstanceComponent>>;
+
+// Represents an instance of an AnimNext graph
+// This struct uses UE reflection because we wish for the GC to keep the graph
+// alive while we own a reference to it. It is not intended to be serialized on disk with a live instance.
+USTRUCT()
+struct ANIMNEXTANIMGRAPH_API FAnimNextGraphInstance : public FAnimNextDataInterfaceInstance
+{
+	GENERATED_BODY()
+
+	// Creates an empty graph instance that doesn't reference anything
+	FAnimNextGraphInstance();
+
+	// No copying, no moving
+	FAnimNextGraphInstance(const FAnimNextGraphInstance&) = delete;
+	FAnimNextGraphInstance& operator=(const FAnimNextGraphInstance&) = delete;
+
+	// If the graph instance is allocated, we release it during destruction
+	~FAnimNextGraphInstance();
+
+	// Releases the graph instance and frees all corresponding memory
+	void Release();
+
+	// Returns true if we have a live graph instance, false otherwise
+	bool IsValid() const;
+
+	// Returns the animation graph used by this instance or nullptr if the instance is invalid
+	const UAnimNextAnimationGraph* GetAnimationGraph() const;
+
+	// Returns the entry point in Graph that this instance corresponds to 
+	FName GetEntryPoint() const;
+	
+	// Returns a weak handle to the root trait instance
+	UE::AnimNext::FWeakTraitPtr GetGraphRootPtr() const;
+
+	// Returns the module instance that owns us or nullptr if we are invalid
+	FAnimNextModuleInstance* GetModuleInstance() const;
+
+	// Returns the parent graph instance that owns us or nullptr for the root graph instance or if we are invalid
+	FAnimNextGraphInstance* GetParentGraphInstance() const;
+
+	// Returns the root graph instance that owns us and the components or nullptr if we are invalid
+	FAnimNextGraphInstance* GetRootGraphInstance() const;
+
+	// Check to see if this instance data matches the provided animation graph
+	bool UsesAnimationGraph(const UAnimNextAnimationGraph* InAnimationGraph) const;
+
+	// Check to see if this instance data matches the provided graph entry point
+	bool UsesEntryPoint(FName InEntryPoint) const;
+	
+	// Returns whether or not this graph instance is the root graph instance or false otherwise
+	bool IsRoot() const;
+
+	// Returns whether or not this graph instance has updated at least once
+	bool HasUpdated() const;
+
+	// Adds strong/hard object references during GC
+	void AddStructReferencedObjects(class FReferenceCollector& Collector);
+
+	// Returns a typed graph instance component, creating it lazily the first time it is queried
+	template<class ComponentType>
+	ComponentType& GetComponent();
+
+	// Returns a typed graph instance component pointer if found or nullptr otherwise
+	template<class ComponentType>
+	ComponentType* TryGetComponent();
+
+	// Returns a typed graph instance component pointer if found or nullptr otherwise
+	template<class ComponentType>
+	const ComponentType* TryGetComponent() const;
+
+	// Returns const iterators to the graph instance component container
+	GraphInstanceComponentMapType::TConstIterator GetComponentIterator() const;
+
+	// Called each time the graph updates
+	void Update();
+
+	// Get the hosting instance, if any, that owns us
+	FAnimNextDataInterfaceInstance* GetHost() const;
+
+	FInstancedPropertyBag& GetVariables() { return Variables; }
+private:
+	// Returns a pointer to the specified component, or nullptr if not found
+	UE::AnimNext::FGraphInstanceComponent* TryGetComponent(int32 ComponentNameHash, FName ComponentName) const;
+
+	// Adds the specified component and returns a reference to it
+	UE::AnimNext::FGraphInstanceComponent& AddComponent(int32 ComponentNameHash, FName ComponentName, TSharedPtr<UE::AnimNext::FGraphInstanceComponent>&& Component);
+
+	// Executes a list of latent RigVM pins and writes the result into the destination pointer (latent handle offsets are using the destination as base)
+	// When frozen, latent handles that can freeze are skipped, all others will execute
+	void ExecuteLatentPins(const TConstArrayView<UE::AnimNext::FLatentPropertyHandle>& LatentHandles, void* DestinationBasePtr, bool bIsFrozen);
+
+#if WITH_EDITORONLY_DATA
+	// During graph compilation, if we have existing graph instances, we freeze them by releasing their memory before thawing them
+	// Freezing is a partial release of resources that retains the necessary information to re-create things safely
+	void Freeze();
+
+	// During graph compilation, once compilation is done we thaw existing graph instances to reallocate their memory
+	void Thaw();
+
+	// Hook into module compilation
+	void OnModuleCompiled(UAnimNextModule* InModule);
+#endif
+
+	template<typename HostType>
+	bool BindToHostHelper(const HostType& InHost, bool bInAutoBind);
+
+	// Bind the variables in the supplied traits in scope to their respective public variables, represented by structs held on a host
+	// Note: struct views are mutable here as bindings allow writes!
+	void BindPublicVariables(TArrayView<FStructView> InHostStructs);
+
+	// Bind the variables in the supplied traits in scope to their respective public variables, so they point at host memory
+	void BindPublicVariables(TConstArrayView<UE::AnimNext::IDataInterfaceHost*> InHosts);
+
+	// Unbind any public variables that were pointing at host memory and re-point them at the internal defaults
+	void UnbindPublicVariables();
+
+	// The entry point in Graph that this instance corresponds to 
+	FName EntryPoint;
+
+	// Hard reference to the graph instance data, we own it
+	UE::AnimNext::FTraitPtr GraphInstancePtr;
+
+	// The module instance that owns the root, us and the components
+	FAnimNextModuleInstance* ModuleInstance = nullptr;
+
+	// The graph instance that owns us
+	FAnimNextGraphInstance* ParentGraphInstance = nullptr;
+
+	// The root graph instance that owns us and the components
+	FAnimNextGraphInstance* RootGraphInstance = nullptr;
+
+#if WITH_EDITORONLY_DATA
+	struct FCachedVariableBinding
+	{
+		FName VariableName;
+		uint8* Memory = nullptr;
+	};
+
+	// Cached public variable bindings used to correctly thaw instances with input pin bindings
+	TArray<FCachedVariableBinding> CachedVariableBindings;
+#endif
+
+	// Graph instance components that persist from update to update
+	GraphInstanceComponentMapType Components;
+
+	// The current state of public variable bindings to the host
+	enum class EPublicVariablesState : uint8
+	{
+		None,		// No public variables present
+		Unbound,	// Present, but currently unbound
+		Bound		// Present and bound
+	};
+
+	EPublicVariablesState PublicVariablesState;
+
+	// Whether or not this graph has updated once
+	bool bHasUpdatedOnce : 1 = false;
+
+	friend UAnimNextAnimationGraph;			// The graph is the one that allocates instances
+	friend FRigUnit_AnimNextGraphEvaluator;	// We evaluate the instance
+	friend UE::AnimNext::FExecutionContext;
+	friend FAnimNextGraphInstancePtr;
+	friend UE::AnimNext::FTraitStackBinding;
+	friend UE::AnimNext::FPlayAnimSlotTrait;
+	friend UE::AnimNext::FBlendStackCoreTrait;
+};
+
+template<>
+struct TStructOpsTypeTraits<FAnimNextGraphInstance> : public TStructOpsTypeTraitsBase2<FAnimNextGraphInstance>
+{
+	enum
+	{
+		WithAddStructReferencedObjects = true,
+		WithCopy = false,
+	};
+};
+
+//////////////////////////////////////////////////////////////////////////
+
+template<class ComponentType>
+ComponentType& FAnimNextGraphInstance::GetComponent()
+{
+	const FName ComponentName = ComponentType::StaticComponentName();
+	const int32 ComponentNameHash = GetTypeHash(ComponentName);
+
+	if (UE::AnimNext::FGraphInstanceComponent* Component = TryGetComponent(ComponentNameHash, ComponentName))
+	{
+		return *static_cast<ComponentType*>(Component);
+	}
+
+	return static_cast<ComponentType&>(AddComponent(ComponentNameHash, ComponentName, MakeShared<ComponentType>(*this)));
+}
+
+template<class ComponentType>
+ComponentType* FAnimNextGraphInstance::TryGetComponent()
+{
+	const FName ComponentName = ComponentType::StaticComponentName();
+	const int32 ComponentNameHash = GetTypeHash(ComponentName);
+
+	return static_cast<ComponentType*>(TryGetComponent(ComponentNameHash, ComponentName));
+}
+
+template<class ComponentType>
+const ComponentType* FAnimNextGraphInstance::TryGetComponent() const
+{
+	const FName ComponentName = ComponentType::StaticComponentName();
+	const int32 ComponentNameHash = GetTypeHash(ComponentName);
+
+	return static_cast<ComponentType*>(TryGetComponent(ComponentNameHash, ComponentName));
+}
diff --git a/Engine/Plugins/Experimental/AnimNextAnimGraph/Source/AnimNextAnimGraph/Public/Graph/AnimNextGraphInstancePtr.h b/Engine/Plugins/Experimental/AnimNextAnimGraph/Source/AnimNextAnimGraph/Public/Graph/AnimNextGraphInstancePtr.h
new file mode 100644
--- /dev/null
+++ b/Engine/Plugins/Experimental/AnimNextAnimGraph/Source/AnimNextAnimGraph/Public/Graph/AnimNextGraphInstancePtr.h
@@ -0,0 +1,161 @@
+// Copyright Epic Games, Inc. All Rights Reserved.
+
+#pragma once
+
+#include "CoreMinimal.h"
+#include "TraitCore/TraitEvent.h"
+#include "TraitCore/TraitPtr.h"
+
+#include "AnimNextGraphInstancePtr.generated.h"
+
+class FReferenceCollector;
+
+struct FAnimNextGraphInstance;
+struct FRigUnit_AnimNextGraphEvaluator;
+struct FRigUnit_AnimNextRunAnimationGraph;
+class UAnimNextAnimationGraph;
+class FRigVMTraitScope;
+struct FRigVMExtendedExecuteContext;
+struct FAnimNextExecuteContext;
+
+namespace UE::AnimNext
+{
+	struct FExecutionContext;
+	struct FGraphInstanceComponent;
+	class IDataInterfaceHost;
+	struct FAnimationGraphComponent;
+}
+
+using GraphInstanceComponentMapType = TMap<FName, TSharedPtr<UE::AnimNext::FGraphInstanceComponent>>;
+
+// Represents an instance of an AnimNext graph
+// This struct uses UE reflection because we wish for the GC to keep the graph
+// alive while we own a reference to it. It is not intended to be serialized on disk with a live instance.
+USTRUCT(BlueprintType, DisplayName="Animation Graph Instance")
+struct ANIMNEXTANIMGRAPH_API FAnimNextGraphInstancePtr
+{
+	GENERATED_BODY()
+
+	// Creates an empty graph instance that doesn't reference anything
+	FAnimNextGraphInstancePtr();
+
+	// If the graph instance is allocated, we release it during destruction
+	~FAnimNextGraphInstancePtr();
+
+	// Releases the weak ptr to the graph instance
+	void Release();
+
+	// Returns true if we have a live graph instance, false otherwise
+	bool IsValid() const;
+
+	// Returns the animation graph that will be used by this instance
+	const UAnimNextAnimationGraph* GetAnimationGraph() const;
+
+	// Returns a weak handle to the root trait instance
+	UE::AnimNext::FWeakTraitPtr GetGraphRootPtr() const;
+
+	// Returns the graph instance implementation. If the instance is invalid (i.e. its owner has released it) this returns nullptr
+	FAnimNextGraphInstance* GetImpl() const;
+
+	// Check to see if this instance data matches the provided animation graph
+	bool UsesAnimationGraph(const UAnimNextAnimationGraph* InAnimationGraph) const;
+
+	// Returns whether or not this graph instance is the root graph instance or false otherwise
+	bool IsRoot() const;
+
+	// Returns whether or not this graph instance has updated at least once
+	bool HasUpdated() const;
+
+	// Adds strong/hard object references during GC
+	void AddStructReferencedObjects(class FReferenceCollector& Collector);
+
+	// Returns a typed graph instance component, creating it lazily the first time it is queried
+	template<class ComponentType>
+	ComponentType& GetComponent();
+
+	// Returns a typed graph instance component pointer if found or nullptr otherwise
+	template<class ComponentType>
+	ComponentType* TryGetComponent();
+
+	// Returns a typed graph instance component pointer if found or nullptr otherwise
+	template<class ComponentType>
+	const ComponentType* TryGetComponent() const;
+
+	// Returns const iterators to the graph instance component container
+	GraphInstanceComponentMapType::TConstIterator GetComponentIterator() const;
+
+	// Called each time the graph updates
+	void Update() const;
+
+	// Get the extended execute context that we own
+	FRigVMExtendedExecuteContext& GetExtendedExecuteContext() const;
+	
+	// Whether public variables require a binding
+	bool RequiresPublicVariableBinding() const;
+
+	// Bind the variables in the supplied traits in scope to their respective public variables
+	void BindPublicVariables(TConstArrayView<UE::AnimNext::IDataInterfaceHost*> InHosts) const;
+
+private:
+	// Private constructor for use by owners
+	FAnimNextGraphInstancePtr(const TSharedRef<FAnimNextGraphInstance>& InInstance);
+
+	// Returns a pointer to the specified component, or nullptr if not found
+	UE::AnimNext::FGraphInstanceComponent* TryGetComponent(int32 ComponentNameHash, FName ComponentName) const;
+
+	// Adds the specified component and returns a reference to it
+	UE::AnimNext::FGraphInstanceComponent& AddComponent(int32 ComponentNameHash, FName ComponentName, TSharedPtr<UE::AnimNext::FGraphInstanceComponent>&& Component);
+
+	// Weak ptr used to ensure references dont persist in work memory when passing this around in RivGM graphs
+	// Shared ptr owned ultimately by the module/sub-graph trait
+	TWeakPtr<FAnimNextGraphInstance> Impl;
+
+	friend UAnimNextAnimationGraph;				// The graph is the one that allocates instances
+	friend UE::AnimNext::FAnimationGraphComponent; // The module component owns the instances
+	friend FRigUnit_AnimNextGraphEvaluator;		// We evaluate the instance
+	friend FRigUnit_AnimNextRunAnimationGraph;	// We evaluate the instance
+	friend UE::AnimNext::FExecutionContext;
+};
+
+template<>
+struct TStructOpsTypeTraits<FAnimNextGraphInstancePtr> : public TStructOpsTypeTraitsBase2<FAnimNextGraphInstancePtr>
+{
+	enum
+	{
+		WithAddStructReferencedObjects = true,
+	};
+};
+
+//////////////////////////////////////////////////////////////////////////
+
+template<class ComponentType>
+ComponentType& FAnimNextGraphInstancePtr::GetComponent()
+{
+	const FName ComponentName = ComponentType::StaticComponentName();
+	const int32 ComponentNameHash = GetTypeHash(ComponentName);
+
+	if (UE::AnimNext::FGraphInstanceComponent* Component = TryGetComponent(ComponentNameHash, ComponentName))
+	{
+		return *static_cast<ComponentType*>(Component);
+	}
+
+	return static_cast<ComponentType&>(AddComponent(ComponentNameHash, ComponentName, MakeShared<ComponentType>()));
+}
+
+template<class ComponentType>
+ComponentType* FAnimNextGraphInstancePtr::TryGetComponent()
+{
+	const FName ComponentName = ComponentType::StaticComponentName();
+	const int32 ComponentNameHash = GetTypeHash(ComponentName);
+
+	return static_cast<ComponentType*>(TryGetComponent(ComponentNameHash, ComponentName));
+}
+
+template<class ComponentType>
+const ComponentType* FAnimNextGraphInstancePtr::TryGetComponent() const
+{
+	const FName ComponentName = ComponentType::StaticComponentName();
+	const int32 ComponentNameHash = GetTypeHash(ComponentName);
+
+	return static_cast<ComponentType*>(TryGetComponent(ComponentNameHash, ComponentName));
+}
diff --git a/Engine/Plugins/Experimental/AnimNextAnimGraph/Source/AnimNextAnimGraph/Public/TraitCore/TraitEventRaising.h b/Engine/Plugins/Experimental/AnimNextAnimGraph/Source/AnimNextAnimGraph/Public/TraitCore/TraitEventRaising.h
new file mode 100644
--- /dev/null
+++ b/Engine/Plugins/Experimental/AnimNextAnimGraph/Source/AnimNextAnimGraph/Public/TraitCore/TraitEventRaising.h
@@ -0,0 +1,26 @@
+﻿// Copyright Epic Games, Inc. All Rights Reserved.
+
+#pragma once
+
+#include "CoreMinimal.h"
+
+struct FAnimNextTraitEvent;
+
+namespace UE::AnimNext
+{
+	struct FExecutionContext;
+	struct FTraitStackBinding;
+	struct FTraitEventList;
+	struct FGraphInstanceComponent;
+
+	// Raises a trait event on the specified trait stack
+	// This calls OnTraitEvent on each trait of the stack following propagation rules
+	ANIMNEXTANIMGRAPH_API void RaiseTraitEvent(FExecutionContext& Context, const FTraitStackBinding& StackBinding, FAnimNextTraitEvent& Event);
+
+	// Raises every trait event from the input list on the specified trait stack
+	// This calls OnTraitEvent on each trait of the stack following propagation rules
+	ANIMNEXTANIMGRAPH_API void RaiseTraitEvents(FExecutionContext& Context, const FTraitStackBinding& StackBinding, const FTraitEventList& EventList);
+
+	// Raises every trait event from the input list on the specified component
+	ANIMNEXTANIMGRAPH_API void RaiseTraitEvents(FExecutionContext& Context, FGraphInstanceComponent& Component, const FTraitEventList& EventList);
+}
diff --git a/Engine/Plugins/Experimental/AnimNextAnimGraph/Source/AnimNextAnimGraph/Public/TraitInterfaces/IEvaluate.h b/Engine/Plugins/Experimental/AnimNextAnimGraph/Source/AnimNextAnimGraph/Public/TraitInterfaces/IEvaluate.h
--- a/Engine/Plugins/Experimental/AnimNextAnimGraph/Source/AnimNextAnimGraph/Public/TraitInterfaces/IEvaluate.h
+++ b/Engine/Plugins/Experimental/AnimNextAnimGraph/Source/AnimNextAnimGraph/Public/TraitInterfaces/IEvaluate.h
@@ -9,6 +9,8 @@
 #include "EvaluationVM/EvaluationProgram.h"
 #include "EvaluationVM/KeyframeState.h"
 
+struct FAnimNextGraphLODPose;
+
 namespace UE::AnimNext
 {
 	/**
@@ -113,4 +115,9 @@ namespace UE::AnimNext
 	 * @see IEvaluate::PreEvaluate, IEvaluate::PostEvaluate, IHierarchy::GetChildren
 	 */
 	[[nodiscard]] ANIMNEXTANIMGRAPH_API FEvaluationProgram EvaluateGraph(const FWeakTraitPtr& GraphRootPtr);
+
+	/**
+	 * Evaluates a sub-graph starting at its root and produces an output pose
+	 */
+	ANIMNEXTANIMGRAPH_API void EvaluateGraph(FAnimNextGraphInstance& GraphInstance, const FReferencePose& RefPose, int32 GraphLODLevel, FAnimNextGraphLODPose& OutputPose);
 }
diff --git a/Engine/Plugins/Experimental/AnimNextAnimGraph/Source/AnimNextAnimGraph/Public/TraitInterfaces/IGarbageCollection.h b/Engine/Plugins/Experimental/AnimNextAnimGraph/Source/AnimNextAnimGraph/Public/TraitInterfaces/IGarbageCollection.h
new file mode 100644
--- /dev/null
+++ b/Engine/Plugins/Experimental/AnimNextAnimGraph/Source/AnimNextAnimGraph/Public/TraitInterfaces/IGarbageCollection.h
@@ -0,0 +1,57 @@
+// Copyright Epic Games, Inc. All Rights Reserved.
+
+#pragma once
+
+#include "CoreMinimal.h"
+#include "TraitCore/ITraitInterface.h"
+#include "TraitCore/TraitBinding.h"
+
+class FReferenceCollector;
+
+namespace UE::AnimNext
+{
+	/**
+	 * IGarbageCollection
+	 * 
+	 * This interface exposes garbage collection reference tracking.
+	 */
+	struct ANIMNEXTANIMGRAPH_API IGarbageCollection : ITraitInterface
+	{
+		DECLARE_ANIM_TRAIT_INTERFACE(IGarbageCollection, 0x231a2017)
+
+
+		// Registers the provided binding for GC callback
+		// Once registered, AddReferencedObjects is called during GC to collect references
+		static void RegisterWithGC(const FExecutionContext& Context, const FTraitBinding& Binding);
+
+		// Unregisters the provided binding from GC callback
+		static void UnregisterWithGC(const FExecutionContext& Context, const FTraitBinding& Binding);
+
+		// Called when garbage collection requests hard/strong object references
+		// @see UObject::AddReferencedObjects
+		virtual void AddReferencedObjects(const FExecutionContext& Context, const TTraitBinding<IGarbageCollection>& Binding, FReferenceCollector& Collector) const;
+
+#if WITH_EDITOR
+		virtual bool IsInternal() const override { return true; }
+
+		virtual const FText& GetDisplayName() const override;
+		virtual const FText& GetDisplayShortName() const override;
+#endif // WITH_EDITOR
+	};
+
+	/**
+	 * Specialization for trait binding.
+	 */
+	template<>
+	struct TTraitBinding<IGarbageCollection> : FTraitBinding
+	{
+		// @see IGarbageCollection::AddReferencedObjects
+		void AddReferencedObjects(const FExecutionContext& Context, FReferenceCollector& Collector) const
+		{
+			GetInterface()->AddReferencedObjects(Context, *this, Collector);
+		}
+
+	protected:
+		const IGarbageCollection* GetInterface() const { return GetInterfaceTyped<IGarbageCollection>(); }
+	};
+}
diff --git a/Engine/Plugins/Experimental/AnimNextAnimGraph/Source/AnimNextAnimGraph/Public/TraitInterfaces/IHierarchy.h b/Engine/Plugins/Experimental/AnimNextAnimGraph/Source/AnimNextAnimGraph/Public/TraitInterfaces/IHierarchy.h
new file mode 100644
--- /dev/null
+++ b/Engine/Plugins/Experimental/AnimNextAnimGraph/Source/AnimNextAnimGraph/Public/TraitInterfaces/IHierarchy.h
@@ -0,0 +1,89 @@
+// Copyright Epic Games, Inc. All Rights Reserved.
+
+#pragma once
+
+#include "CoreMinimal.h"
+#include "Misc/MemStack.h"
+#include "TraitCore/ITraitInterface.h"
+#include "TraitCore/TraitBinding.h"
+
+namespace UE::AnimNext
+{
+	// An array of children pointers
+	// We reserve a small amount inline and spill on the memstack
+	using FChildrenArray = TArray<FWeakTraitPtr, TInlineAllocator<8, TMemStackAllocator<>>>;
+
+	/**
+	 * IHierarchy
+	 * 
+	 * This interface exposes hierarchy traversal information to navigate the graph.
+	 */
+	struct IHierarchy : ITraitInterface
+	{
+		DECLARE_ANIM_TRAIT_INTERFACE(IHierarchy, 0x846d8a37)
+
+		// Returns the number of children of the trait implementation (not the whole stack)
+		// Includes inactive children
+		ANIMNEXTANIMGRAPH_API virtual uint32 GetNumChildren(const FExecutionContext& Context, const TTraitBinding<IHierarchy>& Binding) const;
+
+		// Appends weak handles to any children we wish to traverse on the trait implementation (not the whole stack).
+		// Traits are responsible for allocating and releasing child instance data.
+		// Empty handles can be appended.
+		ANIMNEXTANIMGRAPH_API virtual void GetChildren(const FExecutionContext& Context, const TTraitBinding<IHierarchy>& Binding, FChildrenArray& Children) const;
+
+		// Queries the trait stack and calls GetChildren for each trait, appending the result.
+		ANIMNEXTANIMGRAPH_API static void GetStackChildren(const FExecutionContext& Context, const FTraitStackBinding& Binding, FChildrenArray& Children);
+
+		// Queries the trait stack of the specified binding and calls GetChildren for each trait, appending the result.
+		ANIMNEXTANIMGRAPH_API static void GetStackChildren(const FExecutionContext& Context, const FTraitBinding& Binding, FChildrenArray& Children);
+
+		// Queries the trait stack and calls GetNumChildren for each trait, accumulating the result.
+		ANIMNEXTANIMGRAPH_API static uint32 GetNumStackChildren(const FExecutionContext& Context, const FTraitStackBinding& Binding);
+
+		// Queries the trait stack of the specified binding and calls GetNumChildren for each trait, accumulating the result.
+		ANIMNEXTANIMGRAPH_API static uint32 GetNumStackChildren(const FExecutionContext& Context, const FTraitBinding& Binding);
+
+#if WITH_EDITOR
+		ANIMNEXTANIMGRAPH_API virtual const FText& GetDisplayName() const override;
+		ANIMNEXTANIMGRAPH_API virtual const FText& GetDisplayShortName() const override;
+#endif // WITH_EDITOR
+	};
+
+	/**
+	 * Specialization for trait binding.
+	 */
+	template<>
+	struct TTraitBinding<IHierarchy> : FTraitBinding
+	{
+		// @see IHierarchy::GetNumChildren
+		uint32 GetNumChildren(const FExecutionContext& Context) const
+		{
+			return GetInterface()->GetNumChildren(Context, *this);
+		}
+
+		// @see IHierarchy::GetChildren
+		void GetChildren(const FExecutionContext& Context, FChildrenArray& Children) const
+		{
+			GetInterface()->GetChildren(Context, *this, Children);
+		}
+
+	protected:
+		const IHierarchy* GetInterface() const { return GetInterfaceTyped<IHierarchy>(); }
+	};
+
+	//////////////////////////////////////////////////////////////////////////
+	// Inline implementations
+
+	inline void IHierarchy::GetStackChildren(const FExecutionContext& Context, const FTraitBinding& Binding, FChildrenArray& Children)
+	{
+		if (Binding.IsValid())
+		{
+			GetStackChildren(Context, *Binding.GetStack(), Children);
+		}
+	}
+
+	inline uint32 IHierarchy::GetNumStackChildren(const FExecutionContext& Context, const FTraitBinding& Binding)
+	{
+		return Binding.IsValid() ? GetNumStackChildren(Context, *Binding.GetStack()) : 0;
+	}
+}
diff --git a/Engine/Plugins/Experimental/AnimNextAnimGraph/Source/AnimNextAnimGraphEditor/AnimNextAnimGraphEditor.Build.cs b/Engine/Plugins/Experimental/AnimNextAnimGraph/Source/AnimNextAnimGraphEditor/AnimNextAnimGraphEditor.Build.cs
--- a/Engine/Plugins/Experimental/AnimNextAnimGraph/Source/AnimNextAnimGraphEditor/AnimNextAnimGraphEditor.Build.cs
+++ b/Engine/Plugins/Experimental/AnimNextAnimGraph/Source/AnimNextAnimGraphEditor/AnimNextAnimGraphEditor.Build.cs
@@ -10,9 +10,33 @@ namespace UnrealBuildTool.Rules
 				new string[]
 				{
 					"Core",
+				}
+			);
+
+			PrivateDependencyModuleNames.AddRange(
+				new string[]
+				{
 					"CoreUObject",
+					"Engine",
+					"UnrealEd",
+					"InputCore",
+					"SlateCore",
+					"Slate",
+					"AnimNext", 
+					"AnimNextEditor", 
+					"AnimNextUncookedOnly",
 					"AnimNextAnimGraph",
+					"AnimNextAnimGraphUncookedOnly",
 					"Settings",
+					"WorkspaceEditor", 
+					"EditorWidgets", 
+					"MessageLog", 
+					"AssetDefinition",
+					"RigVM",
+					"RigVMDeveloper",
+					"ToolMenus", 
+					"RigVMEditor",
+					"GraphEditor",
 				}
 			);
 		}
diff --git a/Engine/Plugins/Experimental/AnimNextAnimGraph/Source/AnimNextAnimGraphEditor/Private/AnimNextAnimGraphEditorModule.cpp b/Engine/Plugins/Experimental/AnimNextAnimGraph/Source/AnimNextAnimGraphEditor/Private/AnimNextAnimGraphEditorModule.cpp
--- a/Engine/Plugins/Experimental/AnimNextAnimGraph/Source/AnimNextAnimGraphEditor/Private/AnimNextAnimGraphEditorModule.cpp
+++ b/Engine/Plugins/Experimental/AnimNextAnimGraph/Source/AnimNextAnimGraphEditor/Private/AnimNextAnimGraphEditorModule.cpp
@@ -1,9 +1,27 @@
 // Copyright Epic Games, Inc. All Rights Reserved.
 
+#include "AnimNextAnimGraphWorkspaceAssetUserData.h"
 #include "Modules/ModuleInterface.h"
 #include "Modules/ModuleManager.h"
 #include "ISettingsModule.h"
 #include "AnimNextAnimGraphSettings.h"
+#include "EdGraphUtilities.h"
+#include "IAnimNextEditorModule.h"
+#include "IWorkspaceEditorModule.h"
+#include "TraitStackEditor.h"
+#include "UncookedOnlyUtils.h"
+#include "Entries/AnimNextAnimationGraphEntry.h"
+#include "Features/IModularFeatures.h"
+#include "Framework/Docking/LayoutExtender.h"
+#include "Graph/AnimNextAnimationGraph.h"
+#include "Graph/AnimNextAnimationGraphItemDetails.h"
+#include "Graph/AnimNextAnimationGraphMenuExtensions.h"
+#include "Graph/AnimNextAnimationGraph_EditorData.h"
+#include "Graph/AnimNextGraphPanelNodeFactory.h"
+#include "Graph/RigUnit_AnimNextGraphRoot.h"
+#include "Graph/RigVMTrait_AnimNextPublicVariablesEditor.h"
+#include "Graph/TraitEditorTabSummoner.h"
+#include "WorkflowOrientedApp/WorkflowTabManager.h"
 
 #define LOCTEXT_NAMESPACE "FAnimNextAnimGraphEditorModule"
 
@@ -20,6 +38,51 @@ class FAnimNextAnimGraphEditorModule : public IModuleInterface
 			LOCTEXT("SettingsName", "AnimNext Anim Graph"),
 			LOCTEXT("SettingsDescription", "Configure options for AnimNext animation graphs."),
 			GetMutableDefault<UAnimNextAnimGraphSettings>());
+
+		FPublicVariablesImpl::Register();
+
+		Workspace::IWorkspaceEditorModule& WorkspaceEditorModule = FModuleManager::Get().LoadModuleChecked<Workspace::IWorkspaceEditorModule>("WorkspaceEditor");
+
+		// --- AnimNextAnimationGraph ---
+		Workspace::FObjectDocumentArgs AnimNextAnimationGraphDocumentArgs(
+			Workspace::FOnRedirectWorkspaceContext::CreateLambda([](UObject* InObject)
+			{
+				UAnimNextAnimationGraph* AnimationGraph = CastChecked<UAnimNextAnimationGraph>(InObject);
+				UAnimNextAnimationGraph_EditorData* EditorData = UncookedOnly::FUtils::GetEditorData<UAnimNextAnimationGraph_EditorData>(AnimationGraph);
+
+				// Redirect to the inner graph
+				UAnimNextAnimationGraphEntry* AnimationGraphEntry = CastChecked<UAnimNextAnimationGraphEntry>(EditorData->FindEntry(FRigUnit_AnimNextGraphRoot::DefaultEntryPoint));
+				return AnimationGraphEntry->GetEdGraph();
+			}));
+
+		WorkspaceEditorModule.RegisterObjectDocumentType(FTopLevelAssetPath(TEXT("/Script/AnimNextAnimGraph.AnimNextAnimationGraph")), AnimNextAnimationGraphDocumentArgs);
+		
+		WorkspaceEditorModule.OnRegisterTabsForEditor().AddLambda([](FWorkflowAllowedTabSet& TabFactories, const TSharedRef<FTabManager>& InTabManager, TSharedPtr<UE::Workspace::IWorkspaceEditor> InEditorPtr)
+		{
+			TSharedRef<FTraitEditorTabSummoner> TraitEditorTabSummoner = MakeShared<FTraitEditorTabSummoner>(InEditorPtr);
+			TabFactories.RegisterFactory(TraitEditorTabSummoner);
+			TraitEditorTabSummoner->RegisterTabSpawner(InTabManager, nullptr);
+		});
+		
+		WorkspaceEditorModule.OnExtendTabs().AddLambda([](FLayoutExtender& InLayoutExtender, TSharedPtr<UE::Workspace::IWorkspaceEditor> InEditorPtr)
+		{
+			FTabManager::FTab TraitEditorTab(FTabId(TraitEditorTabName), ETabState::ClosedTab);
+			InLayoutExtender.ExtendLayout(FTabId(Workspace::WorkspaceTabs::TopRightDocumentArea), ELayoutExtensionPosition::After, TraitEditorTab);
+		});
+		
+		const TSharedPtr<FAnimNextAnimationGraphItemDetails> AssetItemDetails = MakeShared<FAnimNextAnimationGraphItemDetails>();
+		WorkspaceEditorModule.RegisterWorkspaceItemDetails(Workspace::FOutlinerItemDetailsId(FAnimNextAnimationGraphOutlinerData::StaticStruct()->GetFName()), AssetItemDetails);
+
+		IAnimNextEditorModule& AnimNextEditorModule = FModuleManager::LoadModuleChecked<IAnimNextEditorModule>("AnimNextEditor");
+		AnimNextEditorModule.AddWorkspaceSupportedAssetClass(UAnimNextAnimationGraph::StaticClass()->GetClassPathName());
+
+		TraitStackEditor = MakeShared<FTraitStackEditor>();
+		IModularFeatures::Get().RegisterModularFeature(ITraitStackEditor::ModularFeatureName, TraitStackEditor.Get());
+
+		AnimNextGraphPanelNodeFactory = MakeShared<FAnimNextGraphPanelNodeFactory>();
+		FEdGraphUtilities::RegisterVisualNodeFactory(AnimNextGraphPanelNodeFactory);
+
+		FAnimationGraphMenuExtensions::RegisterMenus();
 	}
 
 	virtual void ShutdownModule() override
@@ -29,7 +92,30 @@ class FAnimNextAnimGraphEditorModule : public IModuleInterface
 		{
 			SettingsModule->UnregisterSettings("Project", "Plugins", "AnimNextAnimGraph");
 		}
+
+		if(FModuleManager::Get().IsModuleLoaded("WorkspaceEditor"))
+		{
+			Workspace::IWorkspaceEditorModule& WorkspaceEditorModule = FModuleManager::LoadModuleChecked<Workspace::IWorkspaceEditorModule>("WorkspaceEditor");
+			WorkspaceEditorModule.UnregisterObjectDocumentType(FTopLevelAssetPath(TEXT("/Script/AnimNextAnimGraph.AnimNextAnimationGraph")));
+			if(UObjectInitialized())
+			{
+				WorkspaceEditorModule.UnregisterWorkspaceItemDetails(Workspace::FOutlinerItemDetailsId(FAnimNextAnimationGraphOutlinerData::StaticStruct()->GetFName()));
+			}
+		}
+		
+		IModularFeatures::Get().UnregisterModularFeature(ITraitStackEditor::ModularFeatureName, TraitStackEditor.Get());
+		TraitStackEditor.Reset();
+
+		FEdGraphUtilities::UnregisterVisualNodeFactory(AnimNextGraphPanelNodeFactory);
+		
+		FAnimationGraphMenuExtensions::UnregisterMenus();
 	}
+
+	/** Trait stack editor modular feature */
+	TSharedPtr<FTraitStackEditor> TraitStackEditor;
+
+	/** Node factory for the AnimNext graph */
+	TSharedPtr<FAnimNextGraphPanelNodeFactory> AnimNextGraphPanelNodeFactory;
 };
 
 }
diff --git a/Engine/Plugins/Experimental/AnimNextAnimGraph/Source/AnimNextAnimGraphEditor/Private/Graph/AnimNextAnimGraphEntryAssetDefinitions.cpp b/Engine/Plugins/Experimental/AnimNextAnimGraph/Source/AnimNextAnimGraphEditor/Private/Graph/AnimNextAnimGraphEntryAssetDefinitions.cpp
new file mode 100644
--- /dev/null
+++ b/Engine/Plugins/Experimental/AnimNextAnimGraph/Source/AnimNextAnimGraphEditor/Private/Graph/AnimNextAnimGraphEntryAssetDefinitions.cpp
@@ -0,0 +1,13 @@
+﻿// Copyright Epic Games, Inc. All Rights Reserved.
+
+#include "AnimNextAnimGraphEntryAssetDefinitions.h"
+
+#define LOCTEXT_NAMESPACE "AnimNextAnimGraphEntryAssetDefinitions"
+
+FText UAssetDefinition_AnimNextAnimationGraphEntry::GetObjectDisplayNameText(UObject* Object) const
+{
+	UAnimNextRigVMAssetEntry* Parameter = CastChecked<UAnimNextRigVMAssetEntry>(Object);
+	return Parameter->GetDisplayName();
+}
+
+#undef LOCTEXT_NAMESPACE
\ No newline at end of file
diff --git a/Engine/Plugins/Experimental/AnimNextAnimGraph/Source/AnimNextAnimGraphEditor/Private/Graph/AnimNextAnimationGraphAssetDefinition.h b/Engine/Plugins/Experimental/AnimNextAnimGraph/Source/AnimNextAnimGraphEditor/Private/Graph/AnimNextAnimationGraphAssetDefinition.h
new file mode 100644
--- /dev/null
+++ b/Engine/Plugins/Experimental/AnimNextAnimGraph/Source/AnimNextAnimGraphEditor/Private/Graph/AnimNextAnimationGraphAssetDefinition.h
@@ -0,0 +1,30 @@
+﻿// Copyright Epic Games, Inc. All Rights Reserved.
+
+#pragma once
+
+#include "Graph/AnimNextAnimationGraph.h"
+#include "AssetDefinitionDefault.h"
+#include "AnimNextAnimationGraphAssetDefinition.generated.h"
+
+#define LOCTEXT_NAMESPACE "AnimNextAssetDefinitions"
+
+UCLASS()
+class UAssetDefinition_AnimNextAnimationGraph : public UAssetDefinitionDefault
+{
+	GENERATED_BODY()
+
+public:
+	// UAssetDefinition interface
+	virtual FText GetAssetDisplayName() const override { return LOCTEXT("AnimNextAnimationGraph", "AnimNext Animation Graph"); }
+	virtual FLinearColor GetAssetColor() const override { return FLinearColor(FColor(128,96,48)); }
+	virtual TSoftClassPtr<UObject> GetAssetClass() const override { return UAnimNextAnimationGraph::StaticClass(); }
+	virtual EAssetCommandResult OpenAssets(const FAssetOpenArgs& OpenArgs) const override;
+	virtual TConstArrayView<FAssetCategoryPath> GetAssetCategories() const override
+	{
+		static const auto Categories = { FAssetCategoryPath(EAssetCategoryPaths::Animation, LOCTEXT("AnimNextSubMenu", "AnimNext")) };
+		return Categories;
+	}
+	virtual bool ShouldSaveExternalPackages() const override { return true; }
+};
+
+#undef LOCTEXT_NAMESPACE
\ No newline at end of file
diff --git a/Engine/Plugins/Experimental/AnimNextAnimGraph/Source/AnimNextAnimGraphEditor/Private/Graph/AnimNextAnimationGraphFactory.cpp b/Engine/Plugins/Experimental/AnimNextAnimGraph/Source/AnimNextAnimGraphEditor/Private/Graph/AnimNextAnimationGraphFactory.cpp
new file mode 100644
--- /dev/null
+++ b/Engine/Plugins/Experimental/AnimNextAnimGraph/Source/AnimNextAnimGraphEditor/Private/Graph/AnimNextAnimationGraphFactory.cpp
@@ -0,0 +1,48 @@
+﻿// Copyright Epic Games, Inc. All Rights Reserved.
+
+#include "Graph/AnimNextAnimationGraphFactory.h"
+#include "Graph/AnimNextAnimationGraph.h"
+#include "Graph/AnimNextAnimationGraph_EditorData.h"
+#include "UncookedOnlyUtils.h"
+#include "Entries/AnimNextAnimationGraphEntry.h"
+
+UAnimNextAnimationGraphFactory::UAnimNextAnimationGraphFactory()
+{
+	bCreateNew = true;
+	bEditAfterNew = true;
+	SupportedClass = UAnimNextAnimationGraph::StaticClass();
+}
+
+bool UAnimNextAnimationGraphFactory::ConfigureProperties()
+{
+	return true;
+}
+
+UObject* UAnimNextAnimationGraphFactory::FactoryCreateNew(UClass* Class, UObject* InParent, FName Name, EObjectFlags Flags, UObject* Context, FFeedbackContext* Warn, FName CallingContext)
+{
+	EObjectFlags FlagsToUse = Flags | RF_Public | RF_Standalone | RF_Transactional | RF_LoadCompleted;
+	if(InParent == GetTransientPackage())
+	{
+		FlagsToUse &= ~RF_Standalone;
+	}
+
+	UAnimNextAnimationGraph* NewAnimationGraph = NewObject<UAnimNextAnimationGraph>(InParent, Class, Name, FlagsToUse);
+
+	// Create internal editor data
+	UAnimNextAnimationGraph_EditorData* EditorData = NewObject<UAnimNextAnimationGraph_EditorData>(NewAnimationGraph, TEXT("EditorData"), RF_Transactional);
+	NewAnimationGraph->EditorData = EditorData;
+	EditorData->bUsesExternalPackages = false;
+	EditorData->Initialize(/*bRecompileVM*/false);
+
+	// Add a single internal graph entry
+	UAnimNextAnimationGraphEntry* NewGraph = EditorData->AddAnimationGraph(FRigUnit_AnimNextGraphRoot::DefaultEntryPoint);
+
+	// Hide in the outliner - the asset itself represents the graph to users
+	NewGraph->SetHiddenInOutliner(true);
+
+	// Compile the initial skeleton
+	EditorData->RecompileVM();
+	check(!EditorData->bErrorsDuringCompilation);
+
+	return NewAnimationGraph;
+}
\ No newline at end of file
diff --git a/Engine/Plugins/Experimental/AnimNextAnimGraph/Source/AnimNextAnimGraphEditor/Private/Graph/AnimNextAnimationGraphItemDetails.cpp b/Engine/Plugins/Experimental/AnimNextAnimGraph/Source/AnimNextAnimGraphEditor/Private/Graph/AnimNextAnimationGraphItemDetails.cpp
new file mode 100644
--- /dev/null
+++ b/Engine/Plugins/Experimental/AnimNextAnimGraph/Source/AnimNextAnimGraphEditor/Private/Graph/AnimNextAnimationGraphItemDetails.cpp
@@ -0,0 +1,102 @@
+﻿// Copyright Epic Games, Inc. All Rights Reserved.
+
+#include "AnimNextAnimationGraphItemDetails.h"
+
+#include "AnimNextAnimGraphWorkspaceAssetUserData.h"
+#include "ClassIconFinder.h"
+#include "Toolkits/AssetEditorToolkitMenuContext.h"
+#include "Entries/AnimNextAnimationGraphEntry.h"
+#include "InstancedStruct.h"
+#include "WorkspaceItemMenuContext.h"
+#include "IWorkspaceEditor.h"
+#include "ScopedTransaction.h"
+#include "RigVMModel/RigVMClient.h"
+#include "ToolMenus.h"
+#include "Entries/AnimNextDataInterfaceEntry.h"
+#include "Graph/AnimNextAnimationGraph_EditorData.h"
+
+#define LOCTEXT_NAMESPACE "FAnimNextAnimationGraphItemDetails"
+
+namespace UE::AnimNext::Editor
+{
+
+void FAnimNextAnimationGraphItemDetails::RegisterToolMenuExtensions()
+{
+	FToolMenuOwnerScoped OwnerScoped(TEXT("FAnimNextAnimationGraphItemDetails"));
+	UToolMenu* Menu = UToolMenus::Get()->ExtendMenu("WorkspaceOutliner.ItemContextMenu");
+	if (Menu == nullptr)
+	{
+		return;
+	}
+
+	Menu->AddDynamicSection(TEXT("AnimNextAssetItem"), FNewToolMenuDelegate::CreateLambda([](UToolMenu* InMenu)
+	{
+		UWorkspaceItemMenuContext* WorkspaceItemContext = InMenu->FindContext<UWorkspaceItemMenuContext>();
+		const UAssetEditorToolkitMenuContext* AssetEditorContext = InMenu->FindContext<UAssetEditorToolkitMenuContext>();
+		if (WorkspaceItemContext == nullptr || AssetEditorContext == nullptr)
+		{
+			return;
+		}
+
+		const TSharedPtr<UE::Workspace::IWorkspaceEditor> WorkspaceEditor = StaticCastSharedPtr<UE::Workspace::IWorkspaceEditor>(AssetEditorContext->Toolkit.Pin());
+		if(!WorkspaceEditor.IsValid())
+		{
+			return;
+		}
+
+		if (WorkspaceItemContext->SelectedExports.Num() != 1)
+		{
+			return;
+		}
+
+		TInstancedStruct<FWorkspaceOutlinerItemData>& Data = WorkspaceItemContext->SelectedExports[0].GetData();
+		if (!Data.IsValid() || !Data.GetScriptStruct()->IsChildOf(FAnimNextAnimationGraphOutlinerData::StaticStruct()))
+		{
+			return;
+		}
+
+		const FAnimNextAnimationGraphOutlinerData& OutlinerData = Data.Get<FAnimNextAnimationGraphOutlinerData>();
+		UAnimNextAnimationGraph* Asset = OutlinerData.GetAnimationGraph();
+		if(Asset == nullptr)
+		{
+			return;
+		}
+
+		UAnimNextAnimationGraph_EditorData* EditorData = UncookedOnly::FUtils::GetEditorData<UAnimNextAnimationGraph_EditorData>(Asset);
+		if(EditorData == nullptr)
+		{
+			return;
+		}
+
+		FToolMenuSection& AssetSection = InMenu->AddSection("AnimNextAsset", LOCTEXT("AnimNextAssetSectionLabel", "AnimNext Asset"));
+		const UClass* SubEntryClass = UAnimNextAnimationGraphEntry::StaticClass();
+
+		static FTextFormat AddEntryLabelFormat(LOCTEXT("AddEntryLabelFormat", "Add {0}"));
+		static FTextFormat AddEntryTooltipFormat(LOCTEXT("AddEntryTooltipFormat", "Adds a new {0} to this asset"));
+
+		AssetSection.AddMenuEntry(
+			SubEntryClass->GetFName(),
+			FText::Format(AddEntryLabelFormat, SubEntryClass->GetDisplayNameText()),
+			FText::Format(AddEntryTooltipFormat, SubEntryClass->GetDisplayNameText()),
+			FSlateIconFinder::FindIconForClass(SubEntryClass, "ClassIcon.Object"),
+			FUIAction(
+				FExecuteAction::CreateWeakLambda(WorkspaceItemContext, [&OutlinerData, EditorData]()
+				{
+					FScopedTransaction Transaction(LOCTEXT("AddAnimationGraph", "Add Animation Graph"));
+					EditorData->AddAnimationGraph(TEXT("Root"));
+				})
+			));
+	}));
+}
+
+void FAnimNextAnimationGraphItemDetails::UnregisterToolMenuExtensions()
+{
+	if(UToolMenus* ToolMenus = UToolMenus::Get())
+	{
+		ToolMenus->UnregisterOwnerByName("FAnimNextAnimationGraphItemDetails");
+	}
+}
+
+}
+
+#undef LOCTEXT_NAMESPACE // "FAnimNextGraphItemDetails"
\ No newline at end of file
diff --git a/Engine/Plugins/Experimental/AnimNextAnimGraph/Source/AnimNextAnimGraphEditor/Private/Graph/AnimNextAnimationGraphItemDetails.h b/Engine/Plugins/Experimental/AnimNextAnimGraph/Source/AnimNextAnimGraphEditor/Private/Graph/AnimNextAnimationGraphItemDetails.h
new file mode 100644
--- /dev/null
+++ b/Engine/Plugins/Experimental/AnimNextAnimGraph/Source/AnimNextAnimGraphEditor/Private/Graph/AnimNextAnimationGraphItemDetails.h
@@ -0,0 +1,19 @@
+﻿// Copyright Epic Games, Inc. All Rights Reserved.
+
+#pragma once
+
+#include "Common/AnimNextAssetItemDetails.h"
+
+namespace UE::AnimNext::Editor
+{
+
+class FAnimNextAnimationGraphItemDetails : public FAnimNextAssetItemDetails
+{
+public:
+	FAnimNextAnimationGraphItemDetails() = default;
+
+	static void RegisterToolMenuExtensions();
+	static void UnregisterToolMenuExtensions();
+};
+
+}
diff --git a/Engine/Plugins/Experimental/AnimNextAnimGraph/Source/AnimNextAnimGraphEditor/Private/Graph/AnimNextAnimationGraphMenuExtensions.cpp b/Engine/Plugins/Experimental/AnimNextAnimGraph/Source/AnimNextAnimGraphEditor/Private/Graph/AnimNextAnimationGraphMenuExtensions.cpp
new file mode 100644
--- /dev/null
+++ b/Engine/Plugins/Experimental/AnimNextAnimGraph/Source/AnimNextAnimGraphEditor/Private/Graph/AnimNextAnimationGraphMenuExtensions.cpp
@@ -0,0 +1,249 @@
+﻿// Copyright Epic Games, Inc. All Rights Reserved.
+
+#include "AnimNextAnimationGraphMenuExtensions.h"
+
+#include "AnimGraphUncookedOnlyUtils.h"
+#include "AnimNextEdGraphNode.h"
+#include "ContentBrowserModule.h"
+#include "IContentBrowserSingleton.h"
+#include "ToolMenus.h"
+#include "UncookedOnlyUtils.h"
+#include "EdGraph/EdGraphNode.h"
+#include "Framework/Application/SlateApplication.h"
+#include "Graph/RigDecorator_AnimNextCppTrait.h"
+#include "Graph/RigUnit_AnimNextRunAnimationGraph.h"
+#include "Graph/RigUnit_AnimNextTraitStack.h"
+#include "Graph/RigVMTrait_AnimNextPublicVariables.h"
+#include "RigVMModel/RigVMController.h"
+#include "RigVMModel/Nodes/RigVMUnitNode.h"
+#include "TraitCore/TraitRegistry.h"
+
+#define LOCTEXT_NAMESPACE "FAnimationGraphMenuExtensions"
+
+namespace UE::AnimNext::Editor::Private
+{
+	static const FLazyName VariablesTraitBaseName = TEXT("Variables");
+}
+
+namespace UE::AnimNext::Editor
+{
+
+void FAnimationGraphMenuExtensions::RegisterMenus()
+{
+	FToolMenuOwnerScoped OwnerScoped(TEXT("FAnimNextAnimationGraphItemDetails"));
+	UToolMenu* Menu = UToolMenus::Get()->ExtendMenu("GraphEditor.GraphNodeContextMenu.AnimNextEdGraphNode");
+	if (Menu == nullptr)
+	{
+		return;
+	}
+
+	Menu->AddDynamicSection(TEXT("AnimNextEdGraphNode"), FNewToolMenuDelegate::CreateLambda([](UToolMenu* InMenu)
+	{
+		UGraphNodeContextMenuContext* Context = InMenu->FindContext<UGraphNodeContextMenuContext>();
+		if(Context == nullptr)
+		{
+			return;
+		}
+
+		const UAnimNextEdGraphNode* AnimNextEdGraphNode = Cast<UAnimNextEdGraphNode>(Context->Node);
+		if(AnimNextEdGraphNode == nullptr)
+		{
+			return;
+		}
+
+		URigVMNode* ModelNode = AnimNextEdGraphNode->GetModelNode();
+		if(ModelNode == nullptr)
+		{
+			return;
+		}
+
+		auto IsRunGraphNode = [](URigVMNode* InModelNode)
+		{
+			if (const URigVMUnitNode* VMNode = Cast<URigVMUnitNode>(InModelNode))
+			{
+				const UScriptStruct* ScriptStruct = VMNode->GetScriptStruct();
+				return ScriptStruct == FRigUnit_AnimNextRunAnimationGraph::StaticStruct();
+			}
+
+			return false;
+		};
+		
+		if (UncookedOnly::FAnimGraphUtils::IsTraitStackNode(ModelNode))
+		{
+			FToolMenuSection& Section = InMenu->AddSection("AnimNextTraitNodeActions", LOCTEXT("AnimNextTraitNodeActionsMenuHeader", "Traits"));
+
+			auto BuildAddTraitContextMenu = [AnimNextEdGraphNode, ModelNode](UToolMenu* InSubMenu)
+			{
+				const FTraitRegistry& TraitRegistry = FTraitRegistry::Get();
+				TArray<const FTrait*> Traits = TraitRegistry.GetTraits();
+
+				URigVMController* VMController = AnimNextEdGraphNode->GetController();
+
+				const UScriptStruct* CppDecoratorStruct = FRigDecorator_AnimNextCppDecorator::StaticStruct();
+
+				for (const FTrait* Trait : Traits)
+				{
+					UScriptStruct* ScriptStruct = Trait->GetTraitSharedDataStruct();
+			
+					FString DefaultValue;
+					{
+						const FRigDecorator_AnimNextCppDecorator DefaultCppDecoratorStructInstance;
+						FRigDecorator_AnimNextCppDecorator CppDecoratorStructInstance;
+						CppDecoratorStructInstance.DecoratorSharedDataStruct = ScriptStruct;
+			
+						if (!CppDecoratorStructInstance.CanBeAddedToNode(ModelNode, nullptr))
+						{
+							continue;	// This trait isn't supported on this node
+						}
+
+						FRigDecorator_AnimNextCppDecorator::StaticStruct()->ExportText(DefaultValue, &CppDecoratorStructInstance, &DefaultCppDecoratorStructInstance, nullptr, PPF_SerializedAsImportText, nullptr);
+					}
+			
+					FString DisplayNameMetadata;
+					ScriptStruct->GetStringMetaDataHierarchical(FRigVMStruct::DisplayNameMetaName, &DisplayNameMetadata);
+					const FString DisplayName = DisplayNameMetadata.IsEmpty() ? Trait->GetTraitName() : DisplayNameMetadata;
+					const FText ToolTip = ScriptStruct->GetToolTipText();
+			
+					FToolMenuEntry TraitEntry = FToolMenuEntry::InitMenuEntry(
+						*Trait->GetTraitName(),
+						FText::FromString(DisplayName),
+						ToolTip,
+						FSlateIcon(),
+						FUIAction(FExecuteAction::CreateLambda(
+							[Trait, VMController, ModelNode, CppDecoratorStruct, DefaultValue, DisplayName]()
+							{
+								VMController->AddTrait(
+									ModelNode->GetFName(),
+									*CppDecoratorStruct->GetPathName(),
+									*DisplayName,
+									DefaultValue, INDEX_NONE, true, true);
+							})
+						)
+					);
+
+					InSubMenu->AddMenuEntry(NAME_None, TraitEntry);
+				}
+			};
+			
+			Section.AddSubMenu(
+				TEXT("AddTraitMenu"),
+				LOCTEXT("AddTraitMenu", "Add Trait"),
+				LOCTEXT("AddTraitMenuTooltip", "Add the chosen trait to currently selected node"),
+				FNewToolMenuDelegate::CreateLambda(BuildAddTraitContextMenu));
+		}
+		else if(IsRunGraphNode(ModelNode))
+		{
+			FToolMenuSection& Section = InMenu->AddSection("AnimNextRunAnimGraphNodeActions", LOCTEXT("AnimNextAnimGraphNodeActionsMenuHeader", "Animation Graph"));
+	
+			URigVMController* VMController = AnimNextEdGraphNode->GetController();
+			URigVMPin* VMPin = Context->Pin != nullptr ? AnimNextEdGraphNode->FindModelPinFromGraphPin(Context->Pin) : nullptr;
+	
+			if(VMPin != nullptr && ModelNode->FindTrait(VMPin))
+			{
+				Section.AddMenuEntry(
+					TEXT("RemoveExposedVariables"),
+					LOCTEXT("RemoveExposedVariablesMenu", "Remove Exposed Variables"),
+					LOCTEXT("RemoveExposeVariablesMenuTooltip", "Remove the exposed variable trait from this node"),
+					FSlateIcon(),
+					FUIAction(
+						FExecuteAction::CreateLambda([VMController, ModelNode, Name = VMPin->GetFName()]()
+						{
+							VMController->RemoveTrait(ModelNode, Name);
+						})
+					));
+			}
+			else
+			{
+				auto BuildExposeVariablesContextMenu = [AnimNextEdGraphNode, ModelNode](UToolMenu* InSubMenu)
+				{
+					URigVMController* VMController = AnimNextEdGraphNode->GetController();
+
+					FContentBrowserModule& ContentBrowserModule = FModuleManager::Get().LoadModuleChecked<FContentBrowserModule>(TEXT("ContentBrowser"));
+
+					FAssetPickerConfig AssetPickerConfig;
+					AssetPickerConfig.Filter.ClassPaths.Add(UAnimNextDataInterface::StaticClass()->GetClassPathName());
+					AssetPickerConfig.Filter.bRecursiveClasses = true;
+					AssetPickerConfig.InitialAssetViewType = EAssetViewType::List;
+					AssetPickerConfig.AssetShowWarningText = LOCTEXT("NoAssetsWithPublicVariablesMessage", "No animation graphs with public variables found");
+					AssetPickerConfig.OnAssetSelected = FOnAssetSelected::CreateLambda([VMController, ModelNode](const FAssetData& InAssetData)
+					{
+						FSlateApplication::Get().DismissAllMenus();
+
+						FString DefaultValue;
+						FRigVMTrait_AnimNextPublicVariables DefaultTrait;
+						FRigVMTrait_AnimNextPublicVariables NewTrait;
+						UAnimNextDataInterface* Asset = CastChecked<UAnimNextDataInterface>(InAssetData.GetAsset());
+						NewTrait.InternalAsset = Asset;
+						const FInstancedPropertyBag& PublicVariableDefaults = Asset->GetPublicVariableDefaults();
+						TConstArrayView<FPropertyBagPropertyDesc> Descs = PublicVariableDefaults.GetPropertyBagStruct()->GetPropertyDescs();
+						NewTrait.InternalVariableNames.Reserve(Descs.Num());
+						for(const FPropertyBagPropertyDesc& Desc : Descs)
+						{
+							NewTrait.InternalVariableNames.Add(Desc.Name);
+						}
+						FRigVMTrait_AnimNextPublicVariables::StaticStruct()->ExportText(DefaultValue, &NewTrait, &DefaultTrait, nullptr, PPF_SerializedAsImportText, nullptr);
+
+						const FName ValidTraitName = URigVMSchema::GetUniqueName(Private::VariablesTraitBaseName, [ModelNode](const FName& InName)
+						{
+							return ModelNode->FindPin(InName.ToString()) == nullptr;
+						}, false, false);
+
+						VMController->AddTrait(ModelNode, FRigVMTrait_AnimNextPublicVariables::StaticStruct(), ValidTraitName, DefaultValue);
+					});
+
+					AssetPickerConfig.OnShouldFilterAsset = FOnShouldFilterAsset::CreateLambda([](const FAssetData& InAssetData)
+					{
+						// Filter to only show assets with public variables
+						FAnimNextAssetRegistryExports Exports;
+						if(UE::AnimNext::UncookedOnly::FUtils::GetExportedVariablesForAsset(InAssetData, Exports))
+						{
+							for(const FAnimNextAssetRegistryExportedVariable& Export : Exports.Variables)
+							{
+								if((Export.GetFlags() & EAnimNextExportedVariableFlags::Public) != EAnimNextExportedVariableFlags::NoFlags)
+								{
+									return false;
+								}
+							}
+						}
+						return true;
+					});
+
+					FToolMenuEntry Entry = FToolMenuEntry::InitWidget(
+						TEXT("AnimationGraphPicker"),
+						SNew(SBox)
+						.WidthOverride(300.0f)
+						.HeightOverride(400.0f)
+						[
+							ContentBrowserModule.Get().CreateAssetPicker(AssetPickerConfig)
+						],
+						FText::GetEmpty(),
+						true,
+						false,
+						false,
+						LOCTEXT("AnimationGraphPickerTooltip", "Choose an animation graph with public variables to expose")
+					);
+
+					InSubMenu->AddMenuEntry(NAME_None, Entry);
+				};
+				
+				Section.AddSubMenu(
+					TEXT("ExposeVariables"),
+					LOCTEXT("ExposeVariablesMenu", "Expose Variables"),
+					LOCTEXT("ExposeVariablesMenuTooltip", "Expose the variables of a selected animation graph as pins on this node"),
+					FNewToolMenuDelegate::CreateLambda(BuildExposeVariablesContextMenu));
+			}
+		}
+	}));
+}
+
+void FAnimationGraphMenuExtensions::UnregisterMenus()
+{
+	if(UToolMenus* ToolMenus = UToolMenus::Get())
+	{
+		ToolMenus->UnregisterOwnerByName("FAnimNextAnimationGraphItemDetails");
+	}
+}
+
+}
+
+#undef LOCTEXT_NAMESPACE
\ No newline at end of file
diff --git a/Engine/Plugins/Experimental/AnimNextAnimGraph/Source/AnimNextAnimGraphEditor/Private/Graph/AnimNextAnimationGraphMenuExtensions.h b/Engine/Plugins/Experimental/AnimNextAnimGraph/Source/AnimNextAnimGraphEditor/Private/Graph/AnimNextAnimationGraphMenuExtensions.h
new file mode 100644
--- /dev/null
+++ b/Engine/Plugins/Experimental/AnimNextAnimGraph/Source/AnimNextAnimGraphEditor/Private/Graph/AnimNextAnimationGraphMenuExtensions.h
@@ -0,0 +1,17 @@
+﻿// Copyright Epic Games, Inc. All Rights Reserved.
+
+#pragma once
+
+#include "CoreMinimal.h"
+
+namespace UE::AnimNext::Editor
+{
+
+struct FAnimationGraphMenuExtensions
+{
+	static void RegisterMenus();
+	static void UnregisterMenus();
+};
+
+}
+
diff --git a/Engine/Plugins/Experimental/AnimNextAnimGraph/Source/AnimNextAnimGraphTestSuite/AnimNextAnimGraphTestSuite.Build.cs b/Engine/Plugins/Experimental/AnimNextAnimGraph/Source/AnimNextAnimGraphTestSuite/AnimNextAnimGraphTestSuite.Build.cs
--- a/Engine/Plugins/Experimental/AnimNextAnimGraph/Source/AnimNextAnimGraphTestSuite/AnimNextAnimGraphTestSuite.Build.cs
+++ b/Engine/Plugins/Experimental/AnimNextAnimGraph/Source/AnimNextAnimGraphTestSuite/AnimNextAnimGraphTestSuite.Build.cs
@@ -8,7 +8,7 @@ namespace UnrealBuildTool.Rules
 		{
 			PCHUsage = ModuleRules.PCHUsageMode.UseExplicitOrSharedPCHs;
 
-			PublicDependencyModuleNames.AddRange(
+			PrivateDependencyModuleNames.AddRange(
 				new string[]
 				{
 					"Core",
@@ -17,7 +17,8 @@ namespace UnrealBuildTool.Rules
 					"AnimNext",
 					"AnimNextTestSuite",
 					"AnimNextAnimGraph",
-					"RigVM",
+					"RigVM", 
+					"PythonScriptPlugin",
 				}
 			);
 
@@ -26,9 +27,12 @@ namespace UnrealBuildTool.Rules
 				PrivateDependencyModuleNames.AddRange(
 					new string[]
 					{
+						"UnrealEd",
 						"AnimNextUncookedOnly",
 						"AnimNextEditor",
 						"RigVMDeveloper",
+						"AnimNextAnimGraphEditor",
+						"AnimNextAnimGraphUncookedOnly",
 					}
 				);
 			}
diff --git a/Engine/Plugins/Experimental/AnimNextAnimGraph/Source/AnimNextAnimGraphUncookedOnly/AnimNextAnimGraphUncookedOnly.Build.cs b/Engine/Plugins/Experimental/AnimNextAnimGraph/Source/AnimNextAnimGraphUncookedOnly/AnimNextAnimGraphUncookedOnly.Build.cs
--- a/Engine/Plugins/Experimental/AnimNextAnimGraph/Source/AnimNextAnimGraphUncookedOnly/AnimNextAnimGraphUncookedOnly.Build.cs
+++ b/Engine/Plugins/Experimental/AnimNextAnimGraph/Source/AnimNextAnimGraphUncookedOnly/AnimNextAnimGraphUncookedOnly.Build.cs
@@ -15,7 +15,7 @@ namespace UnrealBuildTool.Rules
 					"AnimNext",
 					"AnimNextUncookedOnly",
 					"AnimNextAnimGraph",
-					"ControlRigDeveloper",
+					"RigVMDeveloper",
 					"BlueprintGraph",	// For K2 schema
 					"AnimGraph",
 					"RigVM",
@@ -30,6 +30,7 @@ namespace UnrealBuildTool.Rules
 					{
 						"UnrealEd",
 						"BlueprintGraph",
+						"WorkspaceEditor",
 					}
 				);
 			}
diff --git a/Engine/Plugins/Experimental/AnimNextAnimGraph/Source/AnimNextAnimGraphUncookedOnly/Internal/AnimGraphUncookedOnlyUtils.h b/Engine/Plugins/Experimental/AnimNextAnimGraph/Source/AnimNextAnimGraphUncookedOnly/Internal/AnimGraphUncookedOnlyUtils.h
new file mode 100644
--- /dev/null
+++ b/Engine/Plugins/Experimental/AnimNextAnimGraph/Source/AnimNextAnimGraphUncookedOnly/Internal/AnimGraphUncookedOnlyUtils.h
@@ -0,0 +1,22 @@
+﻿// Copyright Epic Games, Inc. All Rights Reserved.
+
+#pragma once
+
+#include "CoreMinimal.h"
+
+class URigVMNode;
+class URigVMController;
+
+namespace UE::AnimNext::UncookedOnly
+{
+
+struct ANIMNEXTANIMGRAPHUNCOOKEDONLY_API FAnimGraphUtils
+{
+	/** Set up a simple animation graph */
+	static void SetupAnimGraph(const FName EntryName, URigVMController* InController);
+
+	/** Check whether the supplied node is a trait stack */
+	static bool IsTraitStackNode(URigVMNode* InModelNode);
+};
+
+}
\ No newline at end of file
diff --git a/Engine/Plugins/Experimental/AnimNextAnimGraph/Source/AnimNextAnimGraphUncookedOnly/Internal/AnimNextAnimGraphWorkspaceAssetUserData.h b/Engine/Plugins/Experimental/AnimNextAnimGraph/Source/AnimNextAnimGraphUncookedOnly/Internal/AnimNextAnimGraphWorkspaceAssetUserData.h
new file mode 100644
--- /dev/null
+++ b/Engine/Plugins/Experimental/AnimNextAnimGraph/Source/AnimNextAnimGraphUncookedOnly/Internal/AnimNextAnimGraphWorkspaceAssetUserData.h
@@ -0,0 +1,34 @@
+﻿// Copyright Epic Games, Inc. All Rights Reserved.
+
+#pragma once
+
+#include "AnimNextAssetWorkspaceAssetUserData.h"
+#include "Entries/AnimNextVariableEntry.h"
+#include "Graph/AnimNextAnimationGraph.h"
+
+#include "AnimNextAnimGraphWorkspaceAssetUserData.generated.h"
+
+USTRUCT()
+struct FAnimNextAnimationGraphOutlinerData : public FAnimNextRigVMAssetOutlinerData
+{
+	GENERATED_BODY()
+
+	FAnimNextAnimationGraphOutlinerData() = default;
+
+	UAnimNextAnimationGraph* GetAnimationGraph() const
+	{
+		return Cast<UAnimNextAnimationGraph>(Asset);
+	}
+};
+
+UCLASS()
+class UAnimNextAnimGraphWorkspaceAssetUserData : public UAssetUserData
+{
+public:
+	virtual bool IsEditorOnly() const override { return true; }
+
+private:
+	GENERATED_BODY()
+
+	virtual void GetAssetRegistryTags(FAssetRegistryTagsContext Context) const override;
+};
diff --git a/Engine/Plugins/Experimental/AnimNextAnimGraph/Source/AnimNextAnimGraphUncookedOnly/Internal/AnimNextController.h b/Engine/Plugins/Experimental/AnimNextAnimGraph/Source/AnimNextAnimGraphUncookedOnly/Internal/AnimNextController.h
new file mode 100644
--- /dev/null
+++ b/Engine/Plugins/Experimental/AnimNextAnimGraph/Source/AnimNextAnimGraphUncookedOnly/Internal/AnimNextController.h
@@ -0,0 +1,47 @@
+// Copyright Epic Games, Inc. All Rights Reserved.
+
+#pragma once
+
+#include "CoreMinimal.h"
+#include "RigVMModel/RigVMController.h"
+#include "AnimNextController.generated.h"
+
+namespace UE::AnimNext
+{
+	struct FTrait;
+	struct FTraitUID;
+}
+
+/**
+  * Implements AnimNext RigVM controller extensions
+  */
+UCLASS(MinimalAPI)
+class UAnimNextController : public URigVMController
+{
+	GENERATED_BODY()
+
+public:
+	
+	// Adds a new Trait to the Stack, with default struct values
+	// Returns Trait Instance Name (or NAME_None on failure)
+	UFUNCTION(BlueprintCallable, Category = RigVMController)
+	ANIMNEXTANIMGRAPHUNCOOKEDONLY_API FName AddTraitByName(FName InNodeName, FName InNewTraitTypeName, int32 InPinIndex, const FString& InNewTraitDefaultValue = TEXT(""), bool bSetupUndoRedo = true, bool bPrintPythonCommand = false);
+
+	// Removes a Trait from the Stack, using Trait Instance Name
+	// Returns operation success (true) or failure (false)
+	UFUNCTION(BlueprintCallable, Category = RigVMController)
+	ANIMNEXTANIMGRAPHUNCOOKEDONLY_API bool RemoveTraitByName(FName InNodeName, FName InTraitInstanceName, bool bSetupUndoRedo = true, bool bPrintPythonCommand = false);
+
+	// Swap a Trait from the Stack with a new one, using existing Trait Instance Name and new Trait Type Name
+	// Returns Trait Instance Name (or NAME_None on failure)
+	UFUNCTION(BlueprintCallable, Category = RigVMController)
+	ANIMNEXTANIMGRAPHUNCOOKEDONLY_API FName SwapTraitByName(FName InNodeName, FName InTraitInstanceName, int32 InCurrentTraitPinIndex, FName InNewTraitTypeName, const FString& InNewTraitDefaultValue = TEXT(""), bool bSetupUndoRedo = true, bool bPrintPythonCommand = false);
+
+	// Move a Trait from its current PinIndex to the specified one (moving it visually in the stack)
+	// Returns operation success (true) or failure (false)
+	UFUNCTION(BlueprintCallable, Category = RigVMController)
+	ANIMNEXTANIMGRAPHUNCOOKEDONLY_API bool SetTraitPinIndex(FName InNodeName, FName InTraitInstanceName, int32 InNewPinIndex, bool bSetupUndoRedo = true, bool bPrintPythonCommand = false);
+
+	// Adds a unit node with a dynamic number of pins
+	URigVMUnitNode* AddUnitNodeWithPins(UScriptStruct* InScriptStruct, const FRigVMPinInfoArray& PinArray, const FName& InMethodName = TEXT("Execute"), const FVector2D& InPosition = FVector2D::ZeroVector, const FString& InNodeName = TEXT(""), bool bSetupUndoRedo = true, bool bPrintPythonCommand = false);
+};
diff --git a/Engine/Plugins/Experimental/AnimNextAnimGraph/Source/AnimNextAnimGraphUncookedOnly/Internal/Graph/AnimNextAnimationGraphSchema.h b/Engine/Plugins/Experimental/AnimNextAnimGraph/Source/AnimNextAnimGraphUncookedOnly/Internal/Graph/AnimNextAnimationGraphSchema.h
new file mode 100644
--- /dev/null
+++ b/Engine/Plugins/Experimental/AnimNextAnimGraph/Source/AnimNextAnimGraphUncookedOnly/Internal/Graph/AnimNextAnimationGraphSchema.h
@@ -0,0 +1,15 @@
+﻿// Copyright Epic Games, Inc. All Rights Reserved.
+
+#pragma once
+
+#include "CoreMinimal.h"
+#include "AnimNextRigVMAssetSchema.h"
+#include "AnimNextAnimationGraphSchema.generated.h"
+
+UCLASS(MinimalAPI)
+class UAnimNextAnimationGraphSchema : public UAnimNextRigVMAssetSchema
+{
+	GENERATED_BODY()
+
+	// Schema reserved for future filtering of units etc.
+};
diff --git a/Engine/Plugins/Experimental/AnimNextAnimGraph/Source/AnimNextAnimGraphUncookedOnly/Internal/Graph/AnimNextAnimationGraph_EditorData.h b/Engine/Plugins/Experimental/AnimNextAnimGraph/Source/AnimNextAnimGraphUncookedOnly/Internal/Graph/AnimNextAnimationGraph_EditorData.h
new file mode 100644
--- /dev/null
+++ b/Engine/Plugins/Experimental/AnimNextAnimGraph/Source/AnimNextAnimGraphUncookedOnly/Internal/Graph/AnimNextAnimationGraph_EditorData.h
@@ -0,0 +1,74 @@
+﻿// Copyright Epic Games, Inc. All Rights Reserved.
+
+#pragma once
+
+#include "CoreMinimal.h"
+#include "AnimNextController.h"
+#include "DataInterface/AnimNextDataInterface_EditorData.h"
+#include "AnimNextExecuteContext.h"
+
+#include "AnimNextAnimationGraph_EditorData.generated.h"
+
+class FAnimationAnimNextRuntimeTest_GraphAddTrait;
+class FAnimationAnimNextRuntimeTest_GraphExecute;
+class FAnimationAnimNextRuntimeTest_GraphExecuteLatent;
+
+namespace UE::AnimNext::UncookedOnly
+{
+	struct FUtils;
+}
+
+namespace UE::AnimNext::Editor
+{
+	class FModuleEditor;
+	class SAnimNextGraphView;
+	struct FUtils;
+}
+
+// Script-callable editor API hoisted onto UAnimNextAnimationGraph
+UCLASS()
+class UAnimNextAnimationGraphLibrary : public UBlueprintFunctionLibrary
+{
+	GENERATED_BODY()
+
+public:
+	/** Adds an animation graph to an AnimNext asset */
+	UFUNCTION(BlueprintCallable, Category = "AnimNext|Entries", meta=(ScriptMethod))
+	static ANIMNEXTANIMGRAPHUNCOOKEDONLY_API UAnimNextAnimationGraphEntry* AddAnimationGraph(UAnimNextAnimationGraph* InAsset, FName InName, bool bSetupUndoRedo = true, bool bPrintPythonCommand = true);
+};
+
+/** Editor data for AnimNext animation graphs */
+UCLASS()
+class ANIMNEXTANIMGRAPHUNCOOKEDONLY_API UAnimNextAnimationGraph_EditorData : public UAnimNextDataInterface_EditorData
+{
+	GENERATED_BODY()
+
+	friend class UAnimNextAnimationGraphFactory;
+	friend class UAnimNextEdGraph;
+	friend struct UE::AnimNext::UncookedOnly::FUtils;
+	friend struct UE::AnimNext::Editor::FUtils;
+	friend class UE::AnimNext::Editor::FModuleEditor;
+	friend class UE::AnimNext::Editor::SAnimNextGraphView;
+	friend struct FAnimNextGraphSchemaAction_RigUnit;
+	friend struct FAnimNextGraphSchemaAction_DispatchFactory;
+	friend class FAnimationAnimNextEditorTest_GraphAddTrait;
+	friend class FAnimationAnimNextEditorTest_GraphTraitOperations;
+	friend class FAnimationAnimNextRuntimeTest_GraphExecute;
+	friend class FAnimationAnimNextRuntimeTest_GraphExecuteLatent;
+
+public:
+	/** Adds an animation graph to this asset */
+	UAnimNextAnimationGraphEntry* AddAnimationGraph(FName InName, bool bSetupUndoRedo = true, bool bPrintPythonCommand = true);
+
+protected:
+	// IRigVMClientHost interface
+	virtual void RecompileVM() override;
+
+	// UAnimNextRigVMAssetEditorData interface
+	virtual TSubclassOf<URigVMController> GetControllerClass() const override { return UAnimNextController::StaticClass(); }
+	virtual UScriptStruct* GetExecuteContextStruct() const override { return FAnimNextExecuteContext::StaticStruct(); }
+	virtual TConstArrayView<TSubclassOf<UAnimNextRigVMAssetEntry>> GetEntryClasses() const override;
+	virtual bool CanAddNewEntry(TSubclassOf<UAnimNextRigVMAssetEntry> InClass) const override;
+	virtual TSubclassOf<UAssetUserData> GetAssetUserDataClass() const override;
+	virtual void InitializeAssetUserData() override;
+};
diff --git a/Engine/Plugins/Experimental/AnimNextAnimGraph/Source/AnimNextAnimGraphUncookedOnly/Private/AnimGraphUncookedOnlyUtils.cpp b/Engine/Plugins/Experimental/AnimNextAnimGraph/Source/AnimNextAnimGraphUncookedOnly/Private/AnimGraphUncookedOnlyUtils.cpp
new file mode 100644
--- /dev/null
+++ b/Engine/Plugins/Experimental/AnimNextAnimGraph/Source/AnimNextAnimGraphUncookedOnly/Private/AnimGraphUncookedOnlyUtils.cpp
@@ -0,0 +1,51 @@
+﻿// Copyright Epic Games, Inc. All Rights Reserved.
+
+#include "AnimGraphUncookedOnlyUtils.h"
+
+#include "AnimNextEdGraphNode.h"
+#include "Graph/RigUnit_AnimNextGraphRoot.h"
+#include "Graph/RigUnit_AnimNextTraitStack.h"
+#include "RigVMModel/RigVMController.h"	
+
+namespace UE::AnimNext::UncookedOnly
+{
+
+bool FAnimGraphUtils::IsTraitStackNode(URigVMNode* InModelNode)
+{
+	if (const URigVMUnitNode* VMNode = Cast<URigVMUnitNode>(InModelNode))
+	{
+		const UScriptStruct* ScriptStruct = VMNode->GetScriptStruct();
+		return ScriptStruct == FRigUnit_AnimNextTraitStack::StaticStruct();
+	}
+
+	return false;
+}
+
+void FAnimGraphUtils::SetupAnimGraph(const FName EntryName, URigVMController* InController)
+{
+	// Clear the graph
+	InController->RemoveNodes(InController->GetGraph()->GetNodes());
+
+	// Add root node
+	URigVMUnitNode* MainEntryPointNode = InController->AddUnitNode(FRigUnit_AnimNextGraphRoot::StaticStruct(), FRigUnit_AnimNextGraphRoot::EventName, FVector2D(-400.0f, 0.0f), FString(), false);
+	if(MainEntryPointNode == nullptr)
+	{
+		return;
+	}
+
+	URigVMPin* BeginExecutePin = MainEntryPointNode->FindPin(GET_MEMBER_NAME_STRING_CHECKED(FRigUnit_AnimNextGraphRoot, Result));
+	if(BeginExecutePin == nullptr)
+	{
+		return;
+	}
+
+	URigVMPin* EntryPointPin = MainEntryPointNode->FindPin(GET_MEMBER_NAME_STRING_CHECKED(FRigUnit_AnimNextGraphRoot, EntryPoint));
+	if(EntryPointPin == nullptr)
+	{
+		return;
+	}
+
+	InController->SetPinDefaultValue(EntryPointPin->GetPinPath(), EntryName.ToString());
+}
+
+}
diff --git a/Engine/Plugins/Experimental/AnimNextAnimGraph/Source/AnimNextAnimGraphUncookedOnly/Private/AnimNextAnimGraphWorkspaceAssetUserData.cpp b/Engine/Plugins/Experimental/AnimNextAnimGraph/Source/AnimNextAnimGraphUncookedOnly/Private/AnimNextAnimGraphWorkspaceAssetUserData.cpp
new file mode 100644
--- /dev/null
+++ b/Engine/Plugins/Experimental/AnimNextAnimGraph/Source/AnimNextAnimGraphUncookedOnly/Private/AnimNextAnimGraphWorkspaceAssetUserData.cpp
@@ -0,0 +1,31 @@
+﻿// Copyright Epic Games, Inc. All Rights Reserved.
+
+#include "AnimNextAnimGraphWorkspaceAssetUserData.h"
+
+#include "UncookedOnlyUtils.h"
+#include "Graph/AnimNextAnimationGraph_EditorData.h"
+#include "UObject/AssetRegistryTagsContext.h"
+
+void UAnimNextAnimGraphWorkspaceAssetUserData::GetAssetRegistryTags(FAssetRegistryTagsContext Context) const
+{
+	Super::GetAssetRegistryTags(Context);
+
+	FWorkspaceOutlinerItemExports Exports;
+	{
+		UAnimNextAnimationGraph* Asset = CastChecked<UAnimNextAnimationGraph>(GetOuter());
+		const UAnimNextAnimationGraph_EditorData* EditorData = UE::AnimNext::UncookedOnly::FUtils::GetEditorData<UAnimNextAnimationGraph_EditorData>(Asset);
+		{
+			FWorkspaceOutlinerItemExport& RootAssetExport = Exports.Exports.Add_GetRef(FWorkspaceOutlinerItemExport(Asset->GetFName(), Asset));
+			RootAssetExport.GetData().InitializeAsScriptStruct(FAnimNextAnimationGraphOutlinerData::StaticStruct());
+			FAnimNextRigVMAssetOutlinerData& Data = RootAssetExport.GetData().GetMutable<FAnimNextRigVMAssetOutlinerData>();
+			Data.Asset = Asset;
+		}
+
+		UE::AnimNext::UncookedOnly::FUtils::GetAssetOutlinerItems(EditorData, Exports);
+	}
+
+	FString TagValue;
+	FWorkspaceOutlinerItemExports::StaticStruct()->ExportText(TagValue, &Exports, nullptr, nullptr, PPF_None, nullptr);
+	Context.AddTag(FAssetRegistryTag(UE::Workspace::ExportsWorkspaceItemsRegistryTag, TagValue, FAssetRegistryTag::TT_Hidden));
+}
+
diff --git a/Engine/Plugins/Experimental/AnimNextAnimGraph/Source/AnimNextAnimGraphUncookedOnly/Private/Graph/AnimNextAnimationGraphSchema.cpp b/Engine/Plugins/Experimental/AnimNextAnimGraph/Source/AnimNextAnimGraphUncookedOnly/Private/Graph/AnimNextAnimationGraphSchema.cpp
new file mode 100644
--- /dev/null
+++ b/Engine/Plugins/Experimental/AnimNextAnimGraph/Source/AnimNextAnimGraphUncookedOnly/Private/Graph/AnimNextAnimationGraphSchema.cpp
@@ -0,0 +1,4 @@
+﻿// Copyright Epic Games, Inc. All Rights Reserved.
+
+#include "Graph/AnimNextAnimationGraphSchema.h"
+
diff --git a/Engine/Plugins/Experimental/AnimNextChooser/Source/AnimNextChooser/Private/Module.cpp b/Engine/Plugins/Experimental/AnimNextChooser/Source/AnimNextChooser/Private/Module.cpp
--- a/Engine/Plugins/Experimental/AnimNextChooser/Source/AnimNextChooser/Private/Module.cpp
+++ b/Engine/Plugins/Experimental/AnimNextChooser/Source/AnimNextChooser/Private/Module.cpp
@@ -8,9 +8,6 @@
 #include "Animation/AnimSequence.h"
 #include "Chooser.h"
 
-#include "IAnimNextModuleInterface.h"
-#include "EvaluationVM/EvaluationVM.h"
-
 namespace UE::AnimNext::Chooser
 {
 
diff --git a/Engine/Plugins/Experimental/AnimNextChooser/Source/AnimNextChooser/Private/RigUnit_EvaluateChooser.cpp b/Engine/Plugins/Experimental/AnimNextChooser/Source/AnimNextChooser/Private/RigUnit_EvaluateChooser.cpp
--- a/Engine/Plugins/Experimental/AnimNextChooser/Source/AnimNextChooser/Private/RigUnit_EvaluateChooser.cpp
+++ b/Engine/Plugins/Experimental/AnimNextChooser/Source/AnimNextChooser/Private/RigUnit_EvaluateChooser.cpp
@@ -41,6 +41,6 @@ FRigUnit_EvaluateChooser_AnimNext_Execute()
 {
 	Result = nullptr;
 
-	const FAnimNextDataInterfaceInstance& Instance = ExecuteContext.GetInstance();
+	const FAnimNextDataInterfaceInstance& Instance = ExecuteContext.GetContextData<FAnimNextModuleContextData>().GetDataInterfaceInstance();
 	RunChooserHelper({ ContextObject }, FStructView::Make(const_cast<FAnimNextDataInterfaceInstance&>(Instance)), Chooser, Result);
 }
\ No newline at end of file
diff --git a/Engine/Plugins/Experimental/AnimNextPoseSearch/Source/AnimNextPoseSearch/Private/Module.cpp b/Engine/Plugins/Experimental/AnimNextPoseSearch/Source/AnimNextPoseSearch/Private/Module.cpp
--- a/Engine/Plugins/Experimental/AnimNextPoseSearch/Source/AnimNextPoseSearch/Private/Module.cpp
+++ b/Engine/Plugins/Experimental/AnimNextPoseSearch/Source/AnimNextPoseSearch/Private/Module.cpp
@@ -7,8 +7,6 @@
 #include "RigVMCore/RigVMRegistry.h"
 
 #include "AlphaBlend.h"
-#include "IAnimNextModuleInterface.h"
-#include "EvaluationVM/EvaluationVM.h"
 #include "PoseSearch/PoseSearchDatabase.h"
 #include "PoseSearch/PoseSearchSchema.h"
 #include "PoseSearch/PoseSearchTrajectoryTypes.h"
diff --git a/Engine/Plugins/Experimental/AnimNextPoseSearch/Source/AnimNextPoseSearch/Private/MotionMatchingTrait.cpp b/Engine/Plugins/Experimental/AnimNextPoseSearch/Source/AnimNextPoseSearch/Private/MotionMatchingTrait.cpp
--- a/Engine/Plugins/Experimental/AnimNextPoseSearch/Source/AnimNextPoseSearch/Private/MotionMatchingTrait.cpp
+++ b/Engine/Plugins/Experimental/AnimNextPoseSearch/Source/AnimNextPoseSearch/Private/MotionMatchingTrait.cpp
@@ -6,6 +6,7 @@
 #include "Animation/AnimSequence.h"
 #include "DataInterface/DataInterfaceStructAdapter.h"
 #include "IPoseHistory.h"
+#include "Graph/AnimNextGraphInstance.h"
 #include "GraphInterfaces/AnimNextNativeDataInterface_AnimSequencePlayer.h"
 #include "PoseSearch/PoseSearchDatabase.h"
 #include "PoseSearch/PoseSearchDerivedData.h"
diff --git a/Engine/Plugins/Experimental/AnimNextStateTree/Source/AnimNextStateTreeUncookedOnly/AnimNextStateTreeUncookedOnly.Build.cs b/Engine/Plugins/Experimental/AnimNextStateTree/Source/AnimNextStateTreeUncookedOnly/AnimNextStateTreeUncookedOnly.Build.cs
--- a/Engine/Plugins/Experimental/AnimNextStateTree/Source/AnimNextStateTreeUncookedOnly/AnimNextStateTreeUncookedOnly.Build.cs
+++ b/Engine/Plugins/Experimental/AnimNextStateTree/Source/AnimNextStateTreeUncookedOnly/AnimNextStateTreeUncookedOnly.Build.cs
@@ -25,7 +25,8 @@ namespace UnrealBuildTool.Rules
 					"AnimNext",
 					"AnimNextUncookedOnly",
 					"AnimNextStateTree",
-					"AnimNextAnimGraph"
+					"AnimNextAnimGraph",
+					"AnimNextAnimGraphUncookedOnly"
 				}
 			);
 		}
diff --git a/Engine/Plugins/Experimental/AnimNextStateTree/Source/AnimNextStateTreeUncookedOnly/Internal/AnimNextStateTree_EditorData.cpp b/Engine/Plugins/Experimental/AnimNextStateTree/Source/AnimNextStateTreeUncookedOnly/Internal/AnimNextStateTree_EditorData.cpp
--- a/Engine/Plugins/Experimental/AnimNextStateTree/Source/AnimNextStateTreeUncookedOnly/Internal/AnimNextStateTree_EditorData.cpp
+++ b/Engine/Plugins/Experimental/AnimNextStateTree/Source/AnimNextStateTreeUncookedOnly/Internal/AnimNextStateTree_EditorData.cpp
@@ -2,6 +2,7 @@
 
 #include "AnimNextStateTree_EditorData.h"
 
+#include "AnimGraphUncookedOnlyUtils.h"
 #include "AnimNextStateTree.h"
 #include "AnimNextStateTreeWorkspaceAssetUserData.h"
 #include "StateTree.h"
@@ -56,7 +57,7 @@ void UAnimNextStateTree_EditorData::GetProgrammaticGraphs(const FRigVMCompileSet
 		Graph->SetSchemaClass(UAnimNextAnimationGraphSchema::StaticClass());
 
 		UAnimNextController* Controller = CastChecked<UAnimNextController>(RigVMClient.GetOrCreateController(Graph));
-		UE::AnimNext::UncookedOnly::FUtils::SetupAnimGraph(FRigUnit_AnimNextGraphRoot::DefaultEntryPoint, Controller);
+		UE::AnimNext::UncookedOnly::FAnimGraphUtils::SetupAnimGraph(FRigUnit_AnimNextGraphRoot::DefaultEntryPoint, Controller);
 
 		if (Controller->GetGraph()->GetNodes().Num() != 1)
 		{
