AnimNext multiple additive trait support

#okforguthub public

diff --git a/Engine/Plugins/Experimental/AnimNextAnimGraph/Source/AnimNextAnimGraph/Public/TraitCore/Trait.h b/Engine/Plugins/Experimental/AnimNextAnimGraph/Source/AnimNextAnimGraph/Public/TraitCore/Trait.h
--- a/Engine/Plugins/Experimental/AnimNextAnimGraph/Source/AnimNextAnimGraph/Public/TraitCore/Trait.h
+++ b/Engine/Plugins/Experimental/AnimNextAnimGraph/Source/AnimNextAnimGraph/Public/TraitCore/Trait.h
@@ -391,6 +391,8 @@ namespace UE::AnimNext
 
 		// Makes the Trait Editor only display this Trait in Advanced view
 		virtual bool IsHidden() const { return false; }
+
+		virtual bool MultipleInstanceSupport() const { return false; }
 #endif
 
 	protected:
@@ -425,6 +427,10 @@ namespace UE::AnimNext
 		DECLARE_ABSTRACT_ANIM_TRAIT(FAdditiveTrait, 0x0dbd1b3f, FTrait)
 
 		virtual ETraitMode GetTraitMode() const override { return ETraitMode::Additive; }
+
+#if WITH_EDITOR
+		virtual bool MultipleInstanceSupport() const override { return true; }
+#endif
 	};
 
 	/**
diff --git a/Engine/Plugins/Experimental/AnimNextAnimGraph/Source/AnimNextAnimGraphEditor/Private/Graph/STraitEditorView.cpp b/Engine/Plugins/Experimental/AnimNextAnimGraph/Source/AnimNextAnimGraphEditor/Private/Graph/STraitEditorView.cpp
--- a/Engine/Plugins/Experimental/AnimNextAnimGraph/Source/AnimNextAnimGraphEditor/Private/Graph/STraitEditorView.cpp
+++ b/Engine/Plugins/Experimental/AnimNextAnimGraph/Source/AnimNextAnimGraphEditor/Private/Graph/STraitEditorView.cpp
@@ -625,7 +625,8 @@ int32 STraitEditorView::GetTraitPinIndex(UAnimNextEdGraphNode* InEdGraphNode, co
 																																		, Trait->GetTraitUID()
 																																		, Trait->GetTraitMode()
 																																		, TArray<FTraitInterfaceUID>(TraitImplementedInterfaces)
-																																		, TArray<FTraitInterfaceUID>(TraitRequiredInterfaces)));
+																																		, TArray<FTraitInterfaceUID>(TraitRequiredInterfaces)
+																																		, Trait->MultipleInstanceSupport()));
 
 									for (const FTraitInterfaceUID& TraitInterface : TraitImplementedInterfaces)
 									{
@@ -687,7 +688,7 @@ int32 STraitEditorView::GetTraitPinIndex(UAnimNextEdGraphNode* InEdGraphNode, co
 
 			if (NumBaseTraits == 0)
 			{
-				TraitsData.Insert(MakeShared<FTraitDataEditorDef>(NAME_None, LOCTEXT("BaseTraitUnset", "<Base Trait Unset>"), FTraitUID(), ETraitMode::Base, TArray<FTraitInterfaceUID>(), TArray<FTraitInterfaceUID>()), 0);
+				TraitsData.Insert(MakeShared<FTraitDataEditorDef>(NAME_None, LOCTEXT("BaseTraitUnset", "<Base Trait Unset>"), FTraitUID(), ETraitMode::Base, TArray<FTraitInterfaceUID>(), TArray<FTraitInterfaceUID>(), false), 0);
 			}
 		}
 	}
diff --git a/Engine/Plugins/Experimental/AnimNextAnimGraph/Source/AnimNextAnimGraphEditor/Private/Graph/STraitListView.cpp b/Engine/Plugins/Experimental/AnimNextAnimGraph/Source/AnimNextAnimGraphEditor/Private/Graph/STraitListView.cpp
--- a/Engine/Plugins/Experimental/AnimNextAnimGraph/Source/AnimNextAnimGraphEditor/Private/Graph/STraitListView.cpp
+++ b/Engine/Plugins/Experimental/AnimNextAnimGraph/Source/AnimNextAnimGraphEditor/Private/Graph/STraitListView.cpp
@@ -555,7 +555,7 @@ static bool IsTraitAvailable(const TWeakPtr<TArray<TSharedPtr<FTraitDataEditorDe
 		// No duplicated Traits allowed in the list
 		for (const TSharedPtr<FTraitDataEditorDef>& Trait : *CurrentTraitsDataShared.Get())
 		{
-			if (Trait->TraitUID == TraitUID)
+			if (Trait->TraitUID == TraitUID && !Trait->bMultipleInstanceSupport)
 			{
 				return false;
 			}
diff --git a/Engine/Plugins/Experimental/AnimNextAnimGraph/Source/AnimNextAnimGraphEditor/Private/Graph/TraitEditorDefs.h b/Engine/Plugins/Experimental/AnimNextAnimGraph/Source/AnimNextAnimGraphEditor/Private/Graph/TraitEditorDefs.h
--- a/Engine/Plugins/Experimental/AnimNextAnimGraph/Source/AnimNextAnimGraphEditor/Private/Graph/TraitEditorDefs.h
+++ b/Engine/Plugins/Experimental/AnimNextAnimGraph/Source/AnimNextAnimGraphEditor/Private/Graph/TraitEditorDefs.h
@@ -84,13 +84,14 @@ struct FTraitDataEditorDef
 	{
 	}
 
-	FTraitDataEditorDef(const FName& InTraitName, const FText& InTraitDisplayName, FTraitUID InTraitUUID, ETraitMode InTraitMode, const TArray<FTraitInterfaceUID>& InImplementedInterfaces, const TArray<FTraitInterfaceUID>& InRequiredInterfaces)
+	FTraitDataEditorDef(const FName& InTraitName, const FText& InTraitDisplayName, FTraitUID InTraitUUID, ETraitMode InTraitMode, const TArray<FTraitInterfaceUID>& InImplementedInterfaces, const TArray<FTraitInterfaceUID>& InRequiredInterfaces, bool bInMultipleInstanceSupport)
 		: TraitName(InTraitName)
 		, TraitDisplayName(InTraitDisplayName)
 		, TraitUID(InTraitUUID)
 		, TraitMode(InTraitMode)
 		, ImplementedInterfaces(InImplementedInterfaces)
 		, RequiredInterfaces(InRequiredInterfaces)
+		, bMultipleInstanceSupport(bInMultipleInstanceSupport)
 	{}
 
 	FName TraitName;
@@ -103,6 +104,8 @@ struct FTraitDataEditorDef
 	TArray<FTraitInterfaceUID> RequiredInterfaces;
 	TArray<int32> RequiredInterfacesStackListIndexes;
 
+	bool bMultipleInstanceSupport = false;
+
 	FTraitStackTraitStatus StackStatus;
 };
 
