AnimNext: Fix GC crash in sub-graph trait

diff --git a/Engine/Plugins/Experimental/AnimNextAnimGraph/Source/AnimNextAnimGraph/Private/PlayAnim/PlayAnimSlotTrait.cpp b/Engine/Plugins/Experimental/AnimNextAnimGraph/Source/AnimNextAnimGraph/Private/PlayAnim/PlayAnimSlotTrait.cpp
--- a/Engine/Plugins/Experimental/AnimNextAnimGraph/Source/AnimNextAnimGraph/Private/PlayAnim/PlayAnimSlotTrait.cpp
+++ b/Engine/Plugins/Experimental/AnimNextAnimGraph/Source/AnimNextAnimGraph/Private/PlayAnim/PlayAnimSlotTrait.cpp
@@ -516,8 +516,10 @@ namespace UE::AnimNext
 			}
 
 			Collector.AddReferencedObject(SlotRequest.AnimationGraph);
-			FAnimNextGraphInstance* ImplPtr = SlotRequest.GraphInstance.Get();
-			Collector.AddPropertyReferencesWithStructARO(FAnimNextGraphInstance::StaticStruct(), ImplPtr);
+			if(FAnimNextGraphInstance* ImplPtr = SlotRequest.GraphInstance.Get())
+			{
+				Collector.AddPropertyReferencesWithStructARO(FAnimNextGraphInstance::StaticStruct(), ImplPtr);
+			}
 		}
 	}
 }
diff --git a/Engine/Plugins/Experimental/AnimNextAnimGraph/Source/AnimNextAnimGraph/Private/Traits/SubGraphHost.cpp b/Engine/Plugins/Experimental/AnimNextAnimGraph/Source/AnimNextAnimGraph/Private/Traits/SubGraphHost.cpp
--- a/Engine/Plugins/Experimental/AnimNextAnimGraph/Source/AnimNextAnimGraph/Private/Traits/SubGraphHost.cpp
+++ b/Engine/Plugins/Experimental/AnimNextAnimGraph/Source/AnimNextAnimGraph/Private/Traits/SubGraphHost.cpp
@@ -248,8 +248,10 @@ namespace UE::AnimNext
 		for (FSubGraphSlot& SubGraphEntry : InstanceData->SubGraphSlots)
 		{
 			Collector.AddReferencedObject(SubGraphEntry.AnimationGraph);
-			FAnimNextGraphInstance* ImplPtr = SubGraphEntry.GraphInstance.Get();
-			Collector.AddPropertyReferencesWithStructARO(FAnimNextGraphInstance::StaticStruct(), ImplPtr);
+			if(FAnimNextGraphInstance* ImplPtr = SubGraphEntry.GraphInstance.Get())
+			{
+				Collector.AddPropertyReferencesWithStructARO(FAnimNextGraphInstance::StaticStruct(), ImplPtr);
+			}
 		}
 	}
 }
