AnimNext: Fix for crash in blendstack reachability analysis for GC
#rnx

diff --git a/Engine/Plugins/Experimental/AnimNextAnimGraph/Source/AnimNextAnimGraph/Private/Traits/BlendStackTrait.cpp b/Engine/Plugins/Experimental/AnimNextAnimGraph/Source/AnimNextAnimGraph/Private/Traits/BlendStackTrait.cpp
--- a/Engine/Plugins/Experimental/AnimNextAnimGraph/Source/AnimNextAnimGraph/Private/Traits/BlendStackTrait.cpp
+++ b/Engine/Plugins/Experimental/AnimNextAnimGraph/Source/AnimNextAnimGraph/Private/Traits/BlendStackTrait.cpp
@@ -270,8 +270,12 @@ namespace UE::AnimNext
 		for (FBlendStackCoreTrait::FGraphState& Graph : InstanceData->ChildGraphs)
 		{
 			Collector.AddReferencedObject(Graph.Request.AnimationGraph);
-			FAnimNextGraphInstance* ImplPtr = Graph.Instance.Get();
-			Collector.AddPropertyReferencesWithStructARO(FAnimNextGraphInstance::StaticStruct(), ImplPtr);
+
+			// Ignore inactive graphs. Could check the graph's state but use the shared instance for extra safety.
+			if (FAnimNextGraphInstance* ImplPtr = Graph.Instance.Get())
+			{
+				Collector.AddPropertyReferencesWithStructARO(FAnimNextGraphInstance::StaticStruct(), ImplPtr);
+			}
 		}
 	}
 
