AnimNext: Call function trait

Adds a new trait that allows a function to be called at specific points in the animation graph update process.
Support for this is enabled by creating a compiler-generated wrapper event for each function call, then calling that at runtime.

Function library is now exposed directly in the outliner, rather than each reference.
Split graph & latent execution context data into seperate types, as we dont use latent pin info when calling functions from traits
Added support for soft object paths in traits (required for editor-only funciton reference data).
Added support for hidden pins on traits.

diff --git a/Engine/Plugins/Experimental/AnimNext/Source/AnimNext/Internal/AnimNextRigVMAsset.h b/Engine/Plugins/Experimental/AnimNext/Source/AnimNext/Internal/AnimNextRigVMAsset.h
--- a/Engine/Plugins/Experimental/AnimNext/Source/AnimNext/Internal/AnimNextRigVMAsset.h
+++ b/Engine/Plugins/Experimental/AnimNext/Source/AnimNext/Internal/AnimNextRigVMAsset.h
@@ -47,6 +47,9 @@ public:
 	// Get public variable defaults property bag
 	const FInstancedPropertyBag& GetPublicVariableDefaults() const { return PublicVariableDefaults; }
 
+	// Get the RigVM object we host
+	const URigVM* GetRigVM() const { return RigVM; }
+
 private:
 	// URigVMHost interface
 	virtual TArray<FRigVMExternalVariable> GetExternalVariablesImpl(bool bFallbackToBlueprint) const;
diff --git a/Engine/Plugins/Experimental/AnimNext/Source/AnimNext/Private/AnimNextModuleImpl.cpp b/Engine/Plugins/Experimental/AnimNext/Source/AnimNext/Private/AnimNextModuleImpl.cpp
--- a/Engine/Plugins/Experimental/AnimNext/Source/AnimNext/Private/AnimNextModuleImpl.cpp
+++ b/Engine/Plugins/Experimental/AnimNext/Source/AnimNext/Private/AnimNextModuleImpl.cpp
@@ -21,6 +21,7 @@
 #include "RigVMCore/RigVMRegistry.h"
 #include "Variables/AnimNextFieldPath.h"
 #include "Variables/AnimNextSoftFunctionPtr.h"
+#include "RigVMCore/RigVMGraphFunctionDefinition.h"
 
 #define LOCTEXT_NAMESPACE "AnimNextModule"
 
@@ -50,6 +51,7 @@ namespace UE::AnimNext
 			FUniversalObjectLocator::StaticStruct(),
 			FAnimNextFieldPath::StaticStruct(),
 			FAnimNextSoftFunctionPtr::StaticStruct(),
+			FRigVMGraphFunctionHeader::StaticStruct(),
 		};
 
 		RigVMRegistry.RegisterStructTypes(AllowedStructTypes);
diff --git a/Engine/Plugins/Experimental/AnimNext/Source/AnimNextEditor/Internal/Common/AnimNextAssetItemDetails.cpp b/Engine/Plugins/Experimental/AnimNext/Source/AnimNextEditor/Internal/Common/AnimNextAssetItemDetails.cpp
--- a/Engine/Plugins/Experimental/AnimNext/Source/AnimNextEditor/Internal/Common/AnimNextAssetItemDetails.cpp
+++ b/Engine/Plugins/Experimental/AnimNext/Source/AnimNextEditor/Internal/Common/AnimNextAssetItemDetails.cpp
@@ -13,6 +13,8 @@
 #include "ScopedTransaction.h"
 #include "RigVMModel/RigVMClient.h"
 #include "ToolMenus.h"
+#include "AssetRegistry/AssetRegistryModule.h"
+#include "AssetRegistry/IAssetRegistry.h"
 #include "Entries/AnimNextDataInterfaceEntry.h"
 #include "Framework/Commands/GenericCommands.h"
 #include "Module/AnimNextModule_EditorData.h"
@@ -27,7 +29,11 @@ namespace UE::AnimNext::Editor
 
 const FSlateBrush* FAnimNextAssetItemDetails::GetItemIcon(const FWorkspaceOutlinerItemExport& Export) const
 {
-	return FAppStyle::GetBrush(TEXT("LevelEditor.Tabs.Outliner"));
+	const FSoftObjectPath& AssetPath = Export.GetTopLevelAssetPath();
+	IAssetRegistry& AssetRegistry = FModuleManager::LoadModuleChecked<FAssetRegistryModule>(TEXT("AssetRegistry")).Get();
+	const FAssetData AssetData = AssetRegistry.GetAssetByObjectPath(AssetPath);
+	const UClass* AssetClass = FClassIconFinder::GetIconClassForAssetData(AssetData);
+	return FSlateIconFinder::FindIconForClass(AssetClass).GetIcon();
 }
 
 void FAnimNextAssetItemDetails::RegisterToolMenuExtensions()
@@ -196,6 +202,26 @@ void FAnimNextAssetItemDetails::RegisterToolMenuExtensions()
 					));
 			}
 		}
+
+		AssetSection.AddMenuEntry(
+			"AddFunction",
+			LOCTEXT("AddFunctionLabel", "Add Function"),
+			LOCTEXT("AddFunctionTooltip", "Add a function to this asset"),
+			FSlateIcon(FAppStyle::GetAppStyleSetName(), "GraphEditor.Function_16x"),
+			FUIAction(
+				FExecuteAction::CreateWeakLambda(WorkspaceItemContext, [WeakWorkspaceEditor = TWeakPtr<UE::Workspace::IWorkspaceEditor>(WorkspaceEditor), EditorData]()
+				{
+					FScopedTransaction Transaction(LOCTEXT("AddFunctionTransaction", "Add Function"));
+					URigVMLibraryNode* NewFunction = EditorData->AddFunction(TEXT("NewFunction"), true);
+
+					// Open the new function's graph
+					if(TSharedPtr<UE::Workspace::IWorkspaceEditor> PinnedWorkspaceEditor = WeakWorkspaceEditor.Pin())
+					{
+						UObject* EditorObject = EditorData->GetEditorObjectForRigVMGraph(NewFunction->GetContainedGraph());
+						PinnedWorkspaceEditor->OpenObjects({EditorObject});
+					}
+				})
+			));
 	}));
 }
 
diff --git a/Engine/Plugins/Experimental/AnimNext/Source/AnimNextEditor/Internal/Common/SRigVMFunctionPicker.h b/Engine/Plugins/Experimental/AnimNext/Source/AnimNextEditor/Internal/Common/SRigVMFunctionPicker.h
new file mode 100644
--- /dev/null
+++ b/Engine/Plugins/Experimental/AnimNext/Source/AnimNextEditor/Internal/Common/SRigVMFunctionPicker.h
@@ -0,0 +1,139 @@
+﻿// Copyright Epic Games, Inc. All Rights Reserved.
+
+#pragma once
+
+#include "CoreMinimal.h"
+#include "RigVMCore/RigVMGraphFunctionDefinition.h"
+#include "Widgets/SCompoundWidget.h"
+#include "Widgets/Views/STreeView.h"
+
+class UAnimNextRigVMAsset;
+class FTextFilterExpressionEvaluator;
+class SSearchBox;
+
+namespace UE::AnimNext::Editor
+{
+	class SRigVMFunctionRowWidget;
+}
+
+namespace UE::AnimNext::Editor
+{
+
+// Delegate called when the user picks a function from the dropdown menu
+// The header passed when selecting 'None' will be invalid
+using FOnRigVMFunctionPicked = TDelegate<void(const FRigVMGraphFunctionHeader&)>;
+
+// Picker widget the allows choosing a RigVM function
+class SRigVMFunctionPicker : public SCompoundWidget
+{
+	SLATE_BEGIN_ARGS(SRigVMFunctionPicker)
+		: _AllowNew(true)
+		, _AllowClear(true)
+	{}
+
+	// The 'current' asset, used for accessing private functions
+	SLATE_ARGUMENT(FAssetData, CurrentAsset)
+
+	// Delegate called when the user picks a function from the dropdown menu.
+	// The header passed when selecting 'None' will be invalid 
+	SLATE_EVENT(FOnRigVMFunctionPicked, OnRigVMFunctionPicked)
+
+	// Delegate called when the user chooses 'New Function...'
+	SLATE_EVENT(FSimpleDelegate, OnNewFunction)
+
+	// Attribute used to display the currently-picked function in a combo button 
+	SLATE_ATTRIBUTE(FText, FunctionName)
+
+	// Attribute used to display tooltip information about the currently-picked function
+	SLATE_ATTRIBUTE(FText, FunctionToolTip)
+
+	// Whether to show the 'New Function...' item
+	SLATE_ARGUMENT(bool, AllowNew)
+
+	// Whether to show the 'None' item
+	SLATE_ARGUMENT(bool, AllowClear)
+
+	SLATE_END_ARGS()
+
+	ANIMNEXTEDITOR_API void Construct(const FArguments& InArgs);
+
+private:
+	friend UE::AnimNext::Editor::SRigVMFunctionRowWidget;
+
+	// Requests a refresh of the dropdown menu contents
+	void RequestRefreshEntries();
+
+	// Refreshes the dropdown menu contents
+	void RefreshEntries();
+
+	enum class EEntryType
+	{
+		None,
+		Asset,
+		Function,
+		NewFunction,
+	};
+
+	struct FEntry
+	{
+		explicit FEntry(EEntryType InType)
+			: Type(InType)
+		{}
+
+		EEntryType Type;
+		FText Name;
+		FText ToolTip;
+		const FSlateBrush* Icon = nullptr;
+	};
+
+	struct FFunctionEntry : FEntry
+	{
+		FFunctionEntry()
+			: FEntry(EEntryType::Function)
+		{}
+
+		FRigVMGraphFunctionHeader FunctionHeader;
+	};
+
+	struct FAssetEntry : FEntry
+	{
+		FAssetEntry()
+			: FEntry(EEntryType::Asset)
+		{}
+
+		FAssetData Asset;
+		TArray<TSharedPtr<FEntry>> Functions;
+		TArray<TSharedPtr<FEntry>> FilteredFunctions;
+	};
+
+	struct FNewFunctionEntry : FEntry
+	{
+		FNewFunctionEntry()
+			: FEntry(EEntryType::NewFunction)
+		{}
+	};
+
+	struct FNoneEntry : FEntry
+	{
+		FNoneEntry()
+			: FEntry(EEntryType::None)
+		{}
+	};
+
+	FAssetData CurrentAsset;
+	TWeakObjectPtr<UAnimNextRigVMAsset> WeakCurrentAsset;
+	FOnRigVMFunctionPicked OnRigVMFunctionPicked;
+	FSimpleDelegate OnNewFunction;
+	TAttribute<FText> FunctionName;
+	TAttribute<FText> FunctionToolTip;
+	TSharedPtr<SSearchBox> SearchBox;
+	FText FilterText;
+	TSharedPtr<STreeView<TSharedPtr<FEntry>>> TreeView;
+	TArray<TSharedPtr<FEntry>> Entries;
+	TArray<TSharedPtr<FEntry>> FilteredEntries;
+	TSharedPtr<FTextFilterExpressionEvaluator> TextFilter;
+	bool bAllowClear = false;
+	bool bAllowNew = false;
+};
+
+}
diff --git a/Engine/Plugins/Experimental/AnimNext/Source/AnimNextEditor/Private/AnimNextEditorModule.cpp b/Engine/Plugins/Experimental/AnimNext/Source/AnimNextEditor/Private/AnimNextEditorModule.cpp
--- a/Engine/Plugins/Experimental/AnimNext/Source/AnimNextEditor/Private/AnimNextEditorModule.cpp
+++ b/Engine/Plugins/Experimental/AnimNext/Source/AnimNextEditor/Private/AnimNextEditorModule.cpp
@@ -45,6 +45,7 @@
 #include "DataInterface/AnimNextDataInterface.h"
 #include "Entries/AnimNextEventGraphEntry.h"
 #include "Logging/MessageLog.h"
+#include "Graph/AnimNextEdGraphCustomization.h"
 #include "Module/RigUnit_AnimNextModuleEvents.h"
 #include "Variables/SVariablesView.h"
 #include "Variables/VariableOverrideCommands.h"
@@ -118,6 +119,11 @@ void FAnimNextEditorModule::StartupModule()
 					return MakeShared<FAnimNextEdGraphNodeCustomization>(InWorkspaceEditor);
 				}));
 
+			InDetailsView->RegisterInstancedCustomPropertyLayout(UAnimNextEdGraph::StaticClass(), FOnGetDetailCustomizationInstance::CreateLambda([InWorkspaceEditor]()
+				{
+					return MakeShared<FAnimNextEdGraphCustomization>();
+				}));
+		
 			TArray<UScriptStruct*> StructsToCustomize = {
 				TBaseStructure<FVector>::Get(),
 				TBaseStructure<FVector2D>::Get(),
diff --git a/Engine/Plugins/Experimental/AnimNext/Source/AnimNextEditor/Private/Common/AnimNextEdGraphNodeCustomization.cpp b/Engine/Plugins/Experimental/AnimNext/Source/AnimNextEditor/Private/Common/AnimNextEdGraphNodeCustomization.cpp
--- a/Engine/Plugins/Experimental/AnimNext/Source/AnimNextEditor/Private/Common/AnimNextEdGraphNodeCustomization.cpp
+++ b/Engine/Plugins/Experimental/AnimNext/Source/AnimNextEditor/Private/Common/AnimNextEdGraphNodeCustomization.cpp
@@ -377,7 +377,7 @@ void FAnimNextEdGraphNodeCustomization::PopulateCategory(IDetailLayoutBuilder& D
 				}
 
 				const FProperty* Property = PinnedPropertyHandle->GetProperty();
-				const FString ValueStr = FRigVMStruct::ExportToFullyQualifiedText(Property, (uint8*)InRawData, true);
+				const FString ValueStr = FRigVMStruct::ExportToFullyQualifiedText(Property, (uint8*)InRawData, false);
 
 				// Transform property path into pin path
 				FString PropertyPath(PinnedPropertyHandle->GetPropertyPath());
@@ -394,7 +394,10 @@ void FAnimNextEdGraphNodeCustomization::PopulateCategory(IDetailLayoutBuilder& D
 				// Prefix node path
 				const FString NodePath = EdGraphNode->GetModelNodePath();
 				PinPath = NodePath + TEXT(".") + PinPath;
-				Controller->SetPinDefaultValue(PinPath, ValueStr);
+				if(!Controller->SetPinDefaultValue(PinPath, ValueStr))
+				{
+					UE_LOG(LogAnimation, Warning, TEXT("Could not set default value for pin '%s'"), *PinPath);
+				}
 
 				return true;
 			});
diff --git a/Engine/Plugins/Experimental/AnimNext/Source/AnimNextEditor/Private/Common/AnimNextFunctionItemDetails.cpp b/Engine/Plugins/Experimental/AnimNext/Source/AnimNextEditor/Private/Common/AnimNextFunctionItemDetails.cpp
--- a/Engine/Plugins/Experimental/AnimNext/Source/AnimNextEditor/Private/Common/AnimNextFunctionItemDetails.cpp
+++ b/Engine/Plugins/Experimental/AnimNext/Source/AnimNextEditor/Private/Common/AnimNextFunctionItemDetails.cpp
@@ -31,11 +31,34 @@ bool FAnimNextFunctionItemDetails::HandleDoubleClick(const FToolMenuContext& Too
 			const TInstancedStruct<FWorkspaceOutlinerItemData>& Data = WorkspaceItemContext->SelectedExports[0].GetData();
 			if (Data.IsValid() && Data.GetScriptStruct() == FAnimNextGraphFunctionOutlinerData::StaticStruct())
 			{
-				const FAnimNextGraphFunctionOutlinerData& GraphFunctionData = Data.Get<FAnimNextGraphFunctionOutlinerData>();
-				if (GraphFunctionData.EditorObject.IsValid())
+				const FAnimNextGraphFunctionOutlinerData& FunctionData = Data.Get<FAnimNextGraphFunctionOutlinerData>();
+				if (FunctionData.EditorObject.IsValid())
 				{
-					WorkspaceEditor->OpenObjects({ GraphFunctionData.EditorObject.Get() });
-					return true;					
+					WorkspaceEditor->OpenObjects({ FunctionData.EditorObject.Get() });
+					return true;
+				}
+			}
+		}
+	}
+	return false;
+}
+
+bool FAnimNextFunctionItemDetails::HandleSelected(const FToolMenuContext& ToolMenuContext) const
+{
+	const UWorkspaceItemMenuContext* WorkspaceItemContext = ToolMenuContext.FindContext<UWorkspaceItemMenuContext>();
+	const UAssetEditorToolkitMenuContext* AssetEditorContext = ToolMenuContext.FindContext<UAssetEditorToolkitMenuContext>();
+	if (WorkspaceItemContext && AssetEditorContext)
+	{
+		if (const TSharedPtr<UE::Workspace::IWorkspaceEditor> WorkspaceEditor = StaticCastSharedPtr<UE::Workspace::IWorkspaceEditor>(AssetEditorContext->Toolkit.Pin()))
+		{
+			const TInstancedStruct<FWorkspaceOutlinerItemData>& Data = WorkspaceItemContext->SelectedExports[0].GetData();
+			if (Data.IsValid() && Data.GetScriptStruct() == FAnimNextGraphFunctionOutlinerData::StaticStruct())
+			{
+				const FAnimNextGraphFunctionOutlinerData& FunctionData = Data.Get<FAnimNextGraphFunctionOutlinerData>();
+				if (FunctionData.EditorObject.IsValid())
+				{
+					WorkspaceEditor->SetDetailsObjects({ FunctionData.EditorObject.Get() });
+					return true;
 				}
 			}
 		}
@@ -49,10 +72,10 @@ bool FAnimNextFunctionItemDetails::CanDelete(const FWorkspaceOutlinerItemExport&
 	const TInstancedStruct<FWorkspaceOutlinerItemData>& Data = Export.GetData();
 	if (Data.IsValid() && Data.GetScriptStruct() == FAnimNextGraphFunctionOutlinerData::StaticStruct())
 	{
-		const FAnimNextGraphFunctionOutlinerData& CollapseGraphData = Export.GetData().Get<FAnimNextGraphFunctionOutlinerData>();
-		if (CollapseGraphData.EditorObject.IsValid())
+		const FAnimNextGraphFunctionOutlinerData& FunctionData = Export.GetData().Get<FAnimNextGraphFunctionOutlinerData>();
+		if (FunctionData.EditorObject.IsValid())
 		{
-			URigVMEdGraph* EdGraph = CollapseGraphData.EditorObject.Get();
+			URigVMEdGraph* EdGraph = FunctionData.EditorObject.Get();
 			if (EdGraph->bAllowDeletion)
 			{
 				return true;
@@ -70,18 +93,37 @@ void FAnimNextFunctionItemDetails::Delete(TConstArrayView<FWorkspaceOutlinerItem
 		const TInstancedStruct<FWorkspaceOutlinerItemData>& Data = Export.GetData();
 		if (Data.IsValid() && Data.GetScriptStruct() == FAnimNextGraphFunctionOutlinerData::StaticStruct())
 		{
-			const FAnimNextGraphFunctionOutlinerData& GraphData = Export.GetData().Get<FAnimNextGraphFunctionOutlinerData>();
-			if (GraphData.EditorObject.IsValid() && GraphData.EdGraphNode.IsValid())
+			const FAnimNextGraphFunctionOutlinerData& FunctionData = Export.GetData().Get<FAnimNextGraphFunctionOutlinerData>();
+			if (FunctionData.EditorObject.IsValid())
 			{
-				const URigVMEdGraphNode* EdGraphNode = GraphData.EdGraphNode.Get();
-				if (EdGraphNode->CanUserDeleteNode())
+				if(FunctionData.EdGraphNode.IsValid())
+				{
+					const URigVMEdGraphNode* EdGraphNode = FunctionData.EdGraphNode.Get();
+					if (EdGraphNode->CanUserDeleteNode())
+					{
+						if (URigVMGraph* Model = EdGraphNode->GetModel())
+						{
+							if (URigVMNode* ModelNode = Model->FindNodeByName(*EdGraphNode->GetModelNodePath()))
+							{
+								FScopedTransaction Transaction(LOCTEXT("DeleteFunctionInOutliner", "Delete Function"));
+								EdGraphNode->GetController()->RemoveNode(ModelNode);
+							}
+						}
+					}
+				}
+				else
 				{
-					if (URigVMGraph* Model = EdGraphNode->GetModel())
+					URigVMEdGraph* EdGraph = FunctionData.EditorObject.Get();
+					if (EdGraph->bAllowDeletion)
 					{
-						if (URigVMNode* ModelNode = Model->FindNodeByName(*EdGraphNode->GetModelNodePath()))
+						if(IRigVMClientHost* ClientHost = EdGraph->GetImplementingOuter<IRigVMClientHost>())
 						{
-							FScopedTransaction Transaction(LOCTEXT("DeleteFunctionInOutliner", "Delete Function"));
-							EdGraphNode->GetController()->RemoveNode(ModelNode);
+							if(URigVMController* Controller = ClientHost->GetController(ClientHost->GetLocalFunctionLibrary()))
+							{
+								FScopedTransaction Transaction(LOCTEXT("DeleteFunctionInOutliner", "Delete Function"));
+								URigVMLibraryNode* LibraryNode = CastChecked<URigVMLibraryNode>(EdGraph->GetModel()->GetOuter());
+								Controller->RemoveFunctionFromLibrary(LibraryNode->GetFName());
+							}
 						}
 					}
 				}
@@ -95,10 +137,10 @@ bool FAnimNextFunctionItemDetails::CanRename(const FWorkspaceOutlinerItemExport&
 	const TInstancedStruct<FWorkspaceOutlinerItemData>& Data = Export.GetData();
 	if (Data.IsValid() && Data.GetScriptStruct() == FAnimNextGraphFunctionOutlinerData::StaticStruct())
 	{
-		const FAnimNextGraphFunctionOutlinerData& CollapseGraphData = Export.GetData().Get<FAnimNextGraphFunctionOutlinerData>();
-		if (CollapseGraphData.EditorObject.IsValid())
+		const FAnimNextGraphFunctionOutlinerData& FunctionData = Export.GetData().Get<FAnimNextGraphFunctionOutlinerData>();
+		if (FunctionData.EditorObject.IsValid())
 		{
-			URigVMEdGraph* EdGraph = CollapseGraphData.EditorObject.Get();
+			URigVMEdGraph* EdGraph = FunctionData.EditorObject.Get();
 			if (EdGraph->bAllowRenaming)
 			{
 				return true;
@@ -114,10 +156,10 @@ void FAnimNextFunctionItemDetails::Rename(const FWorkspaceOutlinerItemExport& Ex
 	const TInstancedStruct<FWorkspaceOutlinerItemData>& Data = Export.GetData();
 	if (Data.IsValid() && Data.GetScriptStruct() == FAnimNextGraphFunctionOutlinerData::StaticStruct())
 	{
-		const FAnimNextGraphFunctionOutlinerData& CollapseGraphData = Export.GetData().Get<FAnimNextGraphFunctionOutlinerData>();
-		if (CollapseGraphData.EditorObject.IsValid())
+		const FAnimNextGraphFunctionOutlinerData& FunctionData = Export.GetData().Get<FAnimNextGraphFunctionOutlinerData>();
+		if (FunctionData.EditorObject.IsValid())
 		{
-			URigVMEdGraph* EdGraph = CollapseGraphData.EditorObject.Get();
+			URigVMEdGraph* EdGraph = FunctionData.EditorObject.Get();
 			if (EdGraph->bAllowRenaming)
 			{
 				if (const UEdGraphSchema* GraphSchema = EdGraph->GetSchema())
@@ -160,10 +202,10 @@ UPackage* FAnimNextFunctionItemDetails::GetPackage(const FWorkspaceOutlinerItemE
 	const TInstancedStruct<FWorkspaceOutlinerItemData>& Data = Export.GetData();
 	if (Data.IsValid() && Data.GetScriptStruct() == FAnimNextGraphFunctionOutlinerData::StaticStruct())
 	{
-		const FAnimNextGraphFunctionOutlinerData& GraphFunctionData = Data.Get<FAnimNextGraphFunctionOutlinerData>();
-		if (GraphFunctionData.EditorObject.IsValid())
+		const FAnimNextGraphFunctionOutlinerData& FunctionData = Data.Get<FAnimNextGraphFunctionOutlinerData>();
+		if (FunctionData.EditorObject.IsValid())
 		{
-			return GraphFunctionData.EditorObject->GetPackage();
+			return FunctionData.EditorObject->GetPackage();
 		}
 	}
 	return nullptr;
@@ -171,7 +213,7 @@ UPackage* FAnimNextFunctionItemDetails::GetPackage(const FWorkspaceOutlinerItemE
 
 const FSlateBrush* FAnimNextFunctionItemDetails::GetItemIcon(const FWorkspaceOutlinerItemExport& Export) const
 {
-	return FAppStyle::GetBrush(TEXT("GraphEditor.EventGraph_24x"));
+	return FAppStyle::GetBrush(TEXT("GraphEditor.Function_24x"));
 }
 
 void FAnimNextFunctionItemDetails::RegisterToolMenuExtensions()
diff --git a/Engine/Plugins/Experimental/AnimNext/Source/AnimNextEditor/Private/Common/AnimNextFunctionItemDetails.h b/Engine/Plugins/Experimental/AnimNext/Source/AnimNextEditor/Private/Common/AnimNextFunctionItemDetails.h
--- a/Engine/Plugins/Experimental/AnimNext/Source/AnimNextEditor/Private/Common/AnimNextFunctionItemDetails.h
+++ b/Engine/Plugins/Experimental/AnimNext/Source/AnimNextEditor/Private/Common/AnimNextFunctionItemDetails.h
@@ -21,6 +21,7 @@ public:
 	virtual bool ValidateName(const FWorkspaceOutlinerItemExport& Export, const FText& InName, FText& OutErrorMessage) const override;
 	virtual UPackage* GetPackage(const FWorkspaceOutlinerItemExport& Export) const override;
 	virtual const FSlateBrush* GetItemIcon(const FWorkspaceOutlinerItemExport& Export) const override;
+	virtual bool HandleSelected(const FToolMenuContext& ToolMenuContext) const override;
 
 	static void RegisterToolMenuExtensions();
 	static void UnregisterToolMenuExtensions();
diff --git a/Engine/Plugins/Experimental/AnimNext/Source/AnimNextEditor/Private/Common/SActionMenu.cpp b/Engine/Plugins/Experimental/AnimNext/Source/AnimNextEditor/Private/Common/SActionMenu.cpp
--- a/Engine/Plugins/Experimental/AnimNext/Source/AnimNextEditor/Private/Common/SActionMenu.cpp
+++ b/Engine/Plugins/Experimental/AnimNext/Source/AnimNextEditor/Private/Common/SActionMenu.cpp
@@ -133,7 +133,9 @@ void SActionMenu::CollectAllAnimNextGraphActions(FGraphContextMenuBuilder& MenuB
 
 	TMap<FAssetData, FRigVMGraphFunctionHeaderArray> FunctionExports;
 	AnimNext::UncookedOnly::FUtils::GetExportedFunctionsFromAssetRegistry(UE::AnimNext::AnimNextPublicGraphFunctionsExportsRegistryTag, FunctionExports);
-	AnimNext::UncookedOnly::FUtils::GetExportedFunctionsFromAssetRegistry(UE::AnimNext::ControlRigAssetPublicGraphFunctionsExportsRegistryTag, FunctionExports);
+	// TODO: Ideally we can filter functions by schema or execute context, but right now we dont expose the schema and function execute contexts are
+	// all FRigVMExecuteContext, rather than the 'most derived' context in the function.
+//	AnimNext::UncookedOnly::FUtils::GetExportedFunctionsFromAssetRegistry(UE::AnimNext::ControlRigAssetPublicGraphFunctionsExportsRegistryTag, FunctionExports);
 
 	for (const auto& Export : FunctionExports.Array())
 	{
diff --git a/Engine/Plugins/Experimental/AnimNext/Source/AnimNextEditor/Private/Common/SRigVMFunctionPicker.cpp b/Engine/Plugins/Experimental/AnimNext/Source/AnimNextEditor/Private/Common/SRigVMFunctionPicker.cpp
new file mode 100644
--- /dev/null
+++ b/Engine/Plugins/Experimental/AnimNext/Source/AnimNextEditor/Private/Common/SRigVMFunctionPicker.cpp
@@ -0,0 +1,344 @@
+﻿// Copyright Epic Games, Inc. All Rights Reserved.
+
+#include "Common/SRigVMFunctionPicker.h"
+
+#include "AnimNextRigVMAsset.h"
+#include "DetailLayoutBuilder.h"
+#include "RigVMBlueprintGeneratedClass.h"
+#include "UncookedOnlyUtils.h"
+#include "Framework/Application/SlateApplication.h"
+#include "Misc/TextFilterExpressionEvaluator.h"
+#include "Widgets/Input/SComboButton.h"
+#include "Widgets/Input/SSearchBox.h"
+#include "Widgets/Views/STableRow.h"
+#include "AnimNextRigVMAssetEditorData.h"
+
+#define LOCTEXT_NAMESPACE "SRigVMFunctionPicker"
+
+namespace UE::AnimNext::Editor
+{
+
+class SRigVMFunctionRowWidget : public STableRow<TSharedPtr<SRigVMFunctionPicker::FEntry>>
+{
+	SLATE_BEGIN_ARGS(SRigVMFunctionRowWidget) {}
+
+	SLATE_ATTRIBUTE(FText, HighlightText)
+
+	SLATE_END_ARGS()
+
+	void Construct(const FArguments& InArgs, const TSharedRef<STableViewBase>& InOwnerTable, TSharedPtr<SRigVMFunctionPicker::FEntry> InEntry)
+	{
+		Entry = InEntry;
+		HighlightText = InArgs._HighlightText;
+		check(Entry.IsValid());
+
+		STableRow<TSharedPtr<SRigVMFunctionPicker::FEntry>>::Construct(
+			STableRow<TSharedPtr<SRigVMFunctionPicker::FEntry>>::FArguments()
+			.ToolTipText(Entry->ToolTip)
+			.Content()
+			[
+				SNew(SHorizontalBox)
+				+SHorizontalBox::Slot()
+				.VAlign(VAlign_Center)
+				.AutoWidth()
+				[
+					SNew(SImage)
+					.Image(Entry->Icon)
+				]
+				+SHorizontalBox::Slot()
+				.AutoWidth()
+				.Padding(4.0f, 2.0f)
+				.VAlign(VAlign_Center)
+				[
+					SNew(STextBlock)
+					.Text(Entry->Name)
+					.HighlightText(HighlightText)
+				]
+			],
+			InOwnerTable);
+	}
+
+	TSharedPtr<SRigVMFunctionPicker::FEntry> Entry;
+	TAttribute<FText> HighlightText;
+};
+
+void SRigVMFunctionPicker::Construct(const FArguments& InArgs)
+{
+	CurrentAsset = InArgs._CurrentAsset;
+	WeakCurrentAsset = CastChecked<UAnimNextRigVMAsset>(InArgs._CurrentAsset.GetAsset(), ECastCheckedType::NullAllowed);
+	FunctionName = InArgs._FunctionName;
+	FunctionToolTip = InArgs._FunctionToolTip;
+	OnRigVMFunctionPicked = InArgs._OnRigVMFunctionPicked;
+	OnNewFunction = InArgs._OnNewFunction;
+	bAllowNew = InArgs._AllowNew;
+	bAllowClear = InArgs._AllowClear;
+
+	TextFilter = MakeShared<FTextFilterExpressionEvaluator>(ETextFilterExpressionEvaluatorMode::Complex);
+
+	ChildSlot
+	[
+		SNew(SComboButton)
+		.ToolTipText(FunctionToolTip)
+		.ButtonContent()
+		[
+			SNew(SHorizontalBox)
+			+SHorizontalBox::Slot()
+			.AutoWidth()
+			.Padding(4.0f, 2.0f)
+			[
+				SNew(SImage)
+				.Image(FAppStyle::GetBrush(TEXT("GraphEditor.Function_16x")))
+			]
+			+SHorizontalBox::Slot()
+			.AutoWidth()
+			.VAlign(VAlign_Center)
+			[
+				SNew(STextBlock)
+				.Text(FunctionName)
+				.Font(IDetailLayoutBuilder::GetDetailFont())
+			]
+		]
+		.OnMenuOpenChanged_Lambda([this](bool bInOpen)
+		{
+			if(bInOpen)
+			{
+				SearchBox->SetText(FText::GetEmpty());
+				RefreshEntries();
+				RegisterActiveTimer(0.0f, FWidgetActiveTimerDelegate::CreateLambda([this](double InCurrentTime, float InDeltaTime)
+				{
+					FSlateApplication::Get().SetKeyboardFocus(SearchBox);
+					return EActiveTimerReturnType::Stop;
+				}));
+			}
+		})
+		.MenuContent()
+		[
+			SNew(SBox)
+			.WidthOverride(300.0f)
+			.HeightOverride(400.0f)
+			[
+				SNew(SBorder)
+				[
+					SNew(SVerticalBox)
+					+SVerticalBox::Slot()
+					.AutoHeight()
+					[
+						SAssignNew(SearchBox, SSearchBox)
+						.OnTextChanged_Lambda([this](const FText& InText)
+						{
+							FilterText = InText;
+							RequestRefreshEntries();
+						})
+					]
+					+SVerticalBox::Slot()
+					.FillHeight(1.0f)
+					[
+						SAssignNew(TreeView, STreeView<TSharedPtr<FEntry>>)
+						.SelectionMode(ESelectionMode::Single)
+						.TreeItemsSource(&FilteredEntries)
+						.OnGenerateRow_Lambda([this](TSharedPtr<FEntry> InItem, const TSharedRef<STableViewBase>& InOwnerTable) -> TSharedRef<ITableRow>
+						{
+							return SNew(SRigVMFunctionRowWidget, InOwnerTable, InItem)
+								.HighlightText_Lambda([this]()
+								{
+									return FilterText;
+								});
+						})
+						.OnGetChildren_Lambda([this](TSharedPtr<FEntry> InItem, TArray<TSharedPtr<FEntry>>& OutChildren)
+						{
+							switch(InItem->Type)
+							{
+							case EEntryType::Asset:
+								{
+									TSharedPtr<FAssetEntry> AssetEntry = StaticCastSharedPtr<FAssetEntry>(InItem);
+									if(FilterText.IsEmpty())
+									{
+										OutChildren.Append(AssetEntry->Functions);
+									}
+									else
+									{
+										OutChildren.Append(AssetEntry->FilteredFunctions);
+									}
+									break;
+								}
+							default:
+								break;
+							}
+						})
+						.OnSelectionChanged_Lambda([this](TSharedPtr<FEntry> InItem, ESelectInfo::Type)
+						{
+							FSlateApplication::Get().DismissAllMenus();
+
+							if(!InItem.IsValid())
+							{
+								return;
+							}
+
+							switch(InItem->Type)
+							{
+							case EEntryType::None:
+								{
+									OnRigVMFunctionPicked.ExecuteIfBound(FRigVMGraphFunctionHeader());
+									break;
+								}
+							case EEntryType::Asset:
+								break;
+							case EEntryType::Function:
+								{
+									TSharedPtr<FFunctionEntry> FunctionEntry = StaticCastSharedPtr<FFunctionEntry>(InItem);
+									OnRigVMFunctionPicked.ExecuteIfBound(FunctionEntry->FunctionHeader);
+									break;
+								}
+							case EEntryType::NewFunction:
+								{
+									OnNewFunction.ExecuteIfBound();
+									break;
+								}
+							}
+						})
+					]
+				]
+			]
+		]
+	];
+
+	RequestRefreshEntries();
+}
+
+void SRigVMFunctionPicker::RequestRefreshEntries()
+{
+	RegisterActiveTimer(0.0f, FWidgetActiveTimerDelegate::CreateLambda([this](double InCurrentTime, float InDeltaTime)
+	{
+		RefreshEntries();
+		return EActiveTimerReturnType::Stop;
+	}));
+}
+
+void SRigVMFunctionPicker::RefreshEntries()
+{
+	Entries.Reset();
+	FilteredEntries.Reset();
+	TextFilter->SetFilterText(FilterText);
+
+	if(bAllowClear)
+	{
+		TSharedPtr<FNoneEntry> NoneEntry = MakeShared<FNoneEntry>();
+		NoneEntry->Name = LOCTEXT("NoneEntryLabel", "None");
+		NoneEntry->ToolTip = LOCTEXT("NoneEntryTooltip", "Clear the currently selected function");
+		NoneEntry->Icon = FAppStyle::GetBrush(TEXT("Icons.XCircle"));
+		Entries.Add(NoneEntry);
+		FilteredEntries.Add(NoneEntry);
+	}
+
+	if(bAllowNew)
+	{
+		TSharedPtr<FNewFunctionEntry> NewFunctionEntry = MakeShared<FNewFunctionEntry>();
+		NewFunctionEntry->Name = LOCTEXT("NewFunctionEntryLabel", "New Function...");
+		NewFunctionEntry->ToolTip = LOCTEXT("NewFunctionEntryTooltip", "Create a new function");
+		NewFunctionEntry->Icon = FAppStyle::GetBrush(TEXT("Icons.PlusCircle"));
+		Entries.Add(NewFunctionEntry);
+		FilteredEntries.Add(NewFunctionEntry);
+	}
+
+	TMap<FAssetData, FRigVMGraphFunctionHeaderArray> FunctionExports;
+
+	if(UAnimNextRigVMAsset* RigVMAsset = WeakCurrentAsset.Get())
+	{
+		UAnimNextRigVMAssetEditorData* EditorData = UncookedOnly::FUtils::GetEditorData<UAnimNextRigVMAssetEditorData>(RigVMAsset);
+		FRigVMGraphFunctionHeaderArray HeaderArray;
+		UncookedOnly::FUtils::GetAssetPrivateFunctions(EditorData, HeaderArray);
+		if(HeaderArray.Headers.Num() > 0)
+		{
+			FRigVMGraphFunctionHeaderArray& AssetArray = FunctionExports.FindOrAdd(CurrentAsset);
+			AssetArray.Headers.Append(HeaderArray.Headers);
+		}
+	}
+
+	UncookedOnly::FUtils::GetExportedFunctionsFromAssetRegistry(AnimNextPublicGraphFunctionsExportsRegistryTag, FunctionExports);
+	// TODO: Ideally we can filter functions by schema or execute context, but right now we dont expose the schema and function execute contexts are
+	// all FRigVMExecuteContext, rather than the 'most derived' context in the function.
+//	UncookedOnly::FUtils::GetExportedFunctionsFromAssetRegistry(ControlRigAssetPublicGraphFunctionsExportsRegistryTag, FunctionExports);
+
+	for(const TPair<FAssetData, FRigVMGraphFunctionHeaderArray>& AssetExport : FunctionExports)
+	{
+		TSharedPtr<FAssetEntry> AssetEntry = MakeShared<FAssetEntry>();
+		AssetEntry->Name = FText::FromName(AssetExport.Key.AssetName);
+		AssetEntry->ToolTip = FText::FromString(AssetExport.Key.GetFullName());
+		AssetEntry->Asset = AssetExport.Key;
+		Entries.Add(AssetEntry);
+
+		for(const FRigVMGraphFunctionHeader& FunctionHeader : AssetExport.Value.Headers)
+		{
+			TSharedPtr<FFunctionEntry> FunctionEntry = MakeShared<FFunctionEntry>();
+			FunctionEntry->Name = FText::FromName(FunctionHeader.Name);
+			FunctionEntry->ToolTip = FunctionHeader.Description.Len() > 0 ? FText::FromString(FunctionHeader.Description) : FText::FromName(FunctionHeader.Name);
+			FunctionEntry->Icon = FAppStyle::GetBrush(TEXT("GraphEditor.Function_16x"));
+			FunctionEntry->FunctionHeader = FunctionHeader;
+			AssetEntry->Functions.Add(FunctionEntry);
+		}
+	}
+
+	for(const TSharedPtr<FEntry>& Entry : Entries)
+	{
+		TreeView->SetItemExpansion(Entry, true);
+	}
+
+	if(FilterText.IsEmpty())
+	{
+		FilteredEntries.Reserve(Entries.Num());
+	}
+
+	// Filter if we need to
+	for(const TSharedPtr<FEntry>& Entry : Entries)
+	{
+		switch(Entry->Type)
+		{
+		case EEntryType::Asset:
+			{
+				TSharedPtr<FAssetEntry> AssetEntry = StaticCastSharedPtr<FAssetEntry>(Entry);
+				AssetEntry->FilteredFunctions.Reset();
+				if(FilterText.IsEmpty())
+				{
+					FilteredEntries.Add(AssetEntry);
+					AssetEntry->FilteredFunctions.Reserve(AssetEntry->Functions.Num());
+					for(const TSharedPtr<FEntry>& FunctionEntry : AssetEntry->Functions)
+					{
+						AssetEntry->FilteredFunctions.Add(FunctionEntry);
+					}
+				}
+				else
+				{
+					bool bAddedAsset = false;
+					if(TextFilter->TestTextFilter(FBasicStringFilterExpressionContext(AssetEntry->Name.ToString())))
+					{
+						FilteredEntries.Add(AssetEntry);
+						bAddedAsset = true;
+					}
+
+					// Also show asset if one of the functions passes the filter
+					for(const TSharedPtr<FEntry>& FunctionEntry : AssetEntry->Functions)
+					{
+						if(TextFilter->TestTextFilter(FBasicStringFilterExpressionContext(FunctionEntry->Name.ToString())))
+						{
+							AssetEntry->FilteredFunctions.Add(FunctionEntry);
+							if(!bAddedAsset)
+							{
+								FilteredEntries.Add(AssetEntry);
+								bAddedAsset = true;
+							}
+						}
+					}
+				}
+				break;
+			}
+		default:
+			break;
+		}
+	}
+
+	TreeView->RequestTreeRefresh();
+}
+
+}
+
+#undef LOCTEXT_NAMESPACE
\ No newline at end of file
diff --git a/Engine/Plugins/Experimental/AnimNext/Source/AnimNextEditor/Private/Graph/AnimNextEdGraphCustomization.cpp b/Engine/Plugins/Experimental/AnimNext/Source/AnimNextEditor/Private/Graph/AnimNextEdGraphCustomization.cpp
new file mode 100644
--- /dev/null
+++ b/Engine/Plugins/Experimental/AnimNext/Source/AnimNextEditor/Private/Graph/AnimNextEdGraphCustomization.cpp
@@ -0,0 +1,48 @@
+﻿// Copyright Epic Games, Inc. All Rights Reserved.
+
+#include "Graph/AnimNextEdGraphCustomization.h"
+
+#include "AnimNextEdGraph.h"
+#include "DetailLayoutBuilder.h"
+#include "Editor/RigVMGraphDetailCustomization.h"
+#include "RigVMModel/RigVMController.h"
+#include "IWorkspaceEditor.h"
+
+#define LOCTEXT_NAMESPACE "AnimNextEdGraphCustomization"
+
+namespace UE::AnimNext::Editor
+{
+
+void FAnimNextEdGraphCustomization::CustomizeDetails(IDetailLayoutBuilder& DetailBuilder)
+{
+	TArray<TWeakObjectPtr<UObject>> Objects;
+	DetailBuilder.GetObjectsBeingCustomized(Objects);
+
+	RigVMGraphDetailCustomizationImpl.Reset();
+
+	if (Objects.Num() != 1)
+	{
+		return;
+	}
+
+	if (UAnimNextEdGraph* EdGraph = Cast<UAnimNextEdGraph>(Objects[0].Get()))
+	{
+		if(URigVMGraph* Model = EdGraph->GetModel())
+		{
+			TWeakPtr<FRigVMEditor> WeakEditor; // passing an empty RigVMEditor to FRigVMCollapseGraphLayout
+			IRigVMClientHost* RigVMClientHost = EdGraph->GetController()->GetClientHost();
+			check(RigVMClientHost);
+
+			RigVMGraphDetailCustomizationImpl = MakeShared<FRigVMGraphDetailCustomizationImpl>();
+			RigVMGraphDetailCustomizationImpl->CustomizeDetails(DetailBuilder
+				, Model
+				, RigVMClientHost->GetController(Model)
+				, RigVMClientHost
+				, WeakEditor);
+		}
+	}
+}
+
+} // end namespace UE::AnimNext::Editor
+
+#undef LOCTEXT_NAMESPACE
diff --git a/Engine/Plugins/Experimental/AnimNext/Source/AnimNextEditor/Private/Graph/AnimNextEdGraphCustomization.h b/Engine/Plugins/Experimental/AnimNext/Source/AnimNextEditor/Private/Graph/AnimNextEdGraphCustomization.h
new file mode 100644
--- /dev/null
+++ b/Engine/Plugins/Experimental/AnimNext/Source/AnimNextEditor/Private/Graph/AnimNextEdGraphCustomization.h
@@ -0,0 +1,39 @@
+﻿// Copyright Epic Games, Inc. All Rights Reserved.
+
+#pragma once
+
+#include "IDetailCustomization.h"
+#include "IPropertyTypeCustomization.h"
+
+class URigVMPin;
+class IDetailCategoryBuilder;
+class UAnimNextEdGraphNode;
+class FRigVMGraphDetailCustomizationImpl;
+
+namespace UE::Workspace
+{
+	class IWorkspaceEditor;
+}
+
+namespace UE::AnimNext::Editor
+{
+	class FAnimNextEditorModule;
+}
+
+namespace UE::AnimNext::Editor
+{
+
+class FAnimNextEdGraphCustomization : public IDetailCustomization
+{
+	FAnimNextEdGraphCustomization() = default;
+
+	// IDetailCustomization interface
+	virtual void CustomizeDetails(IDetailLayoutBuilder& DetailBuilder) override;
+
+	template <typename ObjectType, ESPMode Mode>
+	friend class SharedPointerInternals::TIntrusiveReferenceController;
+
+	TSharedPtr<FRigVMGraphDetailCustomizationImpl> RigVMGraphDetailCustomizationImpl;
+};
+
+}
diff --git a/Engine/Plugins/Experimental/AnimNext/Source/AnimNextEditor/Private/Variables/SVariablesView.cpp b/Engine/Plugins/Experimental/AnimNext/Source/AnimNextEditor/Private/Variables/SVariablesView.cpp
--- a/Engine/Plugins/Experimental/AnimNext/Source/AnimNextEditor/Private/Variables/SVariablesView.cpp
+++ b/Engine/Plugins/Experimental/AnimNext/Source/AnimNextEditor/Private/Variables/SVariablesView.cpp
@@ -20,6 +20,7 @@
 #include "AnimNextRigVMAsset.h"
 #include "IWorkspaceEditor.h"
 #include "WorkspaceAssetRegistryInfo.h"
+#include "Common/AnimNextFunctionItemDetails.h"
 #include "Outliner/VariablesOutlinerAssetItem.h"
 
 #define LOCTEXT_NAMESPACE "SVariablesView"
@@ -275,6 +276,20 @@ void SVariablesView::HandleWorkspaceOutlinerSelectionChanged(TConstArrayView<FWo
 				}
 			}
 		}
+		else if(Export.GetData().GetScriptStruct()->IsChildOf(FAnimNextGraphFunctionOutlinerData::StaticStruct()))
+		{
+			const FAnimNextGraphFunctionOutlinerData& EntryData = Export.GetData().Get<FAnimNextGraphFunctionOutlinerData>();
+			if(EntryData.EditorObject.IsValid())
+			{
+				if(URigVMEdGraph* EdGraph = EntryData.EditorObject.Get())
+				{
+					if(UAnimNextRigVMAsset* Asset = EdGraph->GetTypedOuter<UAnimNextRigVMAsset>())
+					{
+						Assets.Add(TSoftObjectPtr<UAnimNextRigVMAsset>(Asset));
+					}
+				}
+			}
+		}
 	}
 
 	VariablesOutliner->SetAssets(Assets);
diff --git a/Engine/Plugins/Experimental/AnimNext/Source/AnimNextUncookedOnly/Internal/AnimNextRigVMAssetEditorData.h b/Engine/Plugins/Experimental/AnimNext/Source/AnimNextUncookedOnly/Internal/AnimNextRigVMAssetEditorData.h
--- a/Engine/Plugins/Experimental/AnimNext/Source/AnimNextUncookedOnly/Internal/AnimNextRigVMAssetEditorData.h
+++ b/Engine/Plugins/Experimental/AnimNext/Source/AnimNextUncookedOnly/Internal/AnimNextRigVMAssetEditorData.h
@@ -43,6 +43,7 @@ namespace UE::AnimNext::Editor
 	class FVariableProxyCustomization;
 	struct FPublicVariablesImpl;
 	class SActionMenu;
+	class FCallFunctionSharedDataDetails;
 }
 
 namespace UE::AnimNext::Tests
@@ -129,6 +130,9 @@ public:
 	/** Adds a data interface to this asset */
 	UAnimNextDataInterfaceEntry* AddDataInterface(UAnimNextDataInterface* InDataInterface, bool bSetupUndoRedo = true, bool bPrintPythonCommand = true);
 
+	/** Adds a function to this asset */
+	URigVMLibraryNode* AddFunction(FName InFunctionName, bool bInMutable, bool bInSetupUndoRedo = true, bool bInPrintPythonCommand = true);
+
 	// Find an entry by name
 	UAnimNextRigVMAssetEntry* FindEntry(FName InName) const;
 
@@ -177,6 +181,7 @@ protected:
 	friend class UAnimNextDataInterfaceEntry;
 	friend struct UE::AnimNext::Editor::FPublicVariablesImpl;
 	friend UE::AnimNext::Editor::SActionMenu;
+	friend class UE::AnimNext::Editor::FCallFunctionSharedDataDetails;
 
 	// UObject interface
 	virtual void Serialize(FArchive& Ar) override;
@@ -434,6 +439,9 @@ protected:
 	// Remove an entry to this asset, taking into account external packaging status
 	void RemoveEntryInternal(UAnimNextRigVMAssetEntry* InEntry);
 
+	// Remove any programmatic graphs generated during compilation and consign them to the transient package
+	void RemoveProgrammaticGraphs(TArrayView<URigVMGraph*> InGraphs);
+
 	/** All entries in this asset - not saved, discovered at load time and also contains InternalEntries */
 	UPROPERTY(transient)
 	TArray<TObjectPtr<UAnimNextRigVMAssetEntry>> Entries;
diff --git a/Engine/Plugins/Experimental/AnimNext/Source/AnimNextUncookedOnly/Internal/UncookedOnlyUtils.h b/Engine/Plugins/Experimental/AnimNext/Source/AnimNextUncookedOnly/Internal/UncookedOnlyUtils.h
--- a/Engine/Plugins/Experimental/AnimNext/Source/AnimNextUncookedOnly/Internal/UncookedOnlyUtils.h
+++ b/Engine/Plugins/Experimental/AnimNext/Source/AnimNextUncookedOnly/Internal/UncookedOnlyUtils.h
@@ -98,9 +98,6 @@ struct FAnimNextAssetRegistryExports
 
 	UPROPERTY()
 	TArray<FAnimNextAssetRegistryExportedVariable> Variables;
-
-	UPROPERTY()
-	TArray<FRigVMGraphFunctionHeader> PublicHeaders;
 };
 
 namespace UE::AnimNext::UncookedOnly
@@ -164,7 +161,7 @@ struct ANIMNEXTUNCOOKEDONLY_API FUtils
 	static bool GetExportedVariablesFromAssetRegistry(TMap<FAssetData, FAnimNextAssetRegistryExports>& OutExports);
 
 	// Gets the functions that are exported to the asset registry for an asset
-	static bool GetExportedFunctionsForAsset(const FAssetData& InAsset, FAnimNextAssetRegistryExports& OutExports);
+	static bool GetExportedFunctionsForAsset(const FAssetData& InAsset, FName Tag, FRigVMGraphFunctionHeaderArray& OutExports);
 
 	// Gets all the functions that are exported to the asset registry for the specified Tag
 	static bool GetExportedFunctionsFromAssetRegistry(FName Tag, TMap<FAssetData, FRigVMGraphFunctionHeaderArray>& OutExports);
@@ -173,6 +170,12 @@ struct ANIMNEXTUNCOOKEDONLY_API FUtils
 	static void GetAssetVariables(const UAnimNextRigVMAssetEditorData* EditorData, FAnimNextAssetRegistryExports& OutExports);
 	static void GetAssetVariables(const UAnimNextRigVMAssetEditorData* EditorData, TSet<FAnimNextAssetRegistryExportedVariable>& OutExports);
 
+	// Gets the exported public functions that are used by a RigVM asset
+	static void GetAssetFunctions(const UAnimNextRigVMAssetEditorData* InEditorData, FRigVMGraphFunctionHeaderArray& OutExports);
+
+	// Gets the non-exported private functions that are used by a RigVM asset
+	static void GetAssetPrivateFunctions(const UAnimNextRigVMAssetEditorData* InEditorData, FRigVMGraphFunctionHeaderArray& OutExports);
+
 	// Gets the asset-registry information needed for representing the contained data into the Workspace Outliner
 	static void GetAssetOutlinerItems(const UAnimNextRigVMAssetEditorData* EditorData, FWorkspaceOutlinerItemExports& OutExports, FAssetRegistryTagsContext Context);
 	static void CreateSubGraphsOutlinerItemsRecursive(const UAnimNextRigVMAssetEditorData* EditorData, FWorkspaceOutlinerItemExports& OutExports, FWorkspaceOutlinerItemExport& ParentExport, URigVMEdGraph* RigVMEdGraph, FAssetRegistryTagsContext Context);
@@ -190,6 +193,9 @@ struct ANIMNEXTUNCOOKEDONLY_API FUtils
 #if WITH_EDITOR
 	static void OpenProgrammaticGraphs(UAnimNextRigVMAssetEditorData* EditorData, const TArray<URigVMGraph*>& ProgrammaticGraphs);
 #endif // WITH_EDITOR
+
+	// Make an event name that we use as a wrapper to call RigVM functions 
+	static FString MakeFunctionWrapperEventName(FName InFunctionName);
 };
 
 }
\ No newline at end of file
diff --git a/Engine/Plugins/Experimental/AnimNext/Source/AnimNextUncookedOnly/Private/AnimNextRigVMAssetEditorData.cpp b/Engine/Plugins/Experimental/AnimNext/Source/AnimNextUncookedOnly/Private/AnimNextRigVMAssetEditorData.cpp
--- a/Engine/Plugins/Experimental/AnimNext/Source/AnimNextUncookedOnly/Private/AnimNextRigVMAssetEditorData.cpp
+++ b/Engine/Plugins/Experimental/AnimNext/Source/AnimNextUncookedOnly/Private/AnimNextRigVMAssetEditorData.cpp
@@ -267,13 +267,7 @@ void UAnimNextRigVMAssetEditorData::GetAssetRegistryTags(FAssetRegistryTagsConte
 
 	{
 		FRigVMGraphFunctionHeaderArray FunctionExports;
-		for (const FRigVMGraphFunctionData& FunctionData : GraphFunctionStore.PublicFunctions)
-		{
-			if (FunctionData.CompilationData.IsValid())
-			{
-				FunctionExports.Headers.Add(FunctionData.Header);
-			}
-		}
+		UE::AnimNext::UncookedOnly::FUtils::GetAssetFunctions(this, FunctionExports);
 
 		FString TagValue;
 		const FArrayProperty* HeadersProperty = CastField<FArrayProperty>(FRigVMGraphFunctionHeaderArray::StaticStruct()->FindPropertyByName(TEXT("Headers")));
@@ -827,7 +821,11 @@ void UAnimNextRigVMAssetEditorData::RecompileVM()
 	{
 		FUtils::OpenProgrammaticGraphs(this, ProgrammaticGraphs);
 	}
+	else
 #endif
+	{
+		RemoveProgrammaticGraphs(ProgrammaticGraphs);
+	}
 
 #if WITH_EDITOR
 //	RefreshBreakpoints(EditorData);
@@ -840,6 +838,17 @@ void UAnimNextRigVMAssetEditorData::RecompileVM()
 	}
 }
 
+void UAnimNextRigVMAssetEditorData::RemoveProgrammaticGraphs(TArrayView<URigVMGraph*> InGraphs)
+{
+	FRigVMClient* VMClient = GetRigVMClient();
+	
+	for(URigVMGraph* Graph : InGraphs)
+	{
+		VMClient->RemoveController(Graph);
+		Graph->Rename(nullptr, GetTransientPackage(), REN_ForceNoResetLoaders | REN_DoNotDirty | REN_DontCreateRedirectors | REN_NonTransactional);
+	}
+}
+
 void UAnimNextRigVMAssetEditorData::RecompileVMIfRequired()
 {
 	if (bVMRecompilationRequired)
@@ -1728,6 +1737,12 @@ UAnimNextDataInterfaceEntry* UAnimNextRigVMAssetEditorData::AddDataInterface(UAn
 	return NewEntry;
 }
 
+URigVMLibraryNode* UAnimNextRigVMAssetEditorData::AddFunction(FName InFunctionName, bool bInMutable, bool bInSetupUndoRedo, bool bInPrintPythonCommand)
+{
+	URigVMController* Controller = RigVMClient.GetOrCreateController(GetLocalFunctionLibrary());
+	return Controller->AddFunctionToLibrary(InFunctionName, bInMutable, FVector2D::ZeroVector, bInSetupUndoRedo, bInPrintPythonCommand);
+}
+
 bool UAnimNextRigVMAssetEditorData::HasPublicVariables() const
 {
 	for(UAnimNextRigVMAssetEntry* Entry : Entries)
diff --git a/Engine/Plugins/Experimental/AnimNext/Source/AnimNextUncookedOnly/Private/UncookedOnlyUtils.cpp b/Engine/Plugins/Experimental/AnimNext/Source/AnimNextUncookedOnly/Private/UncookedOnlyUtils.cpp
--- a/Engine/Plugins/Experimental/AnimNext/Source/AnimNextUncookedOnly/Private/UncookedOnlyUtils.cpp
+++ b/Engine/Plugins/Experimental/AnimNext/Source/AnimNextUncookedOnly/Private/UncookedOnlyUtils.cpp
@@ -718,22 +718,32 @@ bool FUtils::GetExportedVariablesFromAssetRegistry(TMap<FAssetData, FAnimNextAss
 	return OutExports.Num() > 0;
 }
 
-bool FUtils::GetExportedFunctionsForAsset(const FAssetData& InAsset, FAnimNextAssetRegistryExports& OutExports)
+void FUtils::GetAssetFunctions(const UAnimNextRigVMAssetEditorData* InEditorData, FRigVMGraphFunctionHeaderArray& OutExports)
 {
-	const FString TagValue = InAsset.GetTagValueRef<FString>(UE::AnimNext::ExportsAnimNextAssetRegistryTag);
-	
-	bool bContainsFunctions = false;
-	FAnimNextAssetRegistryExports AssetExports;
-	if (FAnimNextAssetRegistryExports::StaticStruct()->ImportText(*TagValue, &AssetExports, nullptr, PPF_None, nullptr, FAnimNextAssetRegistryExports::StaticStruct()->GetName()) != nullptr)
+	for (const FRigVMGraphFunctionData& FunctionData : InEditorData->GraphFunctionStore.PublicFunctions)
 	{
-		if (AssetExports.PublicHeaders.Num() > 0)
+		if (FunctionData.CompilationData.IsValid())
 		{
-			OutExports = MoveTemp(AssetExports);
-			bContainsFunctions = true;
+			OutExports.Headers.Add(FunctionData.Header);
 		}
 	}
+}
 
-	return bContainsFunctions;
+void FUtils::GetAssetPrivateFunctions(const UAnimNextRigVMAssetEditorData* InEditorData, FRigVMGraphFunctionHeaderArray& OutExports)
+{
+	for (const FRigVMGraphFunctionData& FunctionData : InEditorData->GraphFunctionStore.PrivateFunctions)
+	{
+		// Note: We dont check compilation data here as private functions are not compiled if they are not referenced
+		OutExports.Headers.Add(FunctionData.Header);
+	}
+}
+
+bool FUtils::GetExportedFunctionsForAsset(const FAssetData& InAsset, FName Tag, FRigVMGraphFunctionHeaderArray& OutExports)
+{
+	const FString TagValue = InAsset.GetTagValueRef<FString>(Tag);
+	const FArrayProperty* HeadersProperty = CastField<FArrayProperty>(FRigVMGraphFunctionHeaderArray::StaticStruct()->FindPropertyByName(TEXT("Headers")));
+	HeadersProperty->ImportText_Direct(*TagValue, &OutExports, nullptr, EPropertyPortFlags::PPF_None);
+	return OutExports.Headers.Num() > 0;
 }
 
 bool FUtils::GetExportedFunctionsFromAssetRegistry(FName Tag, TMap<FAssetData, FRigVMGraphFunctionHeaderArray>& OutExports)
@@ -752,7 +762,8 @@ bool FUtils::GetExportedFunctionsFromAssetRegistry(FName Tag, TMap<FAssetData, F
 		{
 			if (AssetExports.Headers.Num() > 0)
 			{
-				OutExports.Add(Asset, MoveTemp(AssetExports));
+				FRigVMGraphFunctionHeaderArray& AssetArray = OutExports.FindOrAdd(Asset);
+				AssetArray.Headers.Append(MoveTemp(AssetExports.Headers));
 			}
 		}
 	}
@@ -842,7 +853,7 @@ void FUtils::GetAssetOutlinerItems(const UAnimNextRigVMAssetEditorData* EditorDa
 		}
 	}
 
-	//CreateFunctionLibraryOutlinerItemsRecursive(EditorData, OutExports, AssetIdentifier, EditorData->GetRigVMGraphFunctionStore()->PublicFunctions, EditorData->GetRigVMGraphFunctionStore()->PrivateFunctions);
+	CreateFunctionLibraryOutlinerItemsRecursive(EditorData, OutExports, AssetIdentifier, EditorData->GetRigVMGraphFunctionStore()->PublicFunctions, EditorData->GetRigVMGraphFunctionStore()->PrivateFunctions);
 }
 
 void ProcessPinAssetReferences(const URigVMPin* InPin, FWorkspaceOutlinerItemExports& OutExports, const FWorkspaceOutlinerItemExport& ParentExport)
@@ -973,10 +984,8 @@ void FUtils::CreateFunctionLibraryOutlinerItemsRecursive(const UAnimNextRigVMAss
 {
 	if (PrivateFunctions.Num() > 0 || PublicFunctions.Num() > 0)
 	{
-		FWorkspaceOutlinerItemExport& Export = OutExports.Exports.Add_GetRef(FWorkspaceOutlinerItemExport(*GetFunctionLibraryDisplayName().ToString(), ParentExport));
-
-		CreateFunctionsOutlinerItemsRecursive(EditorData, OutExports, Export, PrivateFunctions, false);
-		CreateFunctionsOutlinerItemsRecursive(EditorData, OutExports, Export, PublicFunctions, true);
+		CreateFunctionsOutlinerItemsRecursive(EditorData, OutExports, ParentExport, PrivateFunctions, false);
+		CreateFunctionsOutlinerItemsRecursive(EditorData, OutExports, ParentExport, PublicFunctions, true);
 	}
 }
 
@@ -1041,6 +1050,12 @@ void FUtils::OpenProgrammaticGraphs(UAnimNextRigVMAssetEditorData* EditorData, c
 	}
 }
 #endif // WITH_EDITOR
+
+FString FUtils::MakeFunctionWrapperEventName(FName InFunctionName)
+{
+	return TEXT("__InternalCall_") + InFunctionName.ToString();
+}
+
 }
 
 #undef LOCTEXT_NAMESPACE
\ No newline at end of file
diff --git a/Engine/Plugins/Experimental/AnimNextAnimGraph/Source/AnimNextAnimGraph/Internal/Graph/AnimNextAnimationGraph.h b/Engine/Plugins/Experimental/AnimNextAnimGraph/Source/AnimNextAnimGraph/Internal/Graph/AnimNextAnimationGraph.h
--- a/Engine/Plugins/Experimental/AnimNextAnimGraph/Source/AnimNextAnimGraph/Internal/Graph/AnimNextAnimationGraph.h
+++ b/Engine/Plugins/Experimental/AnimNextAnimGraph/Source/AnimNextAnimGraph/Internal/Graph/AnimNextAnimationGraph.h
@@ -140,6 +140,11 @@ protected:
 	UPROPERTY()
 	TArray<TObjectPtr<UObject>> GraphReferencedObjects;
 
+	// This is a list of all referenced soft objects in the graph shared data
+	// Used to serialize the soft objects correctly as we dont use FArchiveUObject
+	UPROPERTY()
+	TArray<FSoftObjectPath> GraphReferencedSoftObjects;
+
 	// The entry point that this graph defaults to using
 	UPROPERTY(EditAnywhere, Category = "Graph")
 	FName DefaultEntryPoint = FRigUnit_AnimNextGraphRoot::DefaultEntryPoint;
diff --git a/Engine/Plugins/Experimental/AnimNextAnimGraph/Source/AnimNextAnimGraph/Internal/Graph/AnimNextGraphContextData.h b/Engine/Plugins/Experimental/AnimNextAnimGraph/Source/AnimNextAnimGraph/Internal/Graph/AnimNextGraphContextData.h
--- a/Engine/Plugins/Experimental/AnimNextAnimGraph/Source/AnimNextAnimGraph/Internal/Graph/AnimNextGraphContextData.h
+++ b/Engine/Plugins/Experimental/AnimNextAnimGraph/Source/AnimNextAnimGraph/Internal/Graph/AnimNextGraphContextData.h
@@ -19,24 +19,14 @@ struct FAnimNextGraphContextData : public FAnimNextModuleContextData
 
 	FAnimNextGraphContextData() = default;
 
-	FAnimNextGraphContextData(FAnimNextModuleInstance* InModuleInstance, const FAnimNextGraphInstance* InInstance, const TConstArrayView<UE::AnimNext::FLatentPropertyHandle>& InLatentHandles, void* InDestinationBasePtr, bool bInIsFrozen)
+	FAnimNextGraphContextData(FAnimNextModuleInstance* InModuleInstance, const FAnimNextGraphInstance* InInstance)
 		: FAnimNextModuleContextData(InModuleInstance, InInstance)
-		, LatentHandles(InLatentHandles)
-		, DestinationBasePtr(InDestinationBasePtr)
-		, bIsFrozen(bInIsFrozen)
 	{
 	}
 
 	const FAnimNextGraphInstance& GetGraphInstance() const { return static_cast<const FAnimNextGraphInstance&>(GetDataInterfaceInstance()); }
-	const TConstArrayView<UE::AnimNext::FLatentPropertyHandle>& GetLatentHandles() const { return LatentHandles; }
-	void* GetDestinationBasePtr() const { return DestinationBasePtr; }
-	bool IsFrozen() const { return bIsFrozen; }
 
 private:
-	TConstArrayView<UE::AnimNext::FLatentPropertyHandle> LatentHandles;
-	void* DestinationBasePtr = nullptr;
-	bool bIsFrozen = false;
-
 	friend struct FAnimNextGraphInstance;
 	friend struct FAnimNextExecuteContext;
 };
diff --git a/Engine/Plugins/Experimental/AnimNextAnimGraph/Source/AnimNextAnimGraph/Internal/Graph/AnimNextGraphLatentPropertiesContextData.h b/Engine/Plugins/Experimental/AnimNextAnimGraph/Source/AnimNextAnimGraph/Internal/Graph/AnimNextGraphLatentPropertiesContextData.h
new file mode 100644
--- /dev/null
+++ b/Engine/Plugins/Experimental/AnimNextAnimGraph/Source/AnimNextAnimGraph/Internal/Graph/AnimNextGraphLatentPropertiesContextData.h
@@ -0,0 +1,42 @@
+﻿// Copyright Epic Games, Inc. All Rights Reserved.
+
+#pragma once
+
+#include "CoreMinimal.h"
+#include "AnimNextGraphContextData.h"
+#include "Module/AnimNextModuleContextData.h"
+#include "Graph/AnimNextGraphInstance.h"
+#include "AnimNextGraphLatentPropertiesContextData.generated.h"
+
+namespace UE::AnimNext
+{
+	struct FLatentPropertyHandle;
+}
+
+USTRUCT()
+struct FAnimNextGraphLatentPropertiesContextData : public FAnimNextGraphContextData
+{
+	GENERATED_BODY()
+
+	FAnimNextGraphLatentPropertiesContextData() = default;
+
+	FAnimNextGraphLatentPropertiesContextData(FAnimNextModuleInstance* InModuleInstance, const FAnimNextGraphInstance* InInstance, const TConstArrayView<UE::AnimNext::FLatentPropertyHandle>& InLatentHandles, void* InDestinationBasePtr, bool bInIsFrozen)
+		: FAnimNextGraphContextData(InModuleInstance, InInstance)
+		, LatentHandles(InLatentHandles)
+		, DestinationBasePtr(InDestinationBasePtr)
+		, bIsFrozen(bInIsFrozen)
+	{
+	}
+
+	const TConstArrayView<UE::AnimNext::FLatentPropertyHandle>& GetLatentHandles() const { return LatentHandles; }
+	void* GetDestinationBasePtr() const { return DestinationBasePtr; }
+	bool IsFrozen() const { return bIsFrozen; }
+
+private:
+	TConstArrayView<UE::AnimNext::FLatentPropertyHandle> LatentHandles;
+	void* DestinationBasePtr = nullptr;
+	bool bIsFrozen = false;
+
+	friend struct FAnimNextGraphInstance;
+	friend struct FAnimNextExecuteContext;
+};
diff --git a/Engine/Plugins/Experimental/AnimNextAnimGraph/Source/AnimNextAnimGraph/Internal/TraitCore/TraitReader.h b/Engine/Plugins/Experimental/AnimNextAnimGraph/Source/AnimNextAnimGraph/Internal/TraitCore/TraitReader.h
--- a/Engine/Plugins/Experimental/AnimNextAnimGraph/Source/AnimNextAnimGraph/Internal/TraitCore/TraitReader.h
+++ b/Engine/Plugins/Experimental/AnimNextAnimGraph/Source/AnimNextAnimGraph/Internal/TraitCore/TraitReader.h
@@ -35,7 +35,7 @@ namespace UE::AnimNext
 			NodeInstanceDataTooLarge,	// Exceeded the maximum node instance data size, @see FNodeInstance::MAXIMUM_NODE_INSTANCE_DATA_SIZE
 		};
 
-		explicit FTraitReader(const TArray<TObjectPtr<UObject>>& InGraphReferencedObjects, FArchive& Ar);
+		explicit FTraitReader(const TArray<TObjectPtr<UObject>>& InGraphReferencedObjects, const TArray<FSoftObjectPath>& InGraphReferencedSoftObjects, FArchive& Ar);
 
 		[[nodiscard]] EErrorState ReadGraph(TArray<uint8>& GraphSharedData);
 
@@ -54,6 +54,8 @@ namespace UE::AnimNext
 		// FArchive implementation
 		virtual FArchive& operator<<(UObject*& Obj) override;
 		virtual FArchive& operator<<(FObjectPtr& Obj) override;
+		virtual FArchive& operator<<(FSoftObjectPath& Value) override;
+		virtual FArchive& operator<<(FSoftObjectPtr& Value) override;
 
 	private:
 		// Call first to read the graph shared data
@@ -62,6 +64,9 @@ namespace UE::AnimNext
 		// A list of UObject references within the graph
 		const TArray<TObjectPtr<UObject>>& GraphReferencedObjects;
 
+		// A list of soft object references within the graph
+		const TArray<FSoftObjectPath>& GraphReferencedSoftObjects;
+
 		// A list of node handles for each node within the archive
 		TArray<FNodeHandle> NodeHandles;
 
diff --git a/Engine/Plugins/Experimental/AnimNextAnimGraph/Source/AnimNextAnimGraph/Internal/TraitCore/TraitWriter.h b/Engine/Plugins/Experimental/AnimNextAnimGraph/Source/AnimNextAnimGraph/Internal/TraitCore/TraitWriter.h
--- a/Engine/Plugins/Experimental/AnimNextAnimGraph/Source/AnimNextAnimGraph/Internal/TraitCore/TraitWriter.h
+++ b/Engine/Plugins/Experimental/AnimNextAnimGraph/Source/AnimNextAnimGraph/Internal/TraitCore/TraitWriter.h
@@ -10,6 +10,7 @@
 #include "TraitCore/LatentPropertyHandle.h"
 #include "TraitCore/NodeHandle.h"
 #include "TraitCore/NodeTemplateRegistryHandle.h"
+#include "UObject/SoftObjectPath.h"
 
 namespace UE::AnimNext
 {
@@ -63,10 +64,15 @@ namespace UE::AnimNext
 		// Returns the list of referenced UObjects in this graph
 		[[nodiscard]] const TArray<UObject*>& GetGraphReferencedObjects() const;
 
+		// Returns the list of referenced soft objects in this graph
+		[[nodiscard]] const TArray<FSoftObjectPath>& GetGraphReferencedSoftObjects() const;
+
 		// FArchive implementation
 		virtual FArchive& operator<<(UObject*& Obj) override;
 		virtual FArchive& operator<<(FObjectPtr& Obj) override;
-
+		virtual FArchive& operator<<(FSoftObjectPath& Value) override;
+		virtual FArchive& operator<<(FSoftObjectPtr& Value) override;
+		
 	private:
 		struct FNodeMapping
 		{
@@ -88,6 +94,7 @@ namespace UE::AnimNext
 
 		// To track node writing
 		TArray<UObject*> GraphReferencedObjects;
+		TArray<FSoftObjectPath> GraphReferencedSoftObjects;
 		uint32 NumNodesWritten;
 		bool bIsNodeWriting;
 
diff --git a/Engine/Plugins/Experimental/AnimNextAnimGraph/Source/AnimNextAnimGraph/Internal/Traits/CallFunction.h b/Engine/Plugins/Experimental/AnimNextAnimGraph/Source/AnimNextAnimGraph/Internal/Traits/CallFunction.h
new file mode 100644
--- /dev/null
+++ b/Engine/Plugins/Experimental/AnimNextAnimGraph/Source/AnimNextAnimGraph/Internal/Traits/CallFunction.h
@@ -0,0 +1,79 @@
+﻿// Copyright Epic Games, Inc. All Rights Reserved.
+
+#pragma once
+
+#include "CoreMinimal.h"
+#include "RigVMCore/RigVMGraphFunctionDefinition.h"
+#include "TraitCore/Trait.h"
+#include "TraitInterfaces/IUpdate.h"
+#include "CallFunction.generated.h"
+
+UENUM()
+enum class EAnimNextCallFunctionCallSite : uint8
+{
+	// Called each time the trait becomes relevant
+	BecomeRelevant,
+
+	// Called on update before the trait and any stack super-traits update
+	PreUpdate,
+
+	// Called on update after the trait and any stack super-traits update
+	PostUpdate,
+};
+
+USTRUCT(meta = (DisplayName = "Call Function"))
+struct FAnimNextCallFunctionSharedData : public FAnimNextTraitSharedData
+{
+	GENERATED_BODY()
+
+#if WITH_EDITORONLY_DATA
+	UPROPERTY(VisibleAnywhere, Category = "Call Function", meta=(Hidden))
+	FRigVMGraphFunctionHeader FunctionHeader;
+#endif
+
+	/** The function to call */
+	UPROPERTY(EditAnywhere, Category = "Call Function", meta=(Hidden))
+	FName Function;
+
+	/** Internal event name derived from Function above */
+	UPROPERTY(VisibleAnywhere, Category = "Call Function", meta=(Hidden))
+	FName FunctionEvent;
+
+	/** The call site to use */
+	UPROPERTY(EditAnywhere, Category = "Call Function")
+	EAnimNextCallFunctionCallSite CallSite = EAnimNextCallFunctionCallSite::BecomeRelevant;
+
+	// Latent pin support boilerplate
+	#define TRAIT_LATENT_PROPERTIES_ENUMERATOR(GeneratorMacro) \
+		GeneratorMacro(CallSite) \
+
+	GENERATE_TRAIT_LATENT_PROPERTIES(FAnimNextCallFunctionSharedData, TRAIT_LATENT_PROPERTIES_ENUMERATOR)
+	#undef TRAIT_LATENT_PROPERTIES_ENUMERATOR
+};
+
+namespace UE::AnimNext
+{
+/**
+ * FCallFunctionTrait
+ *
+ * A trait that calls a function at the specified update point
+ */
+struct FCallFunctionTrait : FAdditiveTrait, IUpdate
+{
+	DECLARE_ANIM_TRAIT(FCallFunctionTrait, 0x76ef7c0f, FAdditiveTrait)
+
+	using FSharedData = FAnimNextCallFunctionSharedData;
+
+	struct FInstanceData : FTrait::FInstanceData
+	{
+	};
+
+	// IUpdate impl
+	virtual void OnBecomeRelevant(FUpdateTraversalContext& Context, const TTraitBinding<IUpdate>& Binding, const FTraitUpdateState& TraitState) const override;
+	virtual void PreUpdate(FUpdateTraversalContext& Context, const TTraitBinding<IUpdate>& Binding, const FTraitUpdateState& TraitState) const override;
+	virtual void PostUpdate(FUpdateTraversalContext& Context, const TTraitBinding<IUpdate>& Binding, const FTraitUpdateState& TraitState) const override;
+
+	// Calls the RigVM function that is assigned to us
+	void CallFunctionForMatchingSite(const TTraitBinding<IUpdate>& InBinding, const FTraitUpdateState& InTraitState, EAnimNextCallFunctionCallSite InCallSite) const;
+};
+}
diff --git a/Engine/Plugins/Experimental/AnimNextAnimGraph/Source/AnimNextAnimGraph/Private/Graph/AnimNextAnimationGraph.cpp b/Engine/Plugins/Experimental/AnimNextAnimGraph/Source/AnimNextAnimGraph/Private/Graph/AnimNextAnimationGraph.cpp
--- a/Engine/Plugins/Experimental/AnimNextAnimGraph/Source/AnimNextAnimGraph/Private/Graph/AnimNextAnimationGraph.cpp
+++ b/Engine/Plugins/Experimental/AnimNextAnimGraph/Source/AnimNextAnimGraph/Private/Graph/AnimNextAnimationGraph.cpp
@@ -166,7 +166,7 @@ bool UAnimNextAnimationGraph::LoadFromArchiveBuffer(const TArray<uint8>& InShare
 
 	// Reconstruct our graph shared data
 	FMemoryReader GraphSharedDataArchive(InSharedDataArchiveBuffer);
-	FTraitReader TraitReader(GraphReferencedObjects, GraphSharedDataArchive);
+	FTraitReader TraitReader(GraphReferencedObjects, GraphReferencedSoftObjects, GraphSharedDataArchive);
 
 	const FTraitReader::EErrorState ErrorState = TraitReader.ReadGraph(SharedDataBuffer);
 	if (ErrorState == FTraitReader::EErrorState::None)
diff --git a/Engine/Plugins/Experimental/AnimNextAnimGraph/Source/AnimNextAnimGraph/Private/Graph/AnimNextGraphInstance.cpp b/Engine/Plugins/Experimental/AnimNextAnimGraph/Source/AnimNextAnimGraph/Private/Graph/AnimNextGraphInstance.cpp
--- a/Engine/Plugins/Experimental/AnimNextAnimGraph/Source/AnimNextAnimGraph/Private/Graph/AnimNextGraphInstance.cpp
+++ b/Engine/Plugins/Experimental/AnimNextAnimGraph/Source/AnimNextAnimGraph/Private/Graph/AnimNextGraphInstance.cpp
@@ -7,6 +7,7 @@
 #include "DataInterface/DataInterfaceStructAdapter.h"
 #include "Graph/AnimNextAnimationGraph.h"
 #include "Graph/AnimNextGraphContextData.h"
+#include "Graph/AnimNextGraphLatentPropertiesContextData.h"
 #include "TraitCore/ExecutionContext.h"
 #include "Graph/GC_GraphInstanceComponent.h"
 #include "Graph/RigUnit_AnimNextShimRoot.h"
@@ -181,7 +182,7 @@ void FAnimNextGraphInstance::ExecuteLatentPins(const TConstArrayView<UE::AnimNex
 		FAnimNextExecuteContext& AnimNextContext = ExtendedExecuteContext.GetPublicDataSafe<FAnimNextExecuteContext>();
 
 		// Insert our context data for the scope of execution
-		FAnimNextGraphContextData ContextData(ModuleInstance, this, LatentHandles, DestinationBasePtr, bIsFrozen);
+		FAnimNextGraphLatentPropertiesContextData ContextData(ModuleInstance, this, LatentHandles, DestinationBasePtr, bIsFrozen);
 		UE::AnimNext::FScopedExecuteContextData ContextDataScope(AnimNextContext, ContextData);
 
 		VM->ExecuteVM(ExtendedExecuteContext, FRigUnit_AnimNextShimRoot::EventName);
diff --git a/Engine/Plugins/Experimental/AnimNextAnimGraph/Source/AnimNextAnimGraph/Private/Graph/RigUnit_AnimNextGraphEvaluator.cpp b/Engine/Plugins/Experimental/AnimNextAnimGraph/Source/AnimNextAnimGraph/Private/Graph/RigUnit_AnimNextGraphEvaluator.cpp
--- a/Engine/Plugins/Experimental/AnimNextAnimGraph/Source/AnimNextAnimGraph/Private/Graph/RigUnit_AnimNextGraphEvaluator.cpp
+++ b/Engine/Plugins/Experimental/AnimNextAnimGraph/Source/AnimNextAnimGraph/Private/Graph/RigUnit_AnimNextGraphEvaluator.cpp
@@ -2,6 +2,7 @@
 
 #include "Graph/RigUnit_AnimNextGraphEvaluator.h"
 #include "Graph/AnimNextGraphContextData.h"
+#include "Graph/AnimNextGraphLatentPropertiesContextData.h"
 #include "TraitCore/LatentPropertyHandle.h"
 
 #include UE_INLINE_GENERATED_CPP_BY_NAME(RigUnit_AnimNextGraphEvaluator)
@@ -27,11 +28,11 @@ namespace UE::AnimNext::Private
 void FRigUnit_AnimNextGraphEvaluator::StaticExecute(FRigVMExtendedExecuteContext& RigVMExecuteContext, FRigVMMemoryHandleArray RigVMMemoryHandles, FRigVMPredicateBranchArray RigVMBranches)
 {
 	const FAnimNextExecuteContext& VMExecuteContext = RigVMExecuteContext.GetPublicData<FAnimNextExecuteContext>();
-	const FAnimNextGraphContextData& GraphContextData = VMExecuteContext.GetContextData<FAnimNextGraphContextData>();
+	const FAnimNextGraphLatentPropertiesContextData& LatentPropertiesContextData = VMExecuteContext.GetContextData<FAnimNextGraphLatentPropertiesContextData>();
 
-	const TConstArrayView<UE::AnimNext::FLatentPropertyHandle>& LatentHandles = GraphContextData.GetLatentHandles();
-	uint8* DestinationBasePtr = (uint8*)GraphContextData.GetDestinationBasePtr();
-	const bool bIsFrozen = GraphContextData.IsFrozen();
+	const TConstArrayView<UE::AnimNext::FLatentPropertyHandle>& LatentHandles = LatentPropertiesContextData.GetLatentHandles();
+	uint8* DestinationBasePtr = (uint8*)LatentPropertiesContextData.GetDestinationBasePtr();
+	const bool bIsFrozen = LatentPropertiesContextData.IsFrozen();
 
 	for (UE::AnimNext::FLatentPropertyHandle Handle : LatentHandles)
 	{
diff --git a/Engine/Plugins/Experimental/AnimNextAnimGraph/Source/AnimNextAnimGraph/Private/Graph/RigUnit_AnimNextRunAnimationGraph.cpp b/Engine/Plugins/Experimental/AnimNextAnimGraph/Source/AnimNextAnimGraph/Private/Graph/RigUnit_AnimNextRunAnimationGraph.cpp
--- a/Engine/Plugins/Experimental/AnimNextAnimGraph/Source/AnimNextAnimGraph/Private/Graph/RigUnit_AnimNextRunAnimationGraph.cpp
+++ b/Engine/Plugins/Experimental/AnimNextAnimGraph/Source/AnimNextAnimGraph/Private/Graph/RigUnit_AnimNextRunAnimationGraph.cpp
@@ -77,6 +77,10 @@ FRigUnit_AnimNextRunAnimationGraph_Execute()
 		Instance.BindPublicVariables(Hosts);
 	}
 
+	// Propagate delta time
+	FAnimNextExecuteContext& AnimNextExecuteContext = Instance.GetExtendedExecuteContext().GetPublicDataSafe<FAnimNextExecuteContext>();
+	AnimNextExecuteContext.SetDeltaTime(ExecuteContext.GetDeltaTime());
+
 	const UE::AnimNext::FReferencePose& RefPose = ReferencePose.ReferencePose.GetRef<UE::AnimNext::FReferencePose>();
 
 	// TODO: Currently forcing additive flag to false here
diff --git a/Engine/Plugins/Experimental/AnimNextAnimGraph/Source/AnimNextAnimGraph/Private/TraitCore/TraitReader.cpp b/Engine/Plugins/Experimental/AnimNextAnimGraph/Source/AnimNextAnimGraph/Private/TraitCore/TraitReader.cpp
--- a/Engine/Plugins/Experimental/AnimNextAnimGraph/Source/AnimNextAnimGraph/Private/TraitCore/TraitReader.cpp
+++ b/Engine/Plugins/Experimental/AnimNextAnimGraph/Source/AnimNextAnimGraph/Private/TraitCore/TraitReader.cpp
@@ -12,9 +12,10 @@
 
 namespace UE::AnimNext
 {
-	FTraitReader::FTraitReader(const TArray<TObjectPtr<UObject>>& InGraphReferencedObjects, FArchive& Ar)
+	FTraitReader::FTraitReader(const TArray<TObjectPtr<UObject>>& InGraphReferencedObjects, const TArray<FSoftObjectPath>& InGraphReferencedSoftObjects, FArchive& Ar)
 		: FArchiveProxy(Ar)
 		, GraphReferencedObjects(InGraphReferencedObjects)
+		, GraphReferencedSoftObjects(InGraphReferencedSoftObjects)
 	{
 	}
 
@@ -227,6 +228,44 @@ namespace UE::AnimNext
 		return FArchiveUObject::SerializeObjectPtr(*this, Obj);
 	}
 
+	FArchive& FTraitReader::operator<<(FSoftObjectPath& Value)
+	{
+		// Load the object index
+		int32 SoftObjectIndex = INDEX_NONE;
+		*this << SoftObjectIndex;
+
+		if (ensure(GraphReferencedSoftObjects.IsValidIndex(SoftObjectIndex)))
+		{
+			Value = GraphReferencedSoftObjects[SoftObjectIndex];
+		}
+		else
+		{
+			// Something went wrong, the reference list must have gotten out of sync which shouldn't happen
+			Value.Reset();
+		}
+
+		return *this;
+	}
+
+	FArchive& FTraitReader::operator<<(FSoftObjectPtr& Value)
+	{
+		// Load the object index
+		int32 SoftObjectIndex = INDEX_NONE;
+		*this << SoftObjectIndex;
+
+		if (ensure(GraphReferencedSoftObjects.IsValidIndex(SoftObjectIndex)))
+		{
+			Value = FSoftObjectPtr(GraphReferencedSoftObjects[SoftObjectIndex]);
+		}
+		else
+		{
+			// Something went wrong, the reference list must have gotten out of sync which shouldn't happen
+			Value.Reset();
+		}
+
+		return *this;
+	}
+
 	FNodeHandle FTraitReader::ResolveNodeHandle(FNodeHandle NodeHandle) const
 	{
 		if (!NodeHandle.IsValid())
diff --git a/Engine/Plugins/Experimental/AnimNextAnimGraph/Source/AnimNextAnimGraph/Private/TraitCore/TraitWriter.cpp b/Engine/Plugins/Experimental/AnimNextAnimGraph/Source/AnimNextAnimGraph/Private/TraitCore/TraitWriter.cpp
--- a/Engine/Plugins/Experimental/AnimNextAnimGraph/Source/AnimNextAnimGraph/Private/TraitCore/TraitWriter.cpp
+++ b/Engine/Plugins/Experimental/AnimNextAnimGraph/Source/AnimNextAnimGraph/Private/TraitCore/TraitWriter.cpp
@@ -75,6 +75,7 @@ namespace UE::AnimNext
 
 		bIsNodeWriting = true;
 		GraphReferencedObjects.Reset();
+		GraphReferencedSoftObjects.Reset();
 
 		// Serialize the node templates
 		TArray<FNodeTemplateRegistryHandle> NodeTemplateHandles;
@@ -229,6 +230,11 @@ namespace UE::AnimNext
 		return GraphReferencedObjects;
 	}
 
+	const TArray<FSoftObjectPath>& FTraitWriter::GetGraphReferencedSoftObjects() const
+	{
+		return GraphReferencedSoftObjects;
+	}
+
 	FArchive& FTraitWriter::operator<<(UObject*& Obj)
 	{
 		// Add our object for tracking
@@ -244,5 +250,27 @@ namespace UE::AnimNext
 	{
 		return FArchiveUObject::SerializeObjectPtr(*this, Obj);
 	}
+
+	FArchive& FTraitWriter::operator<<(FSoftObjectPath& Value)
+	{
+		// Add our object for tracking
+		int32 SoftObjectIndex = GraphReferencedSoftObjects.AddUnique(Value);
+
+		// Save our index, we'll use it to resolve the soft object on load
+		*this << SoftObjectIndex;
+
+		return *this;
+	}
+
+	FArchive& FTraitWriter::operator<<(FSoftObjectPtr& Value)
+	{
+		// Add our object for tracking
+		int32 SoftObjectIndex = GraphReferencedSoftObjects.AddUnique(Value.ToSoftObjectPath());
+
+		// Save our index, we'll use it to resolve the soft object on load
+		*this << SoftObjectIndex;
+
+		return *this;
+	}
 }
 #endif
diff --git a/Engine/Plugins/Experimental/AnimNextAnimGraph/Source/AnimNextAnimGraph/Private/Traits/CallFunction.cpp b/Engine/Plugins/Experimental/AnimNextAnimGraph/Source/AnimNextAnimGraph/Private/Traits/CallFunction.cpp
new file mode 100644
--- /dev/null
+++ b/Engine/Plugins/Experimental/AnimNextAnimGraph/Source/AnimNextAnimGraph/Private/Traits/CallFunction.cpp
@@ -0,0 +1,80 @@
+﻿// Copyright Epic Games, Inc. All Rights Reserved.
+
+#include "Traits/CallFunction.h"
+
+#include "Graph/AnimNextAnimationGraph.h"
+#include "Graph/AnimNextGraphContextData.h"
+#include "Graph/AnimNextGraphInstance.h"
+#include "TraitCore/ExecutionContext.h"
+#include "TraitCore/NodeInstance.h"
+#include "AnimNextAnimGraphStats.h"
+
+DEFINE_STAT(STAT_AnimNext_CallFunction);
+
+namespace UE::AnimNext
+{
+	AUTO_REGISTER_ANIM_TRAIT(FCallFunctionTrait)
+
+	// Trait implementation boilerplate
+	#define TRAIT_INTERFACE_ENUMERATOR(GeneratorMacro) \
+		GeneratorMacro(IUpdate)
+
+	GENERATE_ANIM_TRAIT_IMPLEMENTATION(FCallFunctionTrait, TRAIT_INTERFACE_ENUMERATOR, NULL_ANIM_TRAIT_INTERFACE_ENUMERATOR, NULL_ANIM_TRAIT_EVENT_ENUMERATOR)
+
+	#undef TRAIT_INTERFACE_ENUMERATOR
+
+	void FCallFunctionTrait::OnBecomeRelevant(FUpdateTraversalContext& Context, const TTraitBinding<IUpdate>& Binding, const FTraitUpdateState& TraitState) const
+	{
+		CallFunctionForMatchingSite(Binding, TraitState, EAnimNextCallFunctionCallSite::BecomeRelevant);
+		IUpdate::OnBecomeRelevant(Context, Binding, TraitState);
+	}
+
+	void FCallFunctionTrait::PreUpdate(FUpdateTraversalContext& Context, const TTraitBinding<IUpdate>& Binding, const FTraitUpdateState& TraitState) const
+	{
+		CallFunctionForMatchingSite(Binding,  TraitState, EAnimNextCallFunctionCallSite::PreUpdate);
+		IUpdate::PreUpdate(Context, Binding, TraitState);
+	}
+
+	void FCallFunctionTrait::PostUpdate(FUpdateTraversalContext& Context, const TTraitBinding<IUpdate>& Binding, const FTraitUpdateState& TraitState) const
+	{
+		IUpdate::PostUpdate(Context, Binding, TraitState);
+		CallFunctionForMatchingSite(Binding, TraitState,  EAnimNextCallFunctionCallSite::PostUpdate);
+	}
+
+	void FCallFunctionTrait::CallFunctionForMatchingSite(const TTraitBinding<IUpdate>& InBinding, const FTraitUpdateState& InTraitState, EAnimNextCallFunctionCallSite InCallSite) const
+	{
+		const FSharedData& SharedData = *InBinding.GetSharedData<FSharedData>(); 
+		if(SharedData.CallSite != InCallSite || SharedData.FunctionEvent.IsNone())
+		{
+			return;
+		}
+
+		SCOPE_CYCLE_COUNTER(STAT_AnimNext_CallFunction);
+
+		FNodeInstance* ParentNodeInstance = InBinding.GetTraitPtr().GetNodeInstance();
+		FAnimNextGraphInstance& GraphInstance = ParentNodeInstance->GetOwner();
+		const UAnimNextAnimationGraph* AnimationGraph = GraphInstance.GetAnimationGraph();
+		if(AnimationGraph == nullptr)
+		{
+			return;
+		}
+
+		// TODO: const_cast: Can ExecuteVM be made const?
+		URigVM* VM = const_cast<URigVM*>(AnimationGraph->GetRigVM());
+		if (VM == nullptr)
+		{
+			return;
+		}
+
+		FRigVMExtendedExecuteContext& ExtendedExecuteContext = GraphInstance.GetExtendedExecuteContext();
+		FAnimNextExecuteContext& AnimNextContext = GraphInstance.GetExtendedExecuteContext().GetPublicDataSafe<FAnimNextExecuteContext>();
+
+		// Propagate delta time
+		AnimNextContext.SetDeltaTime(InTraitState.GetDeltaTime());
+
+		FAnimNextGraphContextData ContextData(GraphInstance.GetModuleInstance(), &GraphInstance);
+		UE::AnimNext::FScopedExecuteContextData ContextDataScope(AnimNextContext, ContextData);
+
+		VM->ExecuteVM(ExtendedExecuteContext, SharedData.FunctionEvent);
+	}
+}
diff --git a/Engine/Plugins/Experimental/AnimNextAnimGraph/Source/AnimNextAnimGraph/Public/AnimNextAnimGraphStats.h b/Engine/Plugins/Experimental/AnimNextAnimGraph/Source/AnimNextAnimGraph/Public/AnimNextAnimGraphStats.h
--- a/Engine/Plugins/Experimental/AnimNextAnimGraph/Source/AnimNextAnimGraph/Public/AnimNextAnimGraphStats.h
+++ b/Engine/Plugins/Experimental/AnimNextAnimGraph/Source/AnimNextAnimGraph/Public/AnimNextAnimGraphStats.h
@@ -14,3 +14,4 @@ DECLARE_CYCLE_STAT_EXTERN(TEXT("AnimNext: Allocate Graph Instance"), STAT_AnimNe
 DECLARE_CYCLE_STAT_EXTERN(TEXT("AnimNext: Graph RigVM"), STAT_AnimNext_Graph_RigVM, STATGROUP_AnimNext, ANIMNEXTANIMGRAPH_API);
 DECLARE_CYCLE_STAT_EXTERN(TEXT("AnimNext: Execute Evaluation Program"), STAT_AnimNext_EvaluationProgram_Execute, STATGROUP_AnimNext, ANIMNEXTANIMGRAPH_API);
 
+DECLARE_CYCLE_STAT_EXTERN(TEXT("AnimNext: Call Function"), STAT_AnimNext_CallFunction, STATGROUP_AnimNext, ANIMNEXTANIMGRAPH_API);
diff --git a/Engine/Plugins/Experimental/AnimNextAnimGraph/Source/AnimNextAnimGraphEditor/Private/AnimNextAnimGraphEditorModule.cpp b/Engine/Plugins/Experimental/AnimNextAnimGraph/Source/AnimNextAnimGraphEditor/Private/AnimNextAnimGraphEditorModule.cpp
--- a/Engine/Plugins/Experimental/AnimNextAnimGraph/Source/AnimNextAnimGraphEditor/Private/AnimNextAnimGraphEditorModule.cpp
+++ b/Engine/Plugins/Experimental/AnimNextAnimGraph/Source/AnimNextAnimGraphEditor/Private/AnimNextAnimGraphEditorModule.cpp
@@ -8,6 +8,7 @@
 #include "EdGraphUtilities.h"
 #include "IAnimNextEditorModule.h"
 #include "IWorkspaceEditorModule.h"
+#include "PropertyEditorModule.h"
 #include "TraitStackEditor.h"
 #include "UncookedOnlyUtils.h"
 #include "Entries/AnimNextAnimationGraphEntry.h"
@@ -21,6 +22,8 @@
 #include "Graph/RigUnit_AnimNextGraphRoot.h"
 #include "Graph/RigVMTrait_AnimNextPublicVariablesEditor.h"
 #include "Graph/TraitEditorTabSummoner.h"
+#include "Traits/AnimNextCallFunctionSharedDataDetails.h"
+#include "Traits/CallFunction.h"
 #include "WorkflowOrientedApp/WorkflowTabManager.h"
 
 #define LOCTEXT_NAMESPACE "FAnimNextAnimGraphEditorModule"
@@ -41,6 +44,11 @@ class FAnimNextAnimGraphEditorModule : public IModuleInterface
 
 		FPublicVariablesImpl::Register();
 
+		FPropertyEditorModule& PropertyModule = FModuleManager::LoadModuleChecked<FPropertyEditorModule>("PropertyEditor");
+		PropertyModule.RegisterCustomPropertyTypeLayout(
+			FAnimNextCallFunctionSharedData::StaticStruct()->GetFName(),
+			FOnGetPropertyTypeCustomizationInstance::CreateLambda([] { return MakeShared<FCallFunctionSharedDataDetails>(); }));
+
 		Workspace::IWorkspaceEditorModule& WorkspaceEditorModule = FModuleManager::Get().LoadModuleChecked<Workspace::IWorkspaceEditorModule>("WorkspaceEditor");
 
 		// --- AnimNextAnimationGraph ---
@@ -102,7 +110,13 @@ class FAnimNextAnimGraphEditorModule : public IModuleInterface
 				WorkspaceEditorModule.UnregisterWorkspaceItemDetails(Workspace::FOutlinerItemDetailsId(FAnimNextAnimationGraphOutlinerData::StaticStruct()->GetFName()));
 			}
 		}
-		
+
+		if(FModuleManager::Get().IsModuleLoaded("PropertyEditor"))
+		{
+			FPropertyEditorModule& PropertyModule = FModuleManager::LoadModuleChecked<FPropertyEditorModule>("PropertyEditor");
+			PropertyModule.UnregisterCustomClassLayout("AnimNextCallFunctionSharedData");
+		}
+
 		IModularFeatures::Get().UnregisterModularFeature(ITraitStackEditor::ModularFeatureName, TraitStackEditor.Get());
 		TraitStackEditor.Reset();
 
diff --git a/Engine/Plugins/Experimental/AnimNextAnimGraph/Source/AnimNextAnimGraphEditor/Private/Traits/AnimNextCallFunctionSharedDataDetails.cpp b/Engine/Plugins/Experimental/AnimNextAnimGraph/Source/AnimNextAnimGraphEditor/Private/Traits/AnimNextCallFunctionSharedDataDetails.cpp
new file mode 100644
--- /dev/null
+++ b/Engine/Plugins/Experimental/AnimNextAnimGraph/Source/AnimNextAnimGraphEditor/Private/Traits/AnimNextCallFunctionSharedDataDetails.cpp
@@ -0,0 +1,218 @@
+﻿// Copyright Epic Games, Inc. All Rights Reserved.
+
+#include "AnimNextCallFunctionSharedDataDetails.h"
+
+#include "DetailLayoutBuilder.h"
+#include "DetailWidgetRow.h"
+#include "IDetailChildrenBuilder.h"
+#include "IDetailPropertyRow.h"
+#include "IPropertyUtilities.h"
+#include "PropertyHandle.h"
+#include "UncookedOnlyUtils.h"
+#include "Traits/CallFunction.h"
+#include "Common/SRigVMFunctionPicker.h"
+#include "AnimNextRigVMAsset.h"
+#include "AnimNextRigVMAssetEditorData.h"
+#include "IWorkspaceEditor.h"
+#include "IWorkspaceEditorModule.h"
+#include "ScopedTransaction.h"
+#include "EdGraph/RigVMEdGraphNode.h"
+#include "Modules/ModuleManager.h"
+#include "RigVMModel/RigVMController.h"
+
+#define LOCTEXT_NAMESPACE "FCallFunctionSharedDataDetails"
+
+namespace UE::AnimNext::Editor
+{
+
+void FCallFunctionSharedDataDetails::CustomizeHeader(TSharedRef<IPropertyHandle> InPropertyHandle, FDetailWidgetRow& InHeaderRow, IPropertyTypeCustomizationUtils& InCustomizationUtils)
+{
+	InHeaderRow
+	.NameContent()
+	[
+		InPropertyHandle->CreatePropertyNameWidget()
+	]
+	.ValueContent()
+	[
+		InPropertyHandle->CreatePropertyValueWidget()
+	];
+}
+
+void FCallFunctionSharedDataDetails::CustomizeChildren(TSharedRef<IPropertyHandle> InPropertyHandle, IDetailChildrenBuilder& InChildBuilder, IPropertyTypeCustomizationUtils& InCustomizationUtils)
+{
+	PropertyHandle = InPropertyHandle;
+	FunctionPropertyHandle = InPropertyHandle->GetChildHandle(GET_MEMBER_NAME_CHECKED(FAnimNextCallFunctionSharedData, Function));
+	IDetailPropertyRow& PropertyRow = InChildBuilder.AddProperty(FunctionPropertyHandle.ToSharedRef());
+
+	TSharedPtr<IPropertyHandle> CallSitePropertyHandle = InPropertyHandle->GetChildHandle(GET_MEMBER_NAME_CHECKED(FAnimNextCallFunctionSharedData, CallSite));
+	InChildBuilder.AddProperty(CallSitePropertyHandle.ToSharedRef());
+
+	FunctionHeaderPropertyHandle = InPropertyHandle->GetChildHandle(GET_MEMBER_NAME_CHECKED(FAnimNextCallFunctionSharedData, FunctionHeader));
+	InChildBuilder.AddProperty(FunctionHeaderPropertyHandle.ToSharedRef())
+	.Visibility(EVisibility::Collapsed);
+
+	FunctionEventPropertyHandle = InPropertyHandle->GetChildHandle(GET_MEMBER_NAME_CHECKED(FAnimNextCallFunctionSharedData, FunctionEvent));
+	InChildBuilder.AddProperty(FunctionEventPropertyHandle.ToSharedRef())
+	.Visibility(EVisibility::Collapsed);
+
+	TArray<TWeakObjectPtr<UObject>> Objects = InCustomizationUtils.GetPropertyUtilities()->GetSelectedObjects();
+	URigVMEdGraphNode* EdGraphNode = Objects.Num() ? Cast<URigVMEdGraphNode>(Objects[0].Get()) : nullptr;
+	UAnimNextRigVMAsset* CurrentAsset = nullptr;
+	URigVMController* Controller = nullptr;
+	if(EdGraphNode != nullptr)
+	{
+		CurrentAsset = EdGraphNode->GetTypedOuter<UAnimNextRigVMAsset>();
+		Controller = EdGraphNode->GetController();
+	}
+
+	auto OnFunctionPicked = [this, WeakController = TWeakObjectPtr<URigVMController>(Controller)](const FRigVMGraphFunctionHeader& InFunctionHeader)
+	{
+		URigVMController* Controller = WeakController.Get();
+		if(Controller == nullptr)
+		{
+			return;
+		}
+
+		FRigVMControllerCompileBracketScope CompileScope(Controller);
+		FScopedTransaction Transaction(LOCTEXT("SetFunctionTransaction", "Set Function"));
+
+		// Update function event & header from function name
+		FunctionPropertyHandle->SetValue(InFunctionHeader.IsValid() ? InFunctionHeader.Name : NAME_None);
+		FunctionEventPropertyHandle->SetValue(InFunctionHeader.IsValid() ? FName(*UncookedOnly::FUtils::MakeFunctionWrapperEventName(InFunctionHeader.Name)): NAME_None);
+
+		FunctionHeaderPropertyHandle->NotifyPreChange();
+		FunctionHeaderPropertyHandle->EnumerateRawData([&InFunctionHeader](void* InRawData, const int32 InDataIndex, const int32 InNumDatas)
+		{
+			FRigVMGraphFunctionHeader* FunctionHeader = static_cast<FRigVMGraphFunctionHeader*>(InRawData);
+			if(InFunctionHeader.IsValid())
+			{
+				*FunctionHeader = InFunctionHeader;
+			}
+			else
+			{
+				*FunctionHeader = FRigVMGraphFunctionHeader();
+			}
+			return true;
+		});
+		FunctionHeaderPropertyHandle->NotifyPostChange(EPropertyChangeType::ValueSet);
+		FunctionHeaderPropertyHandle->NotifyFinishedChangingProperties();
+	};
+
+	PropertyRow.CustomWidget()
+	.NameContent()
+	[
+		FunctionPropertyHandle->CreatePropertyNameWidget()
+	]
+	.ValueContent()
+	[
+		SNew(SRigVMFunctionPicker)
+		.CurrentAsset(FAssetData(CurrentAsset))
+		.FunctionName_Lambda([this]()
+		{
+			bool bMultipleValues = false;
+			TOptional<FRigVMGraphFunctionHeader> Header;
+			PropertyHandle->EnumerateConstRawData([&Header, &bMultipleValues](const void* InRawData, const int32 InDataIndex, const int32 InNumDatas)
+			{
+				const FAnimNextCallFunctionSharedData* SharedData = static_cast<const FAnimNextCallFunctionSharedData*>(InRawData);
+				if(!Header.IsSet())
+				{
+					Header = SharedData->FunctionHeader;
+				}
+				else if(Header.GetValue() != SharedData->FunctionHeader)
+				{
+					Header = FRigVMGraphFunctionHeader();
+					bMultipleValues = true;
+					return false;
+				}
+				return true;
+			});
+
+			if(bMultipleValues)
+			{
+				return LOCTEXT("MultipleValuesLabel", "Multiple Values");
+			}
+			else if(Header.IsSet() && Header.GetValue().IsValid())
+			{
+				return FText::FromName(Header.GetValue().Name);
+			}
+			return LOCTEXT("NoFunctionSelectedLabel", "None");
+		})
+		.FunctionToolTip_Lambda([this]()
+		{
+			bool bMultipleValues = false;
+			TOptional<FRigVMGraphFunctionHeader> Header;
+			PropertyHandle->EnumerateConstRawData([&Header, &bMultipleValues](const void* InRawData, const int32 InDataIndex, const int32 InNumDatas)
+			{
+				const FAnimNextCallFunctionSharedData* SharedData = static_cast<const FAnimNextCallFunctionSharedData*>(InRawData);
+				if(!Header.IsSet())
+				{
+					Header = SharedData->FunctionHeader;
+				}
+				else if(Header.GetValue() != SharedData->FunctionHeader)
+				{
+					Header = FRigVMGraphFunctionHeader();
+					bMultipleValues = true;
+					return false;
+				}
+				return true;
+			});
+
+			if(bMultipleValues)
+			{
+				return LOCTEXT("MultipleValuesLabel", "Multiple Values");
+			}
+			else if(Header.IsSet() && Header.GetValue().IsValid())
+			{
+				static const FTextFormat ToolTipFormat(LOCTEXT("FunctionToolTipFormat", "{0}\n{1}"));
+				
+				return Header.GetValue().Description.Len() > 0 ?
+					FText::FromString(Header.GetValue().Description) :
+					FText::Format(ToolTipFormat, FText::FromString(Header.GetValue().LibraryPointer.GetFunctionName()), FText::FromString(Header.GetValue().LibraryPointer.GetLibraryNodePath()));
+			}
+			return LOCTEXT("NoFunctionSelectedLabel", "None");
+		})
+		.OnRigVMFunctionPicked_Lambda(OnFunctionPicked)
+		.OnNewFunction_Lambda([this, CurrentAsset, WeakController = TWeakObjectPtr<URigVMController>(Controller), OnFunctionPicked]()
+		{
+			URigVMController* Controller = WeakController.Get();
+			if(Controller == nullptr)
+			{
+				return;
+			}
+
+			UAnimNextRigVMAssetEditorData* EditorData = UncookedOnly::FUtils::GetEditorData<UAnimNextRigVMAssetEditorData>(CurrentAsset);
+			URigVMLibraryNode* NewFunction = nullptr;
+			{
+				FRigVMControllerCompileBracketScope CompileScope(Controller);
+				FScopedTransaction Transaction(LOCTEXT("AddFunctionTransaction", "Add Function"));
+				NewFunction = EditorData->AddFunction(TEXT("NewFunction"), true);
+
+				if(NewFunction == nullptr)
+				{
+					return;
+				}
+
+				const FRigVMGraphFunctionData* FunctionData = EditorData->GraphFunctionStore.FindFunction(NewFunction->GetFunctionIdentifier());
+				if(FunctionData == nullptr)
+				{
+					return;
+				}
+
+				// Set our function
+				OnFunctionPicked(FunctionData->Header);
+			}
+
+			// Open the new function's graph
+			UE::Workspace::IWorkspaceEditorModule& WorkspaceEditorModule = FModuleManager::LoadModuleChecked<UE::Workspace::IWorkspaceEditorModule>("WorkspaceEditor");
+			if(UE::Workspace::IWorkspaceEditor* WorkspaceEditor = WorkspaceEditorModule.OpenWorkspaceForObject(CurrentAsset, UE::Workspace::EOpenWorkspaceMethod::Default))
+			{
+				UObject* EditorObject = EditorData->GetEditorObjectForRigVMGraph(NewFunction->GetContainedGraph());
+				WorkspaceEditor->OpenObjects({EditorObject});
+			}
+		})
+	];
+}
+
+}
+
+#undef LOCTEXT_NAMESPACE
\ No newline at end of file
diff --git a/Engine/Plugins/Experimental/AnimNextAnimGraph/Source/AnimNextAnimGraphEditor/Private/Traits/AnimNextCallFunctionSharedDataDetails.h b/Engine/Plugins/Experimental/AnimNextAnimGraph/Source/AnimNextAnimGraphEditor/Private/Traits/AnimNextCallFunctionSharedDataDetails.h
new file mode 100644
--- /dev/null
+++ b/Engine/Plugins/Experimental/AnimNextAnimGraph/Source/AnimNextAnimGraphEditor/Private/Traits/AnimNextCallFunctionSharedDataDetails.h
@@ -0,0 +1,22 @@
+﻿// Copyright Epic Games, Inc. All Rights Reserved.
+
+#pragma once
+
+#include "IPropertyTypeCustomization.h"
+
+namespace UE::AnimNext::Editor
+{
+
+class FCallFunctionSharedDataDetails : public IPropertyTypeCustomization
+{
+	// IPropertyTypeCustomization interface
+	virtual void CustomizeHeader(TSharedRef<IPropertyHandle> InPropertyHandle, FDetailWidgetRow& InHeaderRow, IPropertyTypeCustomizationUtils& InCustomizationUtils) override;
+	virtual void CustomizeChildren(TSharedRef<IPropertyHandle> InPropertyHandle, IDetailChildrenBuilder& InChildBuilder, IPropertyTypeCustomizationUtils& InCustomizationUtils) override;
+
+	TSharedPtr<IPropertyHandle> PropertyHandle;
+	TSharedPtr<IPropertyHandle> FunctionPropertyHandle;
+	TSharedPtr<IPropertyHandle> FunctionHeaderPropertyHandle;
+	TSharedPtr<IPropertyHandle> FunctionEventPropertyHandle;
+};
+
+}
\ No newline at end of file
diff --git a/Engine/Plugins/Experimental/AnimNextAnimGraph/Source/AnimNextAnimGraphTestSuite/Private/AnimNextAnimGraphTraitEventTest.cpp b/Engine/Plugins/Experimental/AnimNextAnimGraph/Source/AnimNextAnimGraphTestSuite/Private/AnimNextAnimGraphTraitEventTest.cpp
--- a/Engine/Plugins/Experimental/AnimNextAnimGraph/Source/AnimNextAnimGraphTestSuite/Private/AnimNextAnimGraphTraitEventTest.cpp
+++ b/Engine/Plugins/Experimental/AnimNextAnimGraph/Source/AnimNextAnimGraphTestSuite/Private/AnimNextAnimGraphTraitEventTest.cpp
@@ -228,6 +228,7 @@ bool FAnimationAnimNextRuntimeTest_GraphTraitEvent::RunTest(const FString& InPar
 		// Write our graph
 		TArray<uint8> GraphSharedDataArchiveBuffer;
 		TArray<TObjectPtr<UObject>> GraphReferencedObjects;
+		TArray<FSoftObjectPath> GraphReferencedSoftObjects;
 		{
 			FTraitWriter TraitWriter;
 
@@ -293,6 +294,7 @@ bool FAnimationAnimNextRuntimeTest_GraphTraitEvent::RunTest(const FString& InPar
 			AddErrorIfFalse(TraitWriter.GetErrorState() == FTraitWriter::EErrorState::None, "FAnimationAnimNextRuntimeTest_GraphTraitEvent -> Failed to write traits");
 			GraphSharedDataArchiveBuffer = TraitWriter.GetGraphSharedData();
 			GraphReferencedObjects = TraitWriter.GetGraphReferencedObjects();
+			GraphReferencedSoftObjects = TraitWriter.GetGraphReferencedSoftObjects();
 		}
 
 		// Read our graph
diff --git a/Engine/Plugins/Experimental/AnimNextAnimGraph/Source/AnimNextAnimGraphTestSuite/Private/AnimNextAnimGraphTraitInterfacesTest.cpp b/Engine/Plugins/Experimental/AnimNextAnimGraph/Source/AnimNextAnimGraphTestSuite/Private/AnimNextAnimGraphTraitInterfacesTest.cpp
--- a/Engine/Plugins/Experimental/AnimNextAnimGraph/Source/AnimNextAnimGraphTestSuite/Private/AnimNextAnimGraphTraitInterfacesTest.cpp
+++ b/Engine/Plugins/Experimental/AnimNextAnimGraph/Source/AnimNextAnimGraphTestSuite/Private/AnimNextAnimGraphTraitInterfacesTest.cpp
@@ -520,6 +520,7 @@ bool FAnimationAnimNextRuntimeTest_IHierarchy::RunTest(const FString& InParamete
 		// Write our graph
 		TArray<uint8> GraphSharedDataArchiveBuffer;
 		TArray<TObjectPtr<UObject>> GraphReferencedObjects;
+		TArray<FSoftObjectPath> GraphReferencedSoftObjects;
 		{
 			FTraitWriter TraitWriter;
 
@@ -588,6 +589,7 @@ bool FAnimationAnimNextRuntimeTest_IHierarchy::RunTest(const FString& InParamete
 			AddErrorIfFalse(TraitWriter.GetErrorState() == FTraitWriter::EErrorState::None, "FAnimationAnimNextRuntimeTest_IHierarchy -> Failed to write traits");
 			GraphSharedDataArchiveBuffer = TraitWriter.GetGraphSharedData();
 			GraphReferencedObjects = TraitWriter.GetGraphReferencedObjects();
+			GraphReferencedSoftObjects = TraitWriter.GetGraphReferencedSoftObjects();
 		}
 
 		// Read our graph
@@ -733,6 +735,7 @@ bool FAnimationAnimNextRuntimeTest_IUpdate::RunTest(const FString& InParameters)
 		// Write our graph
 		TArray<uint8> GraphSharedDataArchiveBuffer;
 		TArray<TObjectPtr<UObject>> GraphReferencedObjects;
+		TArray<FSoftObjectPath> GraphReferencedSoftObjects;
 		{
 			FTraitWriter TraitWriter;
 
@@ -786,6 +789,7 @@ bool FAnimationAnimNextRuntimeTest_IUpdate::RunTest(const FString& InParameters)
 			AddErrorIfFalse(TraitWriter.GetErrorState() == FTraitWriter::EErrorState::None, "FAnimationAnimNextRuntimeTest_IUpdate -> Failed to write traits");
 			GraphSharedDataArchiveBuffer = TraitWriter.GetGraphSharedData();
 			GraphReferencedObjects = TraitWriter.GetGraphReferencedObjects();
+			GraphReferencedSoftObjects = TraitWriter.GetGraphReferencedSoftObjects();
 		}
 
 		// Read our graph
@@ -874,6 +878,7 @@ bool FAnimationAnimNextRuntimeTest_IEvaluate::RunTest(const FString& InParameter
 		// Write our graph
 		TArray<uint8> GraphSharedDataArchiveBuffer;
 		TArray<TObjectPtr<UObject>> GraphReferencedObjects;
+		TArray<FSoftObjectPath> GraphReferencedSoftObjects;
 		{
 			FTraitWriter TraitWriter;
 
@@ -928,6 +933,7 @@ bool FAnimationAnimNextRuntimeTest_IEvaluate::RunTest(const FString& InParameter
 			AddErrorIfFalse(TraitWriter.GetErrorState() == FTraitWriter::EErrorState::None, "FAnimationAnimNextRuntimeTest_IEvaluate -> Failed to write traits");
 			GraphSharedDataArchiveBuffer = TraitWriter.GetGraphSharedData();
 			GraphReferencedObjects = TraitWriter.GetGraphReferencedObjects();
+			GraphReferencedSoftObjects = TraitWriter.GetGraphReferencedSoftObjects();
 		}
 
 		// Read our graph
@@ -1011,6 +1017,7 @@ bool FAnimationAnimNextRuntimeTest_IScopedInterface::RunTest(const FString& InPa
 		// Write our graph
 		TArray<uint8> GraphSharedDataArchiveBuffer;
 		TArray<TObjectPtr<UObject>> GraphReferencedObjects;
+		TArray<FSoftObjectPath> GraphReferencedSoftObjects;
 		{
 			FTraitWriter TraitWriter;
 
@@ -1064,6 +1071,7 @@ bool FAnimationAnimNextRuntimeTest_IScopedInterface::RunTest(const FString& InPa
 			AddErrorIfFalse(TraitWriter.GetErrorState() == FTraitWriter::EErrorState::None, "FAnimationAnimNextRuntimeTest_IScopedInterface -> Failed to write traits");
 			GraphSharedDataArchiveBuffer = TraitWriter.GetGraphSharedData();
 			GraphReferencedObjects = TraitWriter.GetGraphReferencedObjects();
+			GraphReferencedSoftObjects = TraitWriter.GetGraphReferencedSoftObjects();
 		}
 
 		// Read our graph
diff --git a/Engine/Plugins/Experimental/AnimNextAnimGraph/Source/AnimNextAnimGraphTestSuite/Private/AnimNextRuntimeTest.cpp b/Engine/Plugins/Experimental/AnimNextAnimGraph/Source/AnimNextAnimGraphTestSuite/Private/AnimNextRuntimeTest.cpp
--- a/Engine/Plugins/Experimental/AnimNextAnimGraph/Source/AnimNextAnimGraphTestSuite/Private/AnimNextRuntimeTest.cpp
+++ b/Engine/Plugins/Experimental/AnimNextAnimGraph/Source/AnimNextAnimGraphTestSuite/Private/AnimNextRuntimeTest.cpp
@@ -40,10 +40,11 @@ namespace UE::AnimNext
 		AnimationGraph.ExecuteDefinition = ExecuteDefinition;
 		AnimationGraph.SharedDataArchiveBuffer = SharedDataArchiveBuffer;
 		AnimationGraph.GraphReferencedObjects.Empty();
+		AnimationGraph.GraphReferencedSoftObjects.Empty();
 
 		// Reconstruct our graph shared data
 		FMemoryReader GraphSharedDataArchive(SharedDataArchiveBuffer);
-		FTraitReader TraitReader(AnimationGraph.GraphReferencedObjects, GraphSharedDataArchive);
+		FTraitReader TraitReader(AnimationGraph.GraphReferencedObjects, AnimationGraph.GraphReferencedSoftObjects, GraphSharedDataArchive);
 
 		const FTraitReader::EErrorState ErrorState = TraitReader.ReadGraph(AnimationGraph.SharedDataBuffer);
 		if (ErrorState == FTraitReader::EErrorState::None)
diff --git a/Engine/Plugins/Experimental/AnimNextAnimGraph/Source/AnimNextAnimGraphTestSuite/Private/AnimNextTraitBaseTest.cpp b/Engine/Plugins/Experimental/AnimNextAnimGraph/Source/AnimNextAnimGraphTestSuite/Private/AnimNextTraitBaseTest.cpp
--- a/Engine/Plugins/Experimental/AnimNextAnimGraph/Source/AnimNextAnimGraphTestSuite/Private/AnimNextTraitBaseTest.cpp
+++ b/Engine/Plugins/Experimental/AnimNextAnimGraph/Source/AnimNextAnimGraphTestSuite/Private/AnimNextTraitBaseTest.cpp
@@ -547,6 +547,7 @@ bool FAnimationAnimNextRuntimeTest_NodeLifetime::RunTest(const FString& InParame
 		// Write our graph
 		TArray<uint8> GraphSharedDataArchiveBuffer;
 		TArray<TObjectPtr<UObject>> GraphReferencedObjects;
+		TArray<FSoftObjectPath> GraphReferencedSoftObjects;
 		{
 			FTraitWriter TraitWriter;
 
@@ -579,6 +580,7 @@ bool FAnimationAnimNextRuntimeTest_NodeLifetime::RunTest(const FString& InParame
 			AddErrorIfFalse(TraitWriter.GetErrorState() == FTraitWriter::EErrorState::None, "FAnimationAnimNextRuntimeTest_NodeLifetime -> Failed to write traits");
 			GraphSharedDataArchiveBuffer = TraitWriter.GetGraphSharedData();
 			GraphReferencedObjects = TraitWriter.GetGraphReferencedObjects();
+			GraphReferencedSoftObjects = TraitWriter.GetGraphReferencedSoftObjects();
 		}
 
 		// Read our graph
@@ -756,6 +758,7 @@ bool FAnimationAnimNextRuntimeTest_TraitStackBinding::RunTest(const FString& InP
 		// Write our graph
 		TArray<uint8> GraphSharedDataArchiveBuffer;
 		TArray<TObjectPtr<UObject>> GraphReferencedObjects;
+		TArray<FSoftObjectPath> GraphReferencedSoftObjects;
 		{
 			FTraitWriter TraitWriter;
 
@@ -778,6 +781,7 @@ bool FAnimationAnimNextRuntimeTest_TraitStackBinding::RunTest(const FString& InP
 			AddErrorIfFalse(TraitWriter.GetErrorState() == FTraitWriter::EErrorState::None, "FAnimationAnimNextRuntimeTest_TraitStackBinding -> Failed to write traits");
 			GraphSharedDataArchiveBuffer = TraitWriter.GetGraphSharedData();
 			GraphReferencedObjects = TraitWriter.GetGraphReferencedObjects();
+			GraphReferencedSoftObjects = TraitWriter.GetGraphReferencedSoftObjects();
 		}
 
 		// Read our graph
@@ -1257,6 +1261,7 @@ bool FAnimationAnimNextRuntimeTest_TraitSerialization::RunTest(const FString& In
 		// Write our graph
 		TArray<uint8> GraphSharedDataArchiveBuffer;
 		TArray<TObjectPtr<UObject>> GraphReferencedObjects;
+		TArray<FSoftObjectPath> GraphReferencedSoftObjects;
 		{
 			FTraitWriter TraitWriter;
 
@@ -1348,6 +1353,7 @@ bool FAnimationAnimNextRuntimeTest_TraitSerialization::RunTest(const FString& In
 			AddErrorIfFalse(TraitWriter.GetErrorState() == FTraitWriter::EErrorState::None, "FAnimationAnimNextRuntimeTest_TraitSerialization -> Failed to write traits");
 			GraphSharedDataArchiveBuffer = TraitWriter.GetGraphSharedData();
 			GraphReferencedObjects = TraitWriter.GetGraphReferencedObjects();
+			GraphReferencedSoftObjects = TraitWriter.GetGraphReferencedSoftObjects();
 		}
 
 		// Clear out the node template registry to test registration on load
diff --git a/Engine/Plugins/Experimental/AnimNextAnimGraph/Source/AnimNextAnimGraphTestSuite/Private/AnimNextTraitEventTest.cpp b/Engine/Plugins/Experimental/AnimNextAnimGraph/Source/AnimNextAnimGraphTestSuite/Private/AnimNextTraitEventTest.cpp
--- a/Engine/Plugins/Experimental/AnimNextAnimGraph/Source/AnimNextAnimGraphTestSuite/Private/AnimNextTraitEventTest.cpp
+++ b/Engine/Plugins/Experimental/AnimNextAnimGraph/Source/AnimNextAnimGraphTestSuite/Private/AnimNextTraitEventTest.cpp
@@ -268,6 +268,7 @@ bool FAnimationAnimNextRuntimeTest_TraitEventRaising::RunTest(const FString& InP
 		// Write our graph
 		TArray<uint8> GraphSharedDataArchiveBuffer;
 		TArray<TObjectPtr<UObject>> GraphReferencedObjects;
+		TArray<FSoftObjectPath> GraphReferencedSoftObjects;
 		{
 			FTraitWriter TraitWriter;
 
@@ -290,6 +291,7 @@ bool FAnimationAnimNextRuntimeTest_TraitEventRaising::RunTest(const FString& InP
 			AddErrorIfFalse(TraitWriter.GetErrorState() == FTraitWriter::EErrorState::None, "FAnimationAnimNextRuntimeTest_TraitEventRaising -> Failed to write traits");
 			GraphSharedDataArchiveBuffer = TraitWriter.GetGraphSharedData();
 			GraphReferencedObjects = TraitWriter.GetGraphReferencedObjects();
+			GraphReferencedSoftObjects = TraitWriter.GetGraphReferencedSoftObjects();
 		}
 
 		// Read our graph
diff --git a/Engine/Plugins/Experimental/AnimNextAnimGraph/Source/AnimNextAnimGraphUncookedOnly/Internal/Graph/AnimNextAnimationGraph_EditorData.h b/Engine/Plugins/Experimental/AnimNextAnimGraph/Source/AnimNextAnimGraphUncookedOnly/Internal/Graph/AnimNextAnimationGraph_EditorData.h
--- a/Engine/Plugins/Experimental/AnimNextAnimGraph/Source/AnimNextAnimGraphUncookedOnly/Internal/Graph/AnimNextAnimationGraph_EditorData.h
+++ b/Engine/Plugins/Experimental/AnimNextAnimGraph/Source/AnimNextAnimGraphUncookedOnly/Internal/Graph/AnimNextAnimationGraph_EditorData.h
@@ -71,4 +71,5 @@ protected:
 	virtual bool CanAddNewEntry(TSubclassOf<UAnimNextRigVMAssetEntry> InClass) const override;
 	virtual TSubclassOf<UAssetUserData> GetAssetUserDataClass() const override;
 	virtual void InitializeAssetUserData() override;
+	virtual void GetProgrammaticGraphs(const FRigVMCompileSettings& InSettings, TArray<URigVMGraph*>& OutGraphs) override;
 };
diff --git a/Engine/Plugins/Experimental/AnimNextAnimGraph/Source/AnimNextAnimGraphUncookedOnly/Private/Graph/AnimNextAnimationGraph_EditorData.cpp b/Engine/Plugins/Experimental/AnimNextAnimGraph/Source/AnimNextAnimGraphUncookedOnly/Private/Graph/AnimNextAnimationGraph_EditorData.cpp
--- a/Engine/Plugins/Experimental/AnimNextAnimGraph/Source/AnimNextAnimGraphUncookedOnly/Private/Graph/AnimNextAnimationGraph_EditorData.cpp
+++ b/Engine/Plugins/Experimental/AnimNextAnimGraph/Source/AnimNextAnimGraphUncookedOnly/Private/Graph/AnimNextAnimationGraph_EditorData.cpp
@@ -17,12 +17,15 @@
 #include "Graph/RigUnit_AnimNextBeginExecution.h"
 #include "Graph/RigUnit_AnimNextShimRoot.h"
 #include "Graph/RigUnit_AnimNextTraitStack.h"
+#include "Logging/LogScopedVerbosityOverride.h"
 #include "Logging/MessageLog.h"
+#include "RigVMFunctions/Execution/RigVMFunction_UserDefinedEvent.h"
 #include "TraitCore/NodeTemplateBuilder.h"
 #include "TraitCore/TraitRegistry.h"
 #include "TraitCore/TraitWriter.h"
 #include "UObject/AssetRegistryTagsContext.h"
 #include "UObject/LinkerLoad.h"
+#include "Traits/CallFunction.h"
 
 namespace UE::AnimNext::UncookedOnly::Private
 {
@@ -130,9 +133,10 @@ namespace UE::AnimNext::UncookedOnly::Private
 		const TArray<URigVMPin*>& Pins = TraitStack.TraitEntries[TraitIndex].DecoratorEntryPin->GetSubPins();
 		for (const URigVMPin* Pin : Pins)
 		{
-			if (Pin->GetDirection() != ERigVMPinDirection::Input)
+			if (Pin->GetDirection() != ERigVMPinDirection::Input &&
+				Pin->GetDirection() != ERigVMPinDirection::Hidden)
 			{
-				continue;	// We only look for input pins
+				continue;	// We only look for input or hidden pins
 			}
 
 			if (Pin->GetFName() == PropertyName)
@@ -498,6 +502,7 @@ void UAnimNextAnimationGraph_EditorData::RecompileVM()
 		AnimationGraph->ExecuteDefinition = FAnimNextGraphEvaluatorExecuteDefinition();
 		AnimationGraph->SharedDataBuffer.Empty();
 		AnimationGraph->GraphReferencedObjects.Empty();
+		AnimationGraph->GraphReferencedSoftObjects.Empty();
 		AnimationGraph->DefaultEntryPoint = NAME_None;
 
 		FRigVMClient* VMClient = GetRigVMClient();
@@ -516,35 +521,39 @@ void UAnimNextAnimationGraph_EditorData::RecompileVM()
 			return;
 		}
 
-		TArray<URigVMGraph*> TempGraphs;
-		for(const URigVMGraph* SourceGraph : AllGraphs)
+		TArray<URigVMGraph*> AnimGraphs;
+		TArray<URigVMGraph*> NonAnimGraphs;
+		for(URigVMGraph* SourceGraph : AllGraphs)
 		{
 			// We use a temporary graph models to build our final graphs that we'll compile
-			URigVMGraph* TempGraph = CastChecked<URigVMGraph>(StaticDuplicateObject(SourceGraph, GetTransientPackage(), NAME_None, RF_Transient));
-			TempGraph->SetFlags(RF_Transient);
-			TempGraphs.Add(TempGraph);
+			if(SourceGraph->GetSchemaClass() == UAnimNextAnimationGraphSchema::StaticClass())
+			{
+				URigVMGraph* TempGraph = CastChecked<URigVMGraph>(StaticDuplicateObject(SourceGraph, this, NAME_None, RF_Transient));
+				UAnimNextController* TempController = CastChecked<UAnimNextController>(VMClient->GetOrCreateController(TempGraph));
+				TempGraph->SetFlags(RF_Transient);
+				AnimGraphs.Add(TempGraph);
+			}
+			else
+			{
+				NonAnimGraphs.Add(SourceGraph);
+			}
 		}
 
-		if(TempGraphs.Num() == 0)
+		if(AnimGraphs.Num() > 0)
 		{
-			return;
-		}
+			UAnimNextController* TempController = CastChecked<UAnimNextController>(VMClient->GetOrCreateController(AnimGraphs[0]));
 
-		UAnimNextController* TempController = CastChecked<UAnimNextController>(VMClient->GetOrCreateController(TempGraphs[0]));
+			UE::AnimNext::FTraitWriter TraitWriter;
 
-		UE::AnimNext::FTraitWriter TraitWriter;
+			FRigVMPinInfoArray LatentPins;
+			TMap<FName, URigVMPin*> LatentPinMapping;
+			TArray<Private::FTraitGraph> TraitGraphs;
 
-		FRigVMPinInfoArray LatentPins;
-		TMap<FName, URigVMPin*> LatentPinMapping;
-		TArray<Private::FTraitGraph> TraitGraphs;
-
-		// Build entry points and extract their required latent pins
-		for(const URigVMGraph* TempGraph : TempGraphs)
-		{
-			if(TempGraph->GetSchemaClass() == UAnimNextAnimationGraphSchema::StaticClass())
+			// Build entry points and extract their required latent pins
+			for(const URigVMGraph* AnimGraph : AnimGraphs)
 			{
 				// Gather our trait stacks
-				Private::FTraitGraph& TraitGraph = TraitGraphs.Add_GetRef(Private::CollectGraphInfo(AnimationGraph, TempGraph, TempController->GetControllerForGraph(TempGraph)));
+				Private::FTraitGraph& TraitGraph = TraitGraphs.Add_GetRef(Private::CollectGraphInfo(AnimationGraph, AnimGraph, TempController->GetControllerForGraph(AnimGraph)));
 				check(!TraitGraph.TraitStackNodes.IsEmpty());
 
 				FAnimNextGraphEntryPoint& EntryPoint = AnimationGraph->EntryPoints.AddDefaulted_GetRef();
@@ -565,68 +574,73 @@ void UAnimNextAnimationGraph_EditorData::RecompileVM()
 					EntryPoint.RootTraitHandle = FAnimNextEntryPointHandle(TraitGraph.TraitStackNodes[0].TraitStackNodeHandle);
 				}
 			}
-		}
-
-		// Set default entry point
-		if(AnimationGraph->EntryPoints.Num() > 0)
-		{
-			AnimationGraph->DefaultEntryPoint = AnimationGraph->EntryPoints[0].EntryPointName;
-		}
 
-		// Remove our old root nodes
-		for (Private::FTraitGraph& TraitGraph : TraitGraphs)
-		{
-			URigVMController* GraphController = TempController->GetControllerForGraph(TraitGraph.RootNode->GetGraph());
-			GraphController->RemoveNode(TraitGraph.RootNode, false, false);
-		}
+			// Set default entry point
+			if(AnimationGraph->EntryPoints.Num() > 0)
+			{
+				AnimationGraph->DefaultEntryPoint = AnimationGraph->EntryPoints[0].EntryPointName;
+			}
 
-		if(LatentPins.Num() > 0)
-		{
-			// We need a unique method name to match our unique argument list
-			AnimationGraph->ExecuteDefinition = Private::GetGraphEvaluatorExecuteMethod(LatentPins);
+			// Remove our old root nodes
+			for (Private::FTraitGraph& TraitGraph : TraitGraphs)
+			{
+				URigVMController* GraphController = TempController->GetControllerForGraph(TraitGraph.RootNode->GetGraph());
+				GraphController->RemoveNode(TraitGraph.RootNode, false, false);
+			}
 
-			// Add our runtime shim root node
-			URigVMUnitNode* TempShimRootNode = TempController->AddUnitNode(FRigUnit_AnimNextShimRoot::StaticStruct(), FRigUnit_AnimNextShimRoot::EventName, FVector2D::ZeroVector, FString(), false);
-			URigVMUnitNode* GraphEvaluatorNode = TempController->AddUnitNodeWithPins(FRigUnit_AnimNextGraphEvaluator::StaticStruct(), LatentPins, *AnimationGraph->ExecuteDefinition.MethodName, FVector2D::ZeroVector, FString(), false);
+			if(LatentPins.Num() > 0)
+			{
+				// We need a unique method name to match our unique argument list
+				AnimationGraph->ExecuteDefinition = Private::GetGraphEvaluatorExecuteMethod(LatentPins);
 
-			// Link our shim and evaluator nodes together using the execution context
-			TempController->AddLink(
-				TempShimRootNode->FindPin(GET_MEMBER_NAME_STRING_CHECKED(FRigUnit_AnimNextShimRoot, ExecuteContext)),
-				GraphEvaluatorNode->FindPin(GET_MEMBER_NAME_STRING_CHECKED(FRigUnit_AnimNextGraphEvaluator, ExecuteContext)),
-				false);
+				// Add our runtime shim root node
+				URigVMUnitNode* TempShimRootNode = TempController->AddUnitNode(FRigUnit_AnimNextShimRoot::StaticStruct(), FRigUnit_AnimNextShimRoot::EventName, FVector2D::ZeroVector, FString(), false);
+				URigVMUnitNode* GraphEvaluatorNode = TempController->AddUnitNodeWithPins(FRigUnit_AnimNextGraphEvaluator::StaticStruct(), LatentPins, *AnimationGraph->ExecuteDefinition.MethodName, FVector2D::ZeroVector, FString(), false);
 
-			// Link our latent pins
-			for (const FRigVMPinInfo& LatentPin : LatentPins)
-			{
+				// Link our shim and evaluator nodes together using the execution context
 				TempController->AddLink(
-					LatentPinMapping[LatentPin.Name],
-					GraphEvaluatorNode->FindPin(LatentPin.Name.ToString()),
+					TempShimRootNode->FindPin(GET_MEMBER_NAME_STRING_CHECKED(FRigUnit_AnimNextShimRoot, ExecuteContext)),
+					GraphEvaluatorNode->FindPin(GET_MEMBER_NAME_STRING_CHECKED(FRigUnit_AnimNextGraphEvaluator, ExecuteContext)),
 					false);
+
+				// Link our latent pins
+				for (const FRigVMPinInfo& LatentPin : LatentPins)
+				{
+					TempController->AddLink(
+						LatentPinMapping[LatentPin.Name],
+						GraphEvaluatorNode->FindPin(LatentPin.Name.ToString()),
+						false);
+				}
 			}
-		}
 
-		// Write our node shared data
-		TraitWriter.BeginNodeWriting();
+			// Write our node shared data
+			TraitWriter.BeginNodeWriting();
 
-		for(Private::FTraitGraph& TraitGraph : TraitGraphs)
-		{
-			for (const Private::FTraitStackMapping& NodeMapping : TraitGraph.TraitStackNodes)
+			for(Private::FTraitGraph& TraitGraph : TraitGraphs)
 			{
-				Private::WriteTraitProperties(TraitWriter, NodeMapping, TraitGraph.TraitStackNodes);
+				for (const Private::FTraitStackMapping& NodeMapping : TraitGraph.TraitStackNodes)
+				{
+					Private::WriteTraitProperties(TraitWriter, NodeMapping, TraitGraph.TraitStackNodes);
+				}
 			}
-		}
 
-		TraitWriter.EndNodeWriting();
+			TraitWriter.EndNodeWriting();
+
+			// Cache our compiled metadata
+			AnimationGraph->SharedDataArchiveBuffer = TraitWriter.GetGraphSharedData();
+			AnimationGraph->GraphReferencedObjects = TraitWriter.GetGraphReferencedObjects();
+			AnimationGraph->GraphReferencedSoftObjects = TraitWriter.GetGraphReferencedSoftObjects();
 
-		// Cache our compiled metadata
-		AnimationGraph->SharedDataArchiveBuffer = TraitWriter.GetGraphSharedData();
-		AnimationGraph->GraphReferencedObjects = TraitWriter.GetGraphReferencedObjects();
+			// Populate our runtime metadata
+			AnimationGraph->LoadFromArchiveBuffer(AnimationGraph->SharedDataArchiveBuffer);
+		}
 
-		// Populate our runtime metadata
-		AnimationGraph->LoadFromArchiveBuffer(AnimationGraph->SharedDataArchiveBuffer);
+		TArray<URigVMGraph*> GraphsToCompile = MoveTemp(AnimGraphs);
+		GraphsToCompile.Append(NonAnimGraphs);
 
+		UAnimNextController* Controller = CastChecked<UAnimNextController>(VMClient->GetOrCreateController(GraphsToCompile[0]));
 		URigVMCompiler* Compiler = URigVMCompiler::StaticClass()->GetDefaultObject<URigVMCompiler>();
-		Compiler->Compile(Settings, TempGraphs, TempController, AnimationGraph->VM, AnimationGraph->ExtendedExecuteContext, AnimationGraph->GetExternalVariables(), &PinToOperandMap);
+		Compiler->Compile(Settings, GraphsToCompile, Controller, AnimationGraph->VM, AnimationGraph->ExtendedExecuteContext, AnimationGraph->GetExternalVariables(), &PinToOperandMap);
 
 		// Initialize right away, in packaged builds we initialize during PostLoad
 		AnimationGraph->VM->Initialize(AnimationGraph->ExtendedExecuteContext);
@@ -649,23 +663,24 @@ void UAnimNextAnimationGraph_EditorData::RecompileVM()
 			RigVMCompiledEvent.Broadcast(AnimationGraph, AnimationGraph->VM, AnimationGraph->ExtendedExecuteContext);
 		}
 
-		for(URigVMGraph* TempGraph : TempGraphs)
-		{
-			VMClient->RemoveController(TempGraph);
-		}
+		// Remove the temp anim graphs
+		RemoveProgrammaticGraphs(AnimGraphs);
 
 		// Now that the graph has been re-compiled, re-allocate the previous live instances
 		AnimationGraph->ThawGraphInstances();
 	}
-		
+
 #if WITH_EDITOR
 	// Display programmatic graphs
 	if(CVarDumpProgrammaticGraphs.GetValueOnGameThread())
 	{
 		FUtils::OpenProgrammaticGraphs(this, ProgrammaticGraphs);
 	}
+	else
 #endif
-
+	{
+		RemoveProgrammaticGraphs(ProgrammaticGraphs);
+	}
 
 #if WITH_EDITOR
 //	RefreshBreakpoints(EditorData);
@@ -812,4 +827,120 @@ void UAnimNextAnimationGraph_EditorData::InitializeAssetUserData()
 	}
 
 	Super::InitializeAssetUserData();
+}
+
+void UAnimNextAnimationGraph_EditorData::GetProgrammaticGraphs(const FRigVMCompileSettings& InSettings, TArray<URigVMGraph*>& OutGraphs)
+{
+	using namespace UE::AnimNext::UncookedOnly;
+
+	Super::GetProgrammaticGraphs(InSettings, OutGraphs);
+
+	// Gather all 'call function' traits and create shim-calls for them.
+	// For the compiler to pick them up if they are not public we need a calling reference to the function from a graph
+	TArray<FRigVMGraphFunctionHeader> FunctionHeaders;
+
+	UAnimNextAnimationGraph* Asset = FUtils::GetAsset<UAnimNextAnimationGraph>(this);
+	FRigVMClient* VMClient = GetRigVMClient();
+	for (const URigVMGraph* Graph : VMClient->GetAllModels(false, false))
+	{
+		for(const URigVMNode* Node : Graph->GetNodes())
+		{
+			for(const URigVMPin* TraitPin : Node->GetTraitPins())
+			{
+				if (TraitPin->IsExecuteContext())
+				{
+					continue;
+				}
+
+				TSharedPtr<FStructOnScope> ScopedTrait = Node->GetTraitInstance(TraitPin->GetFName());
+				if (!ScopedTrait.IsValid())
+				{
+					continue;
+				}
+
+				const FRigVMTrait* Trait = (FRigVMTrait*)ScopedTrait->GetStructMemory();
+				UScriptStruct* TraitSharedInstanceData = Trait->GetTraitSharedDataStruct();
+				if (TraitSharedInstanceData == nullptr)
+				{
+					continue;
+				}
+
+				if(!TraitSharedInstanceData->IsChildOf(FAnimNextCallFunctionSharedData::StaticStruct()))
+				{
+					continue;
+				}
+
+				const FString DefaultValue = TraitPin->GetDefaultValue();
+				TInstancedStruct<FAnimNextCallFunctionSharedData> InstancedStruct = TInstancedStruct<FAnimNextCallFunctionSharedData>::Make();
+				FRigVMPinDefaultValueImportErrorContext ErrorPipe(ELogVerbosity::Verbose);
+				LOG_SCOPE_VERBOSITY_OVERRIDE(LogExec, ErrorPipe.GetMaxVerbosity());
+				TraitSharedInstanceData->ImportText(*DefaultValue, InstancedStruct.GetMutableMemory(), nullptr, PPF_SerializedAsImportText, &ErrorPipe, TraitSharedInstanceData->GetName());
+
+				const FRigVMGraphFunctionHeader& FunctionHeader = InstancedStruct.Get<FAnimNextCallFunctionSharedData>().FunctionHeader;
+				if(FunctionHeader.IsValid())
+				{
+					FunctionHeaders.Add(FunctionHeader);
+				}
+			}
+		}
+	}
+
+	if(FunctionHeaders.Num() > 0)
+	{
+		// Create all shim events for our traits to call
+		const bool bSetupUndoRedo = false;
+		URigVMGraph* WrapperGraph = NewObject<URigVMGraph>(this, NAME_None, RF_Transient);
+		UAnimNextController* Controller = CastChecked<UAnimNextController>(VMClient->GetOrCreateController(WrapperGraph));
+		FRigVMControllerNotifGuard NotifGuard(Controller);
+		bool bAddedWrapperEvent = true;
+		for(const FRigVMGraphFunctionHeader& FunctionHeader : FunctionHeaders)
+		{
+			URigVMLibraryNode* LibraryNode = Cast<URigVMLibraryNode>(FunctionHeader.LibraryPointer.GetNodeSoftPath().TryLoad());
+			if(LibraryNode == nullptr)
+			{
+				InSettings.ReportError(FString::Printf(TEXT("Could not find function '%s'"), *FunctionHeader.Name.ToString()));
+				continue;
+			}
+
+			// Create user-defined entry point
+			FString WrapperEventName = FUtils::MakeFunctionWrapperEventName(FunctionHeader.Name);
+			URigVMUnitNode* EventNode = Controller->AddUnitNode(FRigVMFunction_UserDefinedEvent::StaticStruct(), TEXT("Execute"), FVector2D::ZeroVector, FunctionHeader.Name.ToString(), bSetupUndoRedo);
+			if(EventNode == nullptr)
+			{
+				InSettings.ReportError(FString::Printf(TEXT("Could not spawn event node for function '%s'"), *FunctionHeader.Name.ToString()));
+				continue;
+			}
+			URigVMPin* EventNamePin = EventNode->FindPin(GET_MEMBER_NAME_STRING_CHECKED(FRigVMFunction_UserDefinedEvent, EventName));
+			if(EventNamePin == nullptr)
+			{
+				InSettings.ReportError(FString::Printf(TEXT("Could not find custom event name pin")));
+				continue;
+			}
+			Controller->SetPinDefaultValue(EventNamePin->GetPinPath(), WrapperEventName, true, bSetupUndoRedo);
+
+			// Call function
+			URigVMFunctionReferenceNode* FunctionNode = Controller->AddFunctionReferenceNode(LibraryNode, FVector2D::ZeroVector, FunctionHeader.Name.ToString(), bSetupUndoRedo);
+			if(FunctionNode == nullptr)
+			{
+				InSettings.ReportError(FString::Printf(TEXT("Could not spawn function node for function '%s'"), *FunctionHeader.Name.ToString()));
+				continue;
+			}
+
+			// Link up nodes
+			URigVMPin* ExecuteOutputPin = EventNode->FindPin(FRigVMStruct::ExecuteContextName.ToString());
+			URigVMPin* ExecuteInputPin = FunctionNode->FindPin(FRigVMStruct::ExecuteContextName.ToString());
+			if(!Controller->AddLink(ExecuteOutputPin, ExecuteInputPin, bSetupUndoRedo))
+			{
+				InSettings.ReportError(FString::Printf(TEXT("Could not link execute pins for function '%s'"), *FunctionHeader.Name.ToString()));
+				continue;
+			}
+
+			bAddedWrapperEvent = true;
+		}
+
+		if(bAddedWrapperEvent)
+		{
+			OutGraphs.Add(WrapperGraph);
+		}
+	}
 }
\ No newline at end of file
diff --git a/Engine/Plugins/Runtime/RigVM/Source/RigVM/Private/RigVMCore/RigVMStruct.cpp b/Engine/Plugins/Runtime/RigVM/Source/RigVM/Private/RigVMCore/RigVMStruct.cpp
--- a/Engine/Plugins/Runtime/RigVM/Source/RigVM/Private/RigVMCore/RigVMStruct.cpp
+++ b/Engine/Plugins/Runtime/RigVM/Source/RigVM/Private/RigVMCore/RigVMStruct.cpp
@@ -607,10 +607,12 @@ FString FRigVMStruct::ExportToFullyQualifiedText(const FProperty* InMemberProper
 	}
 	else
 	{
-		InMemberProperty->ExportTextItem_Direct(DefaultValue, InMemberMemoryPtr, nullptr, nullptr, PPF_None);
+		// Note we pass PPF_Delimited here to correctly escape strings
+		InMemberProperty->ExportTextItem_Direct(DefaultValue, InMemberMemoryPtr, nullptr, nullptr, PPF_Delimited);
 
 		if (CastField<FNameProperty>(InMemberProperty) != nullptr ||
-			CastField<FStrProperty>(InMemberProperty) != nullptr)
+			CastField<FStrProperty>(InMemberProperty) != nullptr ||
+			CastField<FTextProperty>(InMemberProperty) != nullptr)
 		{
 			if(bUseQuotes)
 			{
@@ -618,11 +620,15 @@ FString FRigVMStruct::ExportToFullyQualifiedText(const FProperty* InMemberProper
 				{
 					DefaultValue = TEXT("\"\"");
 				}
-				else
+				else if (DefaultValue[0] != TCHAR('\"') && DefaultValue[DefaultValue.Len() - 1] != TCHAR('\"'))
 				{
 					DefaultValue = FString::Printf(TEXT("\"%s\""), *DefaultValue);
 				}
 			}
+			else
+			{
+				DefaultValue = DefaultValue.TrimQuotes();
+			}
 		}
 	}
 
@@ -634,26 +640,38 @@ FString FRigVMStruct::ExportToFullyQualifiedText(const UScriptStruct* InStruct,
 	check(InStruct);
 	check(InStructMemoryPtr);
 
-	TArray<FString> FieldValues;
-	for (TFieldIterator<FProperty> It(InStruct); It; ++It)
+	// If a struct has a custom text import/export, we must use its ExportText call rather than going per-property
+	if((InStruct->StructFlags & (STRUCT_ImportTextItemNative | STRUCT_ExportTextItemNative)) != 0)
 	{
-		if(It->HasAnyPropertyFlags(CPF_Transient))
+		FInstancedStruct Defaults;
+		Defaults.InitializeAs(InStruct);
+		FString Value;
+		InStruct->ExportText(Value, InStructMemoryPtr, Defaults.GetMemory(), nullptr, PPF_None, nullptr);
+		return Value;
+	}
+	else
+	{
+		TArray<FString> FieldValues;
+		for (TFieldIterator<FProperty> It(InStruct); It; ++It)
 		{
-			continue;
-		}
+			if(It->HasAnyPropertyFlags(CPF_Transient))
+			{
+				continue;
+			}
 		
-		FString PropertyName = It->GetName();
-		const uint8* StructMemberMemoryPtr = It->ContainerPtrToValuePtr<uint8>(InStructMemoryPtr);
-		FString DefaultValue = ExportToFullyQualifiedText(*It, StructMemberMemoryPtr, bUseQuotes);
-		FieldValues.Add(FString::Printf(TEXT("%s=%s"), *PropertyName, *DefaultValue));
-	}
+			FString PropertyName = It->GetName();
+			const uint8* StructMemberMemoryPtr = It->ContainerPtrToValuePtr<uint8>(InStructMemoryPtr);
+			FString DefaultValue = ExportToFullyQualifiedText(*It, StructMemberMemoryPtr, bUseQuotes);
+			FieldValues.Add(FString::Printf(TEXT("%s=%s"), *PropertyName, *DefaultValue));
+		}
 
-	if (FieldValues.Num() == 0)
-	{
-		return TEXT("()");
-	}
+		if (FieldValues.Num() == 0)
+		{
+			return TEXT("()");
+		}
 
-	return FString::Printf(TEXT("(%s)"), *FString::Join(FieldValues, TEXT(",")));
+		return FString::Printf(TEXT("(%s)"), *FString::Join(FieldValues, TEXT(",")));
+	}
 }
 
 FString FRigVMStruct::ExportToFullyQualifiedText(const UScriptStruct* InScriptStruct, const FName& InPropertyName, const uint8* InStructMemoryPointer, bool bUseQuotes) const
@@ -712,28 +730,32 @@ FName FRigVMStruct::GetNextAggregateName(const FName& InLastAggregatePinName) co
 TMap<FName, FString> FRigVMStruct::GetDefaultValues(UScriptStruct* InScriptStruct) const
 {
 	check(InScriptStruct);
-	
+
 	TMap<FName, FString> DefaultValues;
-	for (TFieldIterator<FProperty> It(InScriptStruct); It; ++It)
+	// If a struct has a custom text import/export, we cannot extract its property values individually
+	if((InScriptStruct->StructFlags & (STRUCT_ImportTextItemNative | STRUCT_ExportTextItemNative)) == 0)
 	{
-		if(It->HasAnyPropertyFlags(CPF_Transient))
+		for (TFieldIterator<FProperty> It(InScriptStruct); It; ++It)
 		{
-			continue;
-		}
+			if(It->HasAnyPropertyFlags(CPF_Transient))
+			{
+				continue;
+			}
 
 #if WITH_EDITOR
-		if(!It->HasMetaData(FRigVMStruct::InputMetaName) &&
-			!It->HasMetaData(FRigVMStruct::OutputMetaName) &&
-			!It->HasMetaData(FRigVMStruct::VisibleMetaName))
-		{
-			continue;
-		}
+			if(!It->HasMetaData(FRigVMStruct::InputMetaName) &&
+				!It->HasMetaData(FRigVMStruct::OutputMetaName) &&
+				!It->HasMetaData(FRigVMStruct::VisibleMetaName))
+			{
+				continue;
+			}
 #endif
-		
-		const FName PropertyName = It->GetFName();
-		const uint8* StructMemberMemoryPtr = It->ContainerPtrToValuePtr<uint8>(this);
-		const FString DefaultValue = ExportToFullyQualifiedText(*It, StructMemberMemoryPtr, false);
-		DefaultValues.Add(PropertyName, DefaultValue);
+			
+			const FName PropertyName = It->GetFName();
+			const uint8* StructMemberMemoryPtr = It->ContainerPtrToValuePtr<uint8>(this);
+			const FString DefaultValue = ExportToFullyQualifiedText(*It, StructMemberMemoryPtr, false);
+			DefaultValues.Add(PropertyName, DefaultValue);
+		}
 	}
 
 	return DefaultValues;
diff --git a/Engine/Plugins/Runtime/RigVM/Source/RigVMDeveloper/Private/RigVMModel/RigVMBuildData.cpp b/Engine/Plugins/Runtime/RigVM/Source/RigVMDeveloper/Private/RigVMModel/RigVMBuildData.cpp
--- a/Engine/Plugins/Runtime/RigVM/Source/RigVMDeveloper/Private/RigVMModel/RigVMBuildData.cpp
+++ b/Engine/Plugins/Runtime/RigVM/Source/RigVMDeveloper/Private/RigVMModel/RigVMBuildData.cpp
@@ -491,7 +491,7 @@ void URigVMBuildData::ClearInvalidReferences()
 		return;
 	}
 	
-	Modify();
+	Modify(false);
 	
 	// check each function's each reference
 	int32 NumRemoved = 0;
diff --git a/Engine/Plugins/Runtime/RigVM/Source/RigVMDeveloper/Private/RigVMModel/RigVMController.cpp b/Engine/Plugins/Runtime/RigVM/Source/RigVMDeveloper/Private/RigVMModel/RigVMController.cpp
--- a/Engine/Plugins/Runtime/RigVM/Source/RigVMDeveloper/Private/RigVMModel/RigVMController.cpp
+++ b/Engine/Plugins/Runtime/RigVM/Source/RigVMDeveloper/Private/RigVMModel/RigVMController.cpp
@@ -19154,7 +19154,11 @@ void URigVMController::PostProcessDefaultValue(const URigVMPin* Pin, FString& Ou
 	}
 	if (Pin->IsStruct() || Pin->IsArray())
 	{
-		if (!OutDefaultValue.IsEmpty())
+		// We cannot sanitize struct text defaults that have native import/export, as they do not have to follow conventional syntax 
+		const UScriptStruct* ScriptStruct = Pin->GetScriptStruct();
+		const bool bHasNativeImportExport = ScriptStruct && ((ScriptStruct->StructFlags & (STRUCT_ImportTextItemNative | STRUCT_ExportTextItemNative)) != 0);
+
+		if (!OutDefaultValue.IsEmpty() && !bHasNativeImportExport)
 		{
 			if (OutDefaultValue[0] != TCHAR('(') || OutDefaultValue[OutDefaultValue.Len()-1] != TCHAR(')'))
 			{
diff --git a/Engine/Plugins/Runtime/RigVM/Source/RigVMDeveloper/Private/RigVMModel/RigVMPin.cpp b/Engine/Plugins/Runtime/RigVM/Source/RigVMDeveloper/Private/RigVMModel/RigVMPin.cpp
--- a/Engine/Plugins/Runtime/RigVM/Source/RigVMDeveloper/Private/RigVMModel/RigVMPin.cpp
+++ b/Engine/Plugins/Runtime/RigVM/Source/RigVMDeveloper/Private/RigVMModel/RigVMPin.cpp
@@ -870,40 +870,48 @@ FString URigVMPin::GetDefaultValue(const URigVMPin::FPinOverride& InOverride, bo
 	{
 		static const FString EmptyStructDefaultValue = TEXT("()");
 		
-		// for trait pins, there are cases where a pin is not created for a property (see ShouldCreatePinForProperty())
-		// so we store the value of that property in the default value of the struct pin containing that property
-		// as a result, to retrieve the default value we need to combine the default value on the struct pin, with additional overrides in the available sub pins
-		if (SubPins.Num() > 0 || IsTraitPin())
+		if((GetScriptStruct()->StructFlags & (STRUCT_ImportTextItemNative | STRUCT_ExportTextItemNative)) != 0)
 		{
-			FString FinalDefaultValue = DefaultValue;
-			
-			// root trait pin store their default value in a separate property bag so that
-			// things like soft object ptr can be used and tracked in a uproperty
-			if (IsTraitPin() && IsRootPin())
+			// If the struct has a native import/export, then its default value must be used verbatim (subpins are not displayed)
+			return DefaultValue;
+		}
+		else
+		{
+			// for trait pins, there are cases where a pin is not created for a property (see ShouldCreatePinForProperty())
+			// so we store the value of that property in the default value of the struct pin containing that property
+			// as a result, to retrieve the default value we need to combine the default value on the struct pin, with additional overrides in the available sub pins
+			if (SubPins.Num() > 0 || IsTraitPin())
 			{
-				FRigVMTraitDefaultValueStruct* DefaultValueStructPtr = GetNode()->TraitDefaultValues.Find(GetName());
-				if (ensure(DefaultValueStructPtr))
+				FString FinalDefaultValue = DefaultValue;
+				
+				// root trait pin store their default value in a separate property bag so that
+				// things like soft object ptr can be used and tracked in a uproperty
+				if (IsTraitPin() && IsRootPin())
 				{
-					FinalDefaultValue = DefaultValueStructPtr->GetValue();
+					FRigVMTraitDefaultValueStruct* DefaultValueStructPtr = GetNode()->TraitDefaultValues.Find(GetName());
+					if (ensure(DefaultValueStructPtr))
+					{
+						FinalDefaultValue = DefaultValueStructPtr->GetValue();
+					}
 				}
-			}
 
-			for (const URigVMPin* SubPin : SubPins)
-			{
-				FString MemberDefaultValue = SubPin->GetDefaultValue(InOverride, bAdaptValueForPinType);
-				if (SubPin->IsStringType() && !MemberDefaultValue.IsEmpty())
+				for (const URigVMPin* SubPin : SubPins)
 				{
-					MemberDefaultValue = TEXT("\"") + MemberDefaultValue + TEXT("\"");
-				}
-				else if (MemberDefaultValue.IsEmpty() || MemberDefaultValue == TEXT("()"))
-				{
-					continue;
+					FString MemberDefaultValue = SubPin->GetDefaultValue(InOverride, bAdaptValueForPinType);
+					if (SubPin->IsStringType() && !MemberDefaultValue.IsEmpty())
+					{
+						MemberDefaultValue = TEXT("\"") + MemberDefaultValue + TEXT("\"");
+					}
+					else if (MemberDefaultValue.IsEmpty() || MemberDefaultValue == TEXT("()"))
+					{
+						continue;
+					}
+
+					URigVMController::OverrideDefaultValueMember(SubPin->GetName(), MemberDefaultValue, FinalDefaultValue);
 				}
 
-				URigVMController::OverrideDefaultValueMember(SubPin->GetName(), MemberDefaultValue, FinalDefaultValue);
+				return !FinalDefaultValue.IsEmpty() ? FinalDefaultValue : EmptyStructDefaultValue;
 			}
-
-			return !FinalDefaultValue.IsEmpty() ? FinalDefaultValue : EmptyStructDefaultValue;
 		}
 
 		// special case certain pin types to adapt their values from
@@ -1057,8 +1065,9 @@ bool URigVMPin::IsValidDefaultValue(const FString& InDefaultValue) const
 			FRigVMPinDefaultValueImportErrorContext ErrorPipe(ELogVerbosity::Verbose);
 			{
 				// force logging to the error pipe for error detection
-				LOG_SCOPE_VERBOSITY_OVERRIDE(LogExec, ErrorPipe.GetMaxVerbosity()); 
-				ScriptStruct->ImportText(*Value, TempStructBuffer.GetData(), nullptr, PPF_None, &ErrorPipe, ScriptStruct->GetName()); 
+				LOG_SCOPE_VERBOSITY_OVERRIDE(LogExec, ErrorPipe.GetMaxVerbosity());
+				// Note we use PPF_UseDeprecatedProperties as if a deprecated property is encountered this call will otherwise fail via the ErrorPipe
+				ScriptStruct->ImportText(*Value, TempStructBuffer.GetData(), nullptr, PPF_UseDeprecatedProperties, &ErrorPipe, ScriptStruct->GetName()); 
 			}
 
 			ScriptStruct->DestroyStruct(TempStructBuffer.GetData());
diff --git a/Engine/Plugins/Runtime/RigVM/Source/RigVMDeveloper/Private/RigVMModel/RigVMSchema.cpp b/Engine/Plugins/Runtime/RigVM/Source/RigVMDeveloper/Private/RigVMModel/RigVMSchema.cpp
--- a/Engine/Plugins/Runtime/RigVM/Source/RigVMDeveloper/Private/RigVMModel/RigVMSchema.cpp
+++ b/Engine/Plugins/Runtime/RigVM/Source/RigVMDeveloper/Private/RigVMModel/RigVMSchema.cpp
@@ -245,6 +245,14 @@ bool URigVMSchema::ShouldUnfoldStruct(URigVMController* InController, const UStr
 	{
 		return false;
 	}
+	if(const UScriptStruct* ScriptStruct = Cast<UScriptStruct>(InStruct))
+	{
+		// We cannot unfold structs that have native import/exports, as their members may not correspond directly to pins
+		if ((ScriptStruct->StructFlags & (STRUCT_ImportTextItemNative | STRUCT_ExportTextItemNative)) != 0)
+		{
+			return false;
+		}
+	}
 
 	return true;
 }
diff --git a/Engine/Plugins/Runtime/RigVM/Source/RigVMEditor/Private/Editor/RigVMVariantDetailCustomization.cpp b/Engine/Plugins/Runtime/RigVM/Source/RigVMEditor/Private/Editor/RigVMVariantDetailCustomization.cpp
--- a/Engine/Plugins/Runtime/RigVM/Source/RigVMEditor/Private/Editor/RigVMVariantDetailCustomization.cpp
+++ b/Engine/Plugins/Runtime/RigVM/Source/RigVMEditor/Private/Editor/RigVMVariantDetailCustomization.cpp
@@ -27,7 +27,6 @@ void FRigVMVariantDetailCustomization::CustomizeHeader(TSharedRef<IPropertyHandl
 		
 	TArray<UObject*> Objects;
 	InStructPropertyHandle->GetOuterObjects(Objects);
-	ensure(Objects.Num() == 1); // This is in here to ensure we are only showing the modifier details in the blueprint editor
 
 	for (UObject* Object : Objects)
 	{
