[Backout] - CL37598848
Original CL Desc
-----------------------------------------------------------------
AnimNext: Call function trait

Adds a new trait that allows a function to be called at specific points in the animation graph update process.
Support for this is enabled by creating a compiler-generated wrapper event for each function call, then calling that at runtime.

Function library is now exposed directly in the outliner, rather than each reference.
Split graph & latent execution context data into seperate types, as we dont use latent pin info when calling functions from traits
Added support for soft object paths in traits (required for editor-only funciton reference data).
Added support for hidden pins on traits.

diff --git a/Engine/Plugins/Experimental/AnimNext/Source/AnimNext/Internal/AnimNextRigVMAsset.h b/Engine/Plugins/Experimental/AnimNext/Source/AnimNext/Internal/AnimNextRigVMAsset.h
--- a/Engine/Plugins/Experimental/AnimNext/Source/AnimNext/Internal/AnimNextRigVMAsset.h
+++ b/Engine/Plugins/Experimental/AnimNext/Source/AnimNext/Internal/AnimNextRigVMAsset.h
@@ -47,9 +47,6 @@ public:
 	// Get public variable defaults property bag
 	const FInstancedPropertyBag& GetPublicVariableDefaults() const { return PublicVariableDefaults; }
 
-	// Get the RigVM object we host
-	const URigVM* GetRigVM() const { return RigVM; }
-
 private:
 	// URigVMHost interface
 	virtual TArray<FRigVMExternalVariable> GetExternalVariablesImpl(bool bFallbackToBlueprint) const;
diff --git a/Engine/Plugins/Experimental/AnimNext/Source/AnimNext/Private/AnimNextModuleImpl.cpp b/Engine/Plugins/Experimental/AnimNext/Source/AnimNext/Private/AnimNextModuleImpl.cpp
--- a/Engine/Plugins/Experimental/AnimNext/Source/AnimNext/Private/AnimNextModuleImpl.cpp
+++ b/Engine/Plugins/Experimental/AnimNext/Source/AnimNext/Private/AnimNextModuleImpl.cpp
@@ -21,7 +21,6 @@
 #include "RigVMCore/RigVMRegistry.h"
 #include "Variables/AnimNextFieldPath.h"
 #include "Variables/AnimNextSoftFunctionPtr.h"
-#include "RigVMCore/RigVMGraphFunctionDefinition.h"
 
 #define LOCTEXT_NAMESPACE "AnimNextModule"
 
@@ -51,7 +50,6 @@ namespace UE::AnimNext
 			FUniversalObjectLocator::StaticStruct(),
 			FAnimNextFieldPath::StaticStruct(),
 			FAnimNextSoftFunctionPtr::StaticStruct(),
-			FRigVMGraphFunctionHeader::StaticStruct(),
 		};
 
 		RigVMRegistry.RegisterStructTypes(AllowedStructTypes);
diff --git a/Engine/Plugins/Experimental/AnimNext/Source/AnimNextEditor/Internal/Common/AnimNextAssetItemDetails.cpp b/Engine/Plugins/Experimental/AnimNext/Source/AnimNextEditor/Internal/Common/AnimNextAssetItemDetails.cpp
--- a/Engine/Plugins/Experimental/AnimNext/Source/AnimNextEditor/Internal/Common/AnimNextAssetItemDetails.cpp
+++ b/Engine/Plugins/Experimental/AnimNext/Source/AnimNextEditor/Internal/Common/AnimNextAssetItemDetails.cpp
@@ -13,8 +13,6 @@
 #include "ScopedTransaction.h"
 #include "RigVMModel/RigVMClient.h"
 #include "ToolMenus.h"
-#include "AssetRegistry/AssetRegistryModule.h"
-#include "AssetRegistry/IAssetRegistry.h"
 #include "Entries/AnimNextDataInterfaceEntry.h"
 #include "Framework/Commands/GenericCommands.h"
 #include "Module/AnimNextModule_EditorData.h"
@@ -29,11 +27,7 @@ namespace UE::AnimNext::Editor
 
 const FSlateBrush* FAnimNextAssetItemDetails::GetItemIcon(const FWorkspaceOutlinerItemExport& Export) const
 {
-	const FSoftObjectPath& AssetPath = Export.GetTopLevelAssetPath();
-	IAssetRegistry& AssetRegistry = FModuleManager::LoadModuleChecked<FAssetRegistryModule>(TEXT("AssetRegistry")).Get();
-	const FAssetData AssetData = AssetRegistry.GetAssetByObjectPath(AssetPath);
-	const UClass* AssetClass = FClassIconFinder::GetIconClassForAssetData(AssetData);
-	return FSlateIconFinder::FindIconForClass(AssetClass).GetIcon();
+	return FAppStyle::GetBrush(TEXT("LevelEditor.Tabs.Outliner"));
 }
 
 void FAnimNextAssetItemDetails::RegisterToolMenuExtensions()
@@ -202,26 +196,6 @@ void FAnimNextAssetItemDetails::RegisterToolMenuExtensions()
 					));
 			}
 		}
-
-		AssetSection.AddMenuEntry(
-			"AddFunction",
-			LOCTEXT("AddFunctionLabel", "Add Function"),
-			LOCTEXT("AddFunctionTooltip", "Add a function to this asset"),
-			FSlateIcon(FAppStyle::GetAppStyleSetName(), "GraphEditor.Function_16x"),
-			FUIAction(
-				FExecuteAction::CreateWeakLambda(WorkspaceItemContext, [WeakWorkspaceEditor = TWeakPtr<UE::Workspace::IWorkspaceEditor>(WorkspaceEditor), EditorData]()
-				{
-					FScopedTransaction Transaction(LOCTEXT("AddFunctionTransaction", "Add Function"));
-					URigVMLibraryNode* NewFunction = EditorData->AddFunction(TEXT("NewFunction"), true);
-
-					// Open the new function's graph
-					if(TSharedPtr<UE::Workspace::IWorkspaceEditor> PinnedWorkspaceEditor = WeakWorkspaceEditor.Pin())
-					{
-						UObject* EditorObject = EditorData->GetEditorObjectForRigVMGraph(NewFunction->GetContainedGraph());
-						PinnedWorkspaceEditor->OpenObjects({EditorObject});
-					}
-				})
-			));
 	}));
 }
 
diff --git a/Engine/Plugins/Experimental/AnimNext/Source/AnimNextEditor/Internal/Common/SRigVMFunctionPicker.h b/Engine/Plugins/Experimental/AnimNext/Source/AnimNextEditor/Internal/Common/SRigVMFunctionPicker.h
deleted file mode 100644
diff --git a/Engine/Plugins/Experimental/AnimNext/Source/AnimNextEditor/Private/AnimNextEditorModule.cpp b/Engine/Plugins/Experimental/AnimNext/Source/AnimNextEditor/Private/AnimNextEditorModule.cpp
--- a/Engine/Plugins/Experimental/AnimNext/Source/AnimNextEditor/Private/AnimNextEditorModule.cpp
+++ b/Engine/Plugins/Experimental/AnimNext/Source/AnimNextEditor/Private/AnimNextEditorModule.cpp
@@ -45,7 +45,6 @@
 #include "DataInterface/AnimNextDataInterface.h"
 #include "Entries/AnimNextEventGraphEntry.h"
 #include "Logging/MessageLog.h"
-#include "Graph/AnimNextEdGraphCustomization.h"
 #include "Module/RigUnit_AnimNextModuleEvents.h"
 #include "Variables/SVariablesView.h"
 #include "Variables/VariableOverrideCommands.h"
@@ -119,11 +118,6 @@ void FAnimNextEditorModule::StartupModule()
 					return MakeShared<FAnimNextEdGraphNodeCustomization>(InWorkspaceEditor);
 				}));
 
-			InDetailsView->RegisterInstancedCustomPropertyLayout(UAnimNextEdGraph::StaticClass(), FOnGetDetailCustomizationInstance::CreateLambda([InWorkspaceEditor]()
-				{
-					return MakeShared<FAnimNextEdGraphCustomization>();
-				}));
-		
 			TArray<UScriptStruct*> StructsToCustomize = {
 				TBaseStructure<FVector>::Get(),
 				TBaseStructure<FVector2D>::Get(),
diff --git a/Engine/Plugins/Experimental/AnimNext/Source/AnimNextEditor/Private/Common/AnimNextEdGraphNodeCustomization.cpp b/Engine/Plugins/Experimental/AnimNext/Source/AnimNextEditor/Private/Common/AnimNextEdGraphNodeCustomization.cpp
--- a/Engine/Plugins/Experimental/AnimNext/Source/AnimNextEditor/Private/Common/AnimNextEdGraphNodeCustomization.cpp
+++ b/Engine/Plugins/Experimental/AnimNext/Source/AnimNextEditor/Private/Common/AnimNextEdGraphNodeCustomization.cpp
@@ -377,7 +377,7 @@ void FAnimNextEdGraphNodeCustomization::PopulateCategory(IDetailLayoutBuilder& D
 				}
 
 				const FProperty* Property = PinnedPropertyHandle->GetProperty();
-				const FString ValueStr = FRigVMStruct::ExportToFullyQualifiedText(Property, (uint8*)InRawData, false);
+				const FString ValueStr = FRigVMStruct::ExportToFullyQualifiedText(Property, (uint8*)InRawData, true);
 
 				// Transform property path into pin path
 				FString PropertyPath(PinnedPropertyHandle->GetPropertyPath());
@@ -394,10 +394,7 @@ void FAnimNextEdGraphNodeCustomization::PopulateCategory(IDetailLayoutBuilder& D
 				// Prefix node path
 				const FString NodePath = EdGraphNode->GetModelNodePath();
 				PinPath = NodePath + TEXT(".") + PinPath;
-				if(!Controller->SetPinDefaultValue(PinPath, ValueStr))
-				{
-					UE_LOG(LogAnimation, Warning, TEXT("Could not set default value for pin '%s'"), *PinPath);
-				}
+				Controller->SetPinDefaultValue(PinPath, ValueStr);
 
 				return true;
 			});
diff --git a/Engine/Plugins/Experimental/AnimNext/Source/AnimNextEditor/Private/Common/AnimNextFunctionItemDetails.cpp b/Engine/Plugins/Experimental/AnimNext/Source/AnimNextEditor/Private/Common/AnimNextFunctionItemDetails.cpp
--- a/Engine/Plugins/Experimental/AnimNext/Source/AnimNextEditor/Private/Common/AnimNextFunctionItemDetails.cpp
+++ b/Engine/Plugins/Experimental/AnimNext/Source/AnimNextEditor/Private/Common/AnimNextFunctionItemDetails.cpp
@@ -31,34 +31,11 @@ bool FAnimNextFunctionItemDetails::HandleDoubleClick(const FToolMenuContext& Too
 			const TInstancedStruct<FWorkspaceOutlinerItemData>& Data = WorkspaceItemContext->SelectedExports[0].GetData();
 			if (Data.IsValid() && Data.GetScriptStruct() == FAnimNextGraphFunctionOutlinerData::StaticStruct())
 			{
-				const FAnimNextGraphFunctionOutlinerData& FunctionData = Data.Get<FAnimNextGraphFunctionOutlinerData>();
-				if (FunctionData.EditorObject.IsValid())
+				const FAnimNextGraphFunctionOutlinerData& GraphFunctionData = Data.Get<FAnimNextGraphFunctionOutlinerData>();
+				if (GraphFunctionData.EditorObject.IsValid())
 				{
-					WorkspaceEditor->OpenObjects({ FunctionData.EditorObject.Get() });
-					return true;
-				}
-			}
-		}
-	}
-	return false;
-}
-
-bool FAnimNextFunctionItemDetails::HandleSelected(const FToolMenuContext& ToolMenuContext) const
-{
-	const UWorkspaceItemMenuContext* WorkspaceItemContext = ToolMenuContext.FindContext<UWorkspaceItemMenuContext>();
-	const UAssetEditorToolkitMenuContext* AssetEditorContext = ToolMenuContext.FindContext<UAssetEditorToolkitMenuContext>();
-	if (WorkspaceItemContext && AssetEditorContext)
-	{
-		if (const TSharedPtr<UE::Workspace::IWorkspaceEditor> WorkspaceEditor = StaticCastSharedPtr<UE::Workspace::IWorkspaceEditor>(AssetEditorContext->Toolkit.Pin()))
-		{
-			const TInstancedStruct<FWorkspaceOutlinerItemData>& Data = WorkspaceItemContext->SelectedExports[0].GetData();
-			if (Data.IsValid() && Data.GetScriptStruct() == FAnimNextGraphFunctionOutlinerData::StaticStruct())
-			{
-				const FAnimNextGraphFunctionOutlinerData& FunctionData = Data.Get<FAnimNextGraphFunctionOutlinerData>();
-				if (FunctionData.EditorObject.IsValid())
-				{
-					WorkspaceEditor->SetDetailsObjects({ FunctionData.EditorObject.Get() });
-					return true;
+					WorkspaceEditor->OpenObjects({ GraphFunctionData.EditorObject.Get() });
+					return true;					
 				}
 			}
 		}
@@ -72,10 +49,10 @@ bool FAnimNextFunctionItemDetails::CanDelete(const FWorkspaceOutlinerItemExport&
 	const TInstancedStruct<FWorkspaceOutlinerItemData>& Data = Export.GetData();
 	if (Data.IsValid() && Data.GetScriptStruct() == FAnimNextGraphFunctionOutlinerData::StaticStruct())
 	{
-		const FAnimNextGraphFunctionOutlinerData& FunctionData = Export.GetData().Get<FAnimNextGraphFunctionOutlinerData>();
-		if (FunctionData.EditorObject.IsValid())
+		const FAnimNextGraphFunctionOutlinerData& CollapseGraphData = Export.GetData().Get<FAnimNextGraphFunctionOutlinerData>();
+		if (CollapseGraphData.EditorObject.IsValid())
 		{
-			URigVMEdGraph* EdGraph = FunctionData.EditorObject.Get();
+			URigVMEdGraph* EdGraph = CollapseGraphData.EditorObject.Get();
 			if (EdGraph->bAllowDeletion)
 			{
 				return true;
@@ -93,37 +70,18 @@ void FAnimNextFunctionItemDetails::Delete(TConstArrayView<FWorkspaceOutlinerItem
 		const TInstancedStruct<FWorkspaceOutlinerItemData>& Data = Export.GetData();
 		if (Data.IsValid() && Data.GetScriptStruct() == FAnimNextGraphFunctionOutlinerData::StaticStruct())
 		{
-			const FAnimNextGraphFunctionOutlinerData& FunctionData = Export.GetData().Get<FAnimNextGraphFunctionOutlinerData>();
-			if (FunctionData.EditorObject.IsValid())
+			const FAnimNextGraphFunctionOutlinerData& GraphData = Export.GetData().Get<FAnimNextGraphFunctionOutlinerData>();
+			if (GraphData.EditorObject.IsValid() && GraphData.EdGraphNode.IsValid())
 			{
-				if(FunctionData.EdGraphNode.IsValid())
-				{
-					const URigVMEdGraphNode* EdGraphNode = FunctionData.EdGraphNode.Get();
-					if (EdGraphNode->CanUserDeleteNode())
-					{
-						if (URigVMGraph* Model = EdGraphNode->GetModel())
-						{
-							if (URigVMNode* ModelNode = Model->FindNodeByName(*EdGraphNode->GetModelNodePath()))
-							{
-								FScopedTransaction Transaction(LOCTEXT("DeleteFunctionInOutliner", "Delete Function"));
-								EdGraphNode->GetController()->RemoveNode(ModelNode);
-							}
-						}
-					}
-				}
-				else
+				const URigVMEdGraphNode* EdGraphNode = GraphData.EdGraphNode.Get();
+				if (EdGraphNode->CanUserDeleteNode())
 				{
-					URigVMEdGraph* EdGraph = FunctionData.EditorObject.Get();
-					if (EdGraph->bAllowDeletion)
+					if (URigVMGraph* Model = EdGraphNode->GetModel())
 					{
-						if(IRigVMClientHost* ClientHost = EdGraph->GetImplementingOuter<IRigVMClientHost>())
+						if (URigVMNode* ModelNode = Model->FindNodeByName(*EdGraphNode->GetModelNodePath()))
 						{
-							if(URigVMController* Controller = ClientHost->GetController(ClientHost->GetLocalFunctionLibrary()))
-							{
-								FScopedTransaction Transaction(LOCTEXT("DeleteFunctionInOutliner", "Delete Function"));
-								URigVMLibraryNode* LibraryNode = CastChecked<URigVMLibraryNode>(EdGraph->GetModel()->GetOuter());
-								Controller->RemoveFunctionFromLibrary(LibraryNode->GetFName());
-							}
+							FScopedTransaction Transaction(LOCTEXT("DeleteFunctionInOutliner", "Delete Function"));
+							EdGraphNode->GetController()->RemoveNode(ModelNode);
 						}
 					}
 				}
@@ -137,10 +95,10 @@ bool FAnimNextFunctionItemDetails::CanRename(const FWorkspaceOutlinerItemExport&
 	const TInstancedStruct<FWorkspaceOutlinerItemData>& Data = Export.GetData();
 	if (Data.IsValid() && Data.GetScriptStruct() == FAnimNextGraphFunctionOutlinerData::StaticStruct())
 	{
-		const FAnimNextGraphFunctionOutlinerData& FunctionData = Export.GetData().Get<FAnimNextGraphFunctionOutlinerData>();
-		if (FunctionData.EditorObject.IsValid())
+		const FAnimNextGraphFunctionOutlinerData& CollapseGraphData = Export.GetData().Get<FAnimNextGraphFunctionOutlinerData>();
+		if (CollapseGraphData.EditorObject.IsValid())
 		{
-			URigVMEdGraph* EdGraph = FunctionData.EditorObject.Get();
+			URigVMEdGraph* EdGraph = CollapseGraphData.EditorObject.Get();
 			if (EdGraph->bAllowRenaming)
 			{
 				return true;
@@ -156,10 +114,10 @@ void FAnimNextFunctionItemDetails::Rename(const FWorkspaceOutlinerItemExport& Ex
 	const TInstancedStruct<FWorkspaceOutlinerItemData>& Data = Export.GetData();
 	if (Data.IsValid() && Data.GetScriptStruct() == FAnimNextGraphFunctionOutlinerData::StaticStruct())
 	{
-		const FAnimNextGraphFunctionOutlinerData& FunctionData = Export.GetData().Get<FAnimNextGraphFunctionOutlinerData>();
-		if (FunctionData.EditorObject.IsValid())
+		const FAnimNextGraphFunctionOutlinerData& CollapseGraphData = Export.GetData().Get<FAnimNextGraphFunctionOutlinerData>();
+		if (CollapseGraphData.EditorObject.IsValid())
 		{
-			URigVMEdGraph* EdGraph = FunctionData.EditorObject.Get();
+			URigVMEdGraph* EdGraph = CollapseGraphData.EditorObject.Get();
 			if (EdGraph->bAllowRenaming)
 			{
 				if (const UEdGraphSchema* GraphSchema = EdGraph->GetSchema())
@@ -202,10 +160,10 @@ UPackage* FAnimNextFunctionItemDetails::GetPackage(const FWorkspaceOutlinerItemE
 	const TInstancedStruct<FWorkspaceOutlinerItemData>& Data = Export.GetData();
 	if (Data.IsValid() && Data.GetScriptStruct() == FAnimNextGraphFunctionOutlinerData::StaticStruct())
 	{
-		const FAnimNextGraphFunctionOutlinerData& FunctionData = Data.Get<FAnimNextGraphFunctionOutlinerData>();
-		if (FunctionData.EditorObject.IsValid())
+		const FAnimNextGraphFunctionOutlinerData& GraphFunctionData = Data.Get<FAnimNextGraphFunctionOutlinerData>();
+		if (GraphFunctionData.EditorObject.IsValid())
 		{
-			return FunctionData.EditorObject->GetPackage();
+			return GraphFunctionData.EditorObject->GetPackage();
 		}
 	}
 	return nullptr;
@@ -213,7 +171,7 @@ UPackage* FAnimNextFunctionItemDetails::GetPackage(const FWorkspaceOutlinerItemE
 
 const FSlateBrush* FAnimNextFunctionItemDetails::GetItemIcon(const FWorkspaceOutlinerItemExport& Export) const
 {
-	return FAppStyle::GetBrush(TEXT("GraphEditor.Function_24x"));
+	return FAppStyle::GetBrush(TEXT("GraphEditor.EventGraph_24x"));
 }
 
 void FAnimNextFunctionItemDetails::RegisterToolMenuExtensions()
diff --git a/Engine/Plugins/Experimental/AnimNext/Source/AnimNextEditor/Private/Common/AnimNextFunctionItemDetails.h b/Engine/Plugins/Experimental/AnimNext/Source/AnimNextEditor/Private/Common/AnimNextFunctionItemDetails.h
--- a/Engine/Plugins/Experimental/AnimNext/Source/AnimNextEditor/Private/Common/AnimNextFunctionItemDetails.h
+++ b/Engine/Plugins/Experimental/AnimNext/Source/AnimNextEditor/Private/Common/AnimNextFunctionItemDetails.h
@@ -21,7 +21,6 @@ public:
 	virtual bool ValidateName(const FWorkspaceOutlinerItemExport& Export, const FText& InName, FText& OutErrorMessage) const override;
 	virtual UPackage* GetPackage(const FWorkspaceOutlinerItemExport& Export) const override;
 	virtual const FSlateBrush* GetItemIcon(const FWorkspaceOutlinerItemExport& Export) const override;
-	virtual bool HandleSelected(const FToolMenuContext& ToolMenuContext) const override;
 
 	static void RegisterToolMenuExtensions();
 	static void UnregisterToolMenuExtensions();
diff --git a/Engine/Plugins/Experimental/AnimNext/Source/AnimNextEditor/Private/Common/SActionMenu.cpp b/Engine/Plugins/Experimental/AnimNext/Source/AnimNextEditor/Private/Common/SActionMenu.cpp
--- a/Engine/Plugins/Experimental/AnimNext/Source/AnimNextEditor/Private/Common/SActionMenu.cpp
+++ b/Engine/Plugins/Experimental/AnimNext/Source/AnimNextEditor/Private/Common/SActionMenu.cpp
@@ -133,9 +133,7 @@ void SActionMenu::CollectAllAnimNextGraphActions(FGraphContextMenuBuilder& MenuB
 
 	TMap<FAssetData, FRigVMGraphFunctionHeaderArray> FunctionExports;
 	AnimNext::UncookedOnly::FUtils::GetExportedFunctionsFromAssetRegistry(UE::AnimNext::AnimNextPublicGraphFunctionsExportsRegistryTag, FunctionExports);
-	// TODO: Ideally we can filter functions by schema or execute context, but right now we dont expose the schema and function execute contexts are
-	// all FRigVMExecuteContext, rather than the 'most derived' context in the function.
-//	AnimNext::UncookedOnly::FUtils::GetExportedFunctionsFromAssetRegistry(UE::AnimNext::ControlRigAssetPublicGraphFunctionsExportsRegistryTag, FunctionExports);
+	AnimNext::UncookedOnly::FUtils::GetExportedFunctionsFromAssetRegistry(UE::AnimNext::ControlRigAssetPublicGraphFunctionsExportsRegistryTag, FunctionExports);
 
 	for (const auto& Export : FunctionExports.Array())
 	{
diff --git a/Engine/Plugins/Experimental/AnimNext/Source/AnimNextEditor/Private/Common/SRigVMFunctionPicker.cpp b/Engine/Plugins/Experimental/AnimNext/Source/AnimNextEditor/Private/Common/SRigVMFunctionPicker.cpp
deleted file mode 100644
diff --git a/Engine/Plugins/Experimental/AnimNext/Source/AnimNextEditor/Private/Graph/AnimNextEdGraphCustomization.cpp b/Engine/Plugins/Experimental/AnimNext/Source/AnimNextEditor/Private/Graph/AnimNextEdGraphCustomization.cpp
deleted file mode 100644
diff --git a/Engine/Plugins/Experimental/AnimNext/Source/AnimNextEditor/Private/Graph/AnimNextEdGraphCustomization.h b/Engine/Plugins/Experimental/AnimNext/Source/AnimNextEditor/Private/Graph/AnimNextEdGraphCustomization.h
deleted file mode 100644
diff --git a/Engine/Plugins/Experimental/AnimNext/Source/AnimNextEditor/Private/Variables/SVariablesView.cpp b/Engine/Plugins/Experimental/AnimNext/Source/AnimNextEditor/Private/Variables/SVariablesView.cpp
--- a/Engine/Plugins/Experimental/AnimNext/Source/AnimNextEditor/Private/Variables/SVariablesView.cpp
+++ b/Engine/Plugins/Experimental/AnimNext/Source/AnimNextEditor/Private/Variables/SVariablesView.cpp
@@ -20,7 +20,6 @@
 #include "AnimNextRigVMAsset.h"
 #include "IWorkspaceEditor.h"
 #include "WorkspaceAssetRegistryInfo.h"
-#include "Common/AnimNextFunctionItemDetails.h"
 #include "Outliner/VariablesOutlinerAssetItem.h"
 
 #define LOCTEXT_NAMESPACE "SVariablesView"
@@ -276,20 +275,6 @@ void SVariablesView::HandleWorkspaceOutlinerSelectionChanged(TConstArrayView<FWo
 				}
 			}
 		}
-		else if(Export.GetData().GetScriptStruct()->IsChildOf(FAnimNextGraphFunctionOutlinerData::StaticStruct()))
-		{
-			const FAnimNextGraphFunctionOutlinerData& EntryData = Export.GetData().Get<FAnimNextGraphFunctionOutlinerData>();
-			if(EntryData.EditorObject.IsValid())
-			{
-				if(URigVMEdGraph* EdGraph = EntryData.EditorObject.Get())
-				{
-					if(UAnimNextRigVMAsset* Asset = EdGraph->GetTypedOuter<UAnimNextRigVMAsset>())
-					{
-						Assets.Add(TSoftObjectPtr<UAnimNextRigVMAsset>(Asset));
-					}
-				}
-			}
-		}
 	}
 
 	VariablesOutliner->SetAssets(Assets);
diff --git a/Engine/Plugins/Experimental/AnimNext/Source/AnimNextUncookedOnly/Internal/AnimNextRigVMAssetEditorData.h b/Engine/Plugins/Experimental/AnimNext/Source/AnimNextUncookedOnly/Internal/AnimNextRigVMAssetEditorData.h
--- a/Engine/Plugins/Experimental/AnimNext/Source/AnimNextUncookedOnly/Internal/AnimNextRigVMAssetEditorData.h
+++ b/Engine/Plugins/Experimental/AnimNext/Source/AnimNextUncookedOnly/Internal/AnimNextRigVMAssetEditorData.h
@@ -43,7 +43,6 @@ namespace UE::AnimNext::Editor
 	class FVariableProxyCustomization;
 	struct FPublicVariablesImpl;
 	class SActionMenu;
-	class FCallFunctionSharedDataDetails;
 }
 
 namespace UE::AnimNext::Tests
@@ -130,9 +129,6 @@ public:
 	/** Adds a data interface to this asset */
 	UAnimNextDataInterfaceEntry* AddDataInterface(UAnimNextDataInterface* InDataInterface, bool bSetupUndoRedo = true, bool bPrintPythonCommand = true);
 
-	/** Adds a function to this asset */
-	URigVMLibraryNode* AddFunction(FName InFunctionName, bool bInMutable, bool bInSetupUndoRedo = true, bool bInPrintPythonCommand = true);
-
 	// Find an entry by name
 	UAnimNextRigVMAssetEntry* FindEntry(FName InName) const;
 
@@ -181,7 +177,6 @@ protected:
 	friend class UAnimNextDataInterfaceEntry;
 	friend struct UE::AnimNext::Editor::FPublicVariablesImpl;
 	friend UE::AnimNext::Editor::SActionMenu;
-	friend class UE::AnimNext::Editor::FCallFunctionSharedDataDetails;
 
 	// UObject interface
 	virtual void Serialize(FArchive& Ar) override;
@@ -439,9 +434,6 @@ protected:
 	// Remove an entry to this asset, taking into account external packaging status
 	void RemoveEntryInternal(UAnimNextRigVMAssetEntry* InEntry);
 
-	// Remove any programmatic graphs generated during compilation and consign them to the transient package
-	void RemoveProgrammaticGraphs(TArrayView<URigVMGraph*> InGraphs);
-
 	/** All entries in this asset - not saved, discovered at load time and also contains InternalEntries */
 	UPROPERTY(transient)
 	TArray<TObjectPtr<UAnimNextRigVMAssetEntry>> Entries;
diff --git a/Engine/Plugins/Experimental/AnimNext/Source/AnimNextUncookedOnly/Internal/UncookedOnlyUtils.h b/Engine/Plugins/Experimental/AnimNext/Source/AnimNextUncookedOnly/Internal/UncookedOnlyUtils.h
--- a/Engine/Plugins/Experimental/AnimNext/Source/AnimNextUncookedOnly/Internal/UncookedOnlyUtils.h
+++ b/Engine/Plugins/Experimental/AnimNext/Source/AnimNextUncookedOnly/Internal/UncookedOnlyUtils.h
@@ -98,6 +98,9 @@ struct FAnimNextAssetRegistryExports
 
 	UPROPERTY()
 	TArray<FAnimNextAssetRegistryExportedVariable> Variables;
+
+	UPROPERTY()
+	TArray<FRigVMGraphFunctionHeader> PublicHeaders;
 };
 
 namespace UE::AnimNext::UncookedOnly
@@ -161,7 +164,7 @@ struct ANIMNEXTUNCOOKEDONLY_API FUtils
 	static bool GetExportedVariablesFromAssetRegistry(TMap<FAssetData, FAnimNextAssetRegistryExports>& OutExports);
 
 	// Gets the functions that are exported to the asset registry for an asset
-	static bool GetExportedFunctionsForAsset(const FAssetData& InAsset, FName Tag, FRigVMGraphFunctionHeaderArray& OutExports);
+	static bool GetExportedFunctionsForAsset(const FAssetData& InAsset, FAnimNextAssetRegistryExports& OutExports);
 
 	// Gets all the functions that are exported to the asset registry for the specified Tag
 	static bool GetExportedFunctionsFromAssetRegistry(FName Tag, TMap<FAssetData, FRigVMGraphFunctionHeaderArray>& OutExports);
@@ -170,12 +173,6 @@ struct ANIMNEXTUNCOOKEDONLY_API FUtils
 	static void GetAssetVariables(const UAnimNextRigVMAssetEditorData* EditorData, FAnimNextAssetRegistryExports& OutExports);
 	static void GetAssetVariables(const UAnimNextRigVMAssetEditorData* EditorData, TSet<FAnimNextAssetRegistryExportedVariable>& OutExports);
 
-	// Gets the exported public functions that are used by a RigVM asset
-	static void GetAssetFunctions(const UAnimNextRigVMAssetEditorData* InEditorData, FRigVMGraphFunctionHeaderArray& OutExports);
-
-	// Gets the non-exported private functions that are used by a RigVM asset
-	static void GetAssetPrivateFunctions(const UAnimNextRigVMAssetEditorData* InEditorData, FRigVMGraphFunctionHeaderArray& OutExports);
-
 	// Gets the asset-registry information needed for representing the contained data into the Workspace Outliner
 	static void GetAssetOutlinerItems(const UAnimNextRigVMAssetEditorData* EditorData, FWorkspaceOutlinerItemExports& OutExports, FAssetRegistryTagsContext Context);
 	static void CreateSubGraphsOutlinerItemsRecursive(const UAnimNextRigVMAssetEditorData* EditorData, FWorkspaceOutlinerItemExports& OutExports, FWorkspaceOutlinerItemExport& ParentExport, URigVMEdGraph* RigVMEdGraph, FAssetRegistryTagsContext Context);
@@ -193,9 +190,6 @@ struct ANIMNEXTUNCOOKEDONLY_API FUtils
 #if WITH_EDITOR
 	static void OpenProgrammaticGraphs(UAnimNextRigVMAssetEditorData* EditorData, const TArray<URigVMGraph*>& ProgrammaticGraphs);
 #endif // WITH_EDITOR
-
-	// Make an event name that we use as a wrapper to call RigVM functions 
-	static FString MakeFunctionWrapperEventName(FName InFunctionName);
 };
 
 }
\ No newline at end of file
diff --git a/Engine/Plugins/Experimental/AnimNext/Source/AnimNextUncookedOnly/Private/AnimNextRigVMAssetEditorData.cpp b/Engine/Plugins/Experimental/AnimNext/Source/AnimNextUncookedOnly/Private/AnimNextRigVMAssetEditorData.cpp
--- a/Engine/Plugins/Experimental/AnimNext/Source/AnimNextUncookedOnly/Private/AnimNextRigVMAssetEditorData.cpp
+++ b/Engine/Plugins/Experimental/AnimNext/Source/AnimNextUncookedOnly/Private/AnimNextRigVMAssetEditorData.cpp
@@ -267,7 +267,13 @@ void UAnimNextRigVMAssetEditorData::GetAssetRegistryTags(FAssetRegistryTagsConte
 
 	{
 		FRigVMGraphFunctionHeaderArray FunctionExports;
-		UE::AnimNext::UncookedOnly::FUtils::GetAssetFunctions(this, FunctionExports);
+		for (const FRigVMGraphFunctionData& FunctionData : GraphFunctionStore.PublicFunctions)
+		{
+			if (FunctionData.CompilationData.IsValid())
+			{
+				FunctionExports.Headers.Add(FunctionData.Header);
+			}
+		}
 
 		FString TagValue;
 		const FArrayProperty* HeadersProperty = CastField<FArrayProperty>(FRigVMGraphFunctionHeaderArray::StaticStruct()->FindPropertyByName(TEXT("Headers")));
@@ -821,11 +827,7 @@ void UAnimNextRigVMAssetEditorData::RecompileVM()
 	{
 		FUtils::OpenProgrammaticGraphs(this, ProgrammaticGraphs);
 	}
-	else
 #endif
-	{
-		RemoveProgrammaticGraphs(ProgrammaticGraphs);
-	}
 
 #if WITH_EDITOR
 //	RefreshBreakpoints(EditorData);
@@ -838,17 +840,6 @@ void UAnimNextRigVMAssetEditorData::RecompileVM()
 	}
 }
 
-void UAnimNextRigVMAssetEditorData::RemoveProgrammaticGraphs(TArrayView<URigVMGraph*> InGraphs)
-{
-	FRigVMClient* VMClient = GetRigVMClient();
-	
-	for(URigVMGraph* Graph : InGraphs)
-	{
-		VMClient->RemoveController(Graph);
-		Graph->Rename(nullptr, GetTransientPackage(), REN_ForceNoResetLoaders | REN_DoNotDirty | REN_DontCreateRedirectors | REN_NonTransactional);
-	}
-}
-
 void UAnimNextRigVMAssetEditorData::RecompileVMIfRequired()
 {
 	if (bVMRecompilationRequired)
@@ -1737,12 +1728,6 @@ UAnimNextDataInterfaceEntry* UAnimNextRigVMAssetEditorData::AddDataInterface(UAn
 	return NewEntry;
 }
 
-URigVMLibraryNode* UAnimNextRigVMAssetEditorData::AddFunction(FName InFunctionName, bool bInMutable, bool bInSetupUndoRedo, bool bInPrintPythonCommand)
-{
-	URigVMController* Controller = RigVMClient.GetOrCreateController(GetLocalFunctionLibrary());
-	return Controller->AddFunctionToLibrary(InFunctionName, bInMutable, FVector2D::ZeroVector, bInSetupUndoRedo, bInPrintPythonCommand);
-}
-
 bool UAnimNextRigVMAssetEditorData::HasPublicVariables() const
 {
 	for(UAnimNextRigVMAssetEntry* Entry : Entries)
diff --git a/Engine/Plugins/Experimental/AnimNext/Source/AnimNextUncookedOnly/Private/UncookedOnlyUtils.cpp b/Engine/Plugins/Experimental/AnimNext/Source/AnimNextUncookedOnly/Private/UncookedOnlyUtils.cpp
--- a/Engine/Plugins/Experimental/AnimNext/Source/AnimNextUncookedOnly/Private/UncookedOnlyUtils.cpp
+++ b/Engine/Plugins/Experimental/AnimNext/Source/AnimNextUncookedOnly/Private/UncookedOnlyUtils.cpp
@@ -718,32 +718,22 @@ bool FUtils::GetExportedVariablesFromAssetRegistry(TMap<FAssetData, FAnimNextAss
 	return OutExports.Num() > 0;
 }
 
-void FUtils::GetAssetFunctions(const UAnimNextRigVMAssetEditorData* InEditorData, FRigVMGraphFunctionHeaderArray& OutExports)
+bool FUtils::GetExportedFunctionsForAsset(const FAssetData& InAsset, FAnimNextAssetRegistryExports& OutExports)
 {
-	for (const FRigVMGraphFunctionData& FunctionData : InEditorData->GraphFunctionStore.PublicFunctions)
+	const FString TagValue = InAsset.GetTagValueRef<FString>(UE::AnimNext::ExportsAnimNextAssetRegistryTag);
+	
+	bool bContainsFunctions = false;
+	FAnimNextAssetRegistryExports AssetExports;
+	if (FAnimNextAssetRegistryExports::StaticStruct()->ImportText(*TagValue, &AssetExports, nullptr, PPF_None, nullptr, FAnimNextAssetRegistryExports::StaticStruct()->GetName()) != nullptr)
 	{
-		if (FunctionData.CompilationData.IsValid())
+		if (AssetExports.PublicHeaders.Num() > 0)
 		{
-			OutExports.Headers.Add(FunctionData.Header);
+			OutExports = MoveTemp(AssetExports);
+			bContainsFunctions = true;
 		}
 	}
-}
 
-void FUtils::GetAssetPrivateFunctions(const UAnimNextRigVMAssetEditorData* InEditorData, FRigVMGraphFunctionHeaderArray& OutExports)
-{
-	for (const FRigVMGraphFunctionData& FunctionData : InEditorData->GraphFunctionStore.PrivateFunctions)
-	{
-		// Note: We dont check compilation data here as private functions are not compiled if they are not referenced
-		OutExports.Headers.Add(FunctionData.Header);
-	}
-}
-
-bool FUtils::GetExportedFunctionsForAsset(const FAssetData& InAsset, FName Tag, FRigVMGraphFunctionHeaderArray& OutExports)
-{
-	const FString TagValue = InAsset.GetTagValueRef<FString>(Tag);
-	const FArrayProperty* HeadersProperty = CastField<FArrayProperty>(FRigVMGraphFunctionHeaderArray::StaticStruct()->FindPropertyByName(TEXT("Headers")));
-	HeadersProperty->ImportText_Direct(*TagValue, &OutExports, nullptr, EPropertyPortFlags::PPF_None);
-	return OutExports.Headers.Num() > 0;
+	return bContainsFunctions;
 }
 
 bool FUtils::GetExportedFunctionsFromAssetRegistry(FName Tag, TMap<FAssetData, FRigVMGraphFunctionHeaderArray>& OutExports)
@@ -762,8 +752,7 @@ bool FUtils::GetExportedFunctionsFromAssetRegistry(FName Tag, TMap<FAssetData, F
 		{
 			if (AssetExports.Headers.Num() > 0)
 			{
-				FRigVMGraphFunctionHeaderArray& AssetArray = OutExports.FindOrAdd(Asset);
-				AssetArray.Headers.Append(MoveTemp(AssetExports.Headers));
+				OutExports.Add(Asset, MoveTemp(AssetExports));
 			}
 		}
 	}
@@ -853,7 +842,7 @@ void FUtils::GetAssetOutlinerItems(const UAnimNextRigVMAssetEditorData* EditorDa
 		}
 	}
 
-	CreateFunctionLibraryOutlinerItemsRecursive(EditorData, OutExports, AssetIdentifier, EditorData->GetRigVMGraphFunctionStore()->PublicFunctions, EditorData->GetRigVMGraphFunctionStore()->PrivateFunctions);
+	//CreateFunctionLibraryOutlinerItemsRecursive(EditorData, OutExports, AssetIdentifier, EditorData->GetRigVMGraphFunctionStore()->PublicFunctions, EditorData->GetRigVMGraphFunctionStore()->PrivateFunctions);
 }
 
 void ProcessPinAssetReferences(const URigVMPin* InPin, FWorkspaceOutlinerItemExports& OutExports, const FWorkspaceOutlinerItemExport& ParentExport)
@@ -984,8 +973,10 @@ void FUtils::CreateFunctionLibraryOutlinerItemsRecursive(const UAnimNextRigVMAss
 {
 	if (PrivateFunctions.Num() > 0 || PublicFunctions.Num() > 0)
 	{
-		CreateFunctionsOutlinerItemsRecursive(EditorData, OutExports, ParentExport, PrivateFunctions, false);
-		CreateFunctionsOutlinerItemsRecursive(EditorData, OutExports, ParentExport, PublicFunctions, true);
+		FWorkspaceOutlinerItemExport& Export = OutExports.Exports.Add_GetRef(FWorkspaceOutlinerItemExport(*GetFunctionLibraryDisplayName().ToString(), ParentExport));
+
+		CreateFunctionsOutlinerItemsRecursive(EditorData, OutExports, Export, PrivateFunctions, false);
+		CreateFunctionsOutlinerItemsRecursive(EditorData, OutExports, Export, PublicFunctions, true);
 	}
 }
 
@@ -1050,12 +1041,6 @@ void FUtils::OpenProgrammaticGraphs(UAnimNextRigVMAssetEditorData* EditorData, c
 	}
 }
 #endif // WITH_EDITOR
-
-FString FUtils::MakeFunctionWrapperEventName(FName InFunctionName)
-{
-	return TEXT("__InternalCall_") + InFunctionName.ToString();
-}
-
 }
 
 #undef LOCTEXT_NAMESPACE
\ No newline at end of file
diff --git a/Engine/Plugins/Experimental/AnimNextAnimGraph/Source/AnimNextAnimGraph/Internal/Graph/AnimNextAnimationGraph.h b/Engine/Plugins/Experimental/AnimNextAnimGraph/Source/AnimNextAnimGraph/Internal/Graph/AnimNextAnimationGraph.h
--- a/Engine/Plugins/Experimental/AnimNextAnimGraph/Source/AnimNextAnimGraph/Internal/Graph/AnimNextAnimationGraph.h
+++ b/Engine/Plugins/Experimental/AnimNextAnimGraph/Source/AnimNextAnimGraph/Internal/Graph/AnimNextAnimationGraph.h
@@ -140,11 +140,6 @@ protected:
 	UPROPERTY()
 	TArray<TObjectPtr<UObject>> GraphReferencedObjects;
 
-	// This is a list of all referenced soft objects in the graph shared data
-	// Used to serialize the soft objects correctly as we dont use FArchiveUObject
-	UPROPERTY()
-	TArray<FSoftObjectPath> GraphReferencedSoftObjects;
-
 	// The entry point that this graph defaults to using
 	UPROPERTY(EditAnywhere, Category = "Graph")
 	FName DefaultEntryPoint = FRigUnit_AnimNextGraphRoot::DefaultEntryPoint;
diff --git a/Engine/Plugins/Experimental/AnimNextAnimGraph/Source/AnimNextAnimGraph/Internal/Graph/AnimNextGraphContextData.h b/Engine/Plugins/Experimental/AnimNextAnimGraph/Source/AnimNextAnimGraph/Internal/Graph/AnimNextGraphContextData.h
--- a/Engine/Plugins/Experimental/AnimNextAnimGraph/Source/AnimNextAnimGraph/Internal/Graph/AnimNextGraphContextData.h
+++ b/Engine/Plugins/Experimental/AnimNextAnimGraph/Source/AnimNextAnimGraph/Internal/Graph/AnimNextGraphContextData.h
@@ -19,14 +19,24 @@ struct FAnimNextGraphContextData : public FAnimNextModuleContextData
 
 	FAnimNextGraphContextData() = default;
 
-	FAnimNextGraphContextData(FAnimNextModuleInstance* InModuleInstance, const FAnimNextGraphInstance* InInstance)
+	FAnimNextGraphContextData(FAnimNextModuleInstance* InModuleInstance, const FAnimNextGraphInstance* InInstance, const TConstArrayView<UE::AnimNext::FLatentPropertyHandle>& InLatentHandles, void* InDestinationBasePtr, bool bInIsFrozen)
 		: FAnimNextModuleContextData(InModuleInstance, InInstance)
+		, LatentHandles(InLatentHandles)
+		, DestinationBasePtr(InDestinationBasePtr)
+		, bIsFrozen(bInIsFrozen)
 	{
 	}
 
 	const FAnimNextGraphInstance& GetGraphInstance() const { return static_cast<const FAnimNextGraphInstance&>(GetDataInterfaceInstance()); }
+	const TConstArrayView<UE::AnimNext::FLatentPropertyHandle>& GetLatentHandles() const { return LatentHandles; }
+	void* GetDestinationBasePtr() const { return DestinationBasePtr; }
+	bool IsFrozen() const { return bIsFrozen; }
 
 private:
+	TConstArrayView<UE::AnimNext::FLatentPropertyHandle> LatentHandles;
+	void* DestinationBasePtr = nullptr;
+	bool bIsFrozen = false;
+
 	friend struct FAnimNextGraphInstance;
 	friend struct FAnimNextExecuteContext;
 };
diff --git a/Engine/Plugins/Experimental/AnimNextAnimGraph/Source/AnimNextAnimGraph/Internal/Graph/AnimNextGraphLatentPropertiesContextData.h b/Engine/Plugins/Experimental/AnimNextAnimGraph/Source/AnimNextAnimGraph/Internal/Graph/AnimNextGraphLatentPropertiesContextData.h
deleted file mode 100644
diff --git a/Engine/Plugins/Experimental/AnimNextAnimGraph/Source/AnimNextAnimGraph/Internal/TraitCore/TraitReader.h b/Engine/Plugins/Experimental/AnimNextAnimGraph/Source/AnimNextAnimGraph/Internal/TraitCore/TraitReader.h
--- a/Engine/Plugins/Experimental/AnimNextAnimGraph/Source/AnimNextAnimGraph/Internal/TraitCore/TraitReader.h
+++ b/Engine/Plugins/Experimental/AnimNextAnimGraph/Source/AnimNextAnimGraph/Internal/TraitCore/TraitReader.h
@@ -35,7 +35,7 @@ namespace UE::AnimNext
 			NodeInstanceDataTooLarge,	// Exceeded the maximum node instance data size, @see FNodeInstance::MAXIMUM_NODE_INSTANCE_DATA_SIZE
 		};
 
-		explicit FTraitReader(const TArray<TObjectPtr<UObject>>& InGraphReferencedObjects, const TArray<FSoftObjectPath>& InGraphReferencedSoftObjects, FArchive& Ar);
+		explicit FTraitReader(const TArray<TObjectPtr<UObject>>& InGraphReferencedObjects, FArchive& Ar);
 
 		[[nodiscard]] EErrorState ReadGraph(TArray<uint8>& GraphSharedData);
 
@@ -54,8 +54,6 @@ namespace UE::AnimNext
 		// FArchive implementation
 		virtual FArchive& operator<<(UObject*& Obj) override;
 		virtual FArchive& operator<<(FObjectPtr& Obj) override;
-		virtual FArchive& operator<<(FSoftObjectPath& Value) override;
-		virtual FArchive& operator<<(FSoftObjectPtr& Value) override;
 
 	private:
 		// Call first to read the graph shared data
@@ -64,9 +62,6 @@ namespace UE::AnimNext
 		// A list of UObject references within the graph
 		const TArray<TObjectPtr<UObject>>& GraphReferencedObjects;
 
-		// A list of soft object references within the graph
-		const TArray<FSoftObjectPath>& GraphReferencedSoftObjects;
-
 		// A list of node handles for each node within the archive
 		TArray<FNodeHandle> NodeHandles;
 
diff --git a/Engine/Plugins/Experimental/AnimNextAnimGraph/Source/AnimNextAnimGraph/Internal/TraitCore/TraitWriter.h b/Engine/Plugins/Experimental/AnimNextAnimGraph/Source/AnimNextAnimGraph/Internal/TraitCore/TraitWriter.h
--- a/Engine/Plugins/Experimental/AnimNextAnimGraph/Source/AnimNextAnimGraph/Internal/TraitCore/TraitWriter.h
+++ b/Engine/Plugins/Experimental/AnimNextAnimGraph/Source/AnimNextAnimGraph/Internal/TraitCore/TraitWriter.h
@@ -10,7 +10,6 @@
 #include "TraitCore/LatentPropertyHandle.h"
 #include "TraitCore/NodeHandle.h"
 #include "TraitCore/NodeTemplateRegistryHandle.h"
-#include "UObject/SoftObjectPath.h"
 
 namespace UE::AnimNext
 {
@@ -64,15 +63,10 @@ namespace UE::AnimNext
 		// Returns the list of referenced UObjects in this graph
 		[[nodiscard]] const TArray<UObject*>& GetGraphReferencedObjects() const;
 
-		// Returns the list of referenced soft objects in this graph
-		[[nodiscard]] const TArray<FSoftObjectPath>& GetGraphReferencedSoftObjects() const;
-
 		// FArchive implementation
 		virtual FArchive& operator<<(UObject*& Obj) override;
 		virtual FArchive& operator<<(FObjectPtr& Obj) override;
-		virtual FArchive& operator<<(FSoftObjectPath& Value) override;
-		virtual FArchive& operator<<(FSoftObjectPtr& Value) override;
-		
+
 	private:
 		struct FNodeMapping
 		{
@@ -94,7 +88,6 @@ namespace UE::AnimNext
 
 		// To track node writing
 		TArray<UObject*> GraphReferencedObjects;
-		TArray<FSoftObjectPath> GraphReferencedSoftObjects;
 		uint32 NumNodesWritten;
 		bool bIsNodeWriting;
 
diff --git a/Engine/Plugins/Experimental/AnimNextAnimGraph/Source/AnimNextAnimGraph/Internal/Traits/CallFunction.h b/Engine/Plugins/Experimental/AnimNextAnimGraph/Source/AnimNextAnimGraph/Internal/Traits/CallFunction.h
deleted file mode 100644
diff --git a/Engine/Plugins/Experimental/AnimNextAnimGraph/Source/AnimNextAnimGraph/Private/Graph/AnimNextAnimationGraph.cpp b/Engine/Plugins/Experimental/AnimNextAnimGraph/Source/AnimNextAnimGraph/Private/Graph/AnimNextAnimationGraph.cpp
--- a/Engine/Plugins/Experimental/AnimNextAnimGraph/Source/AnimNextAnimGraph/Private/Graph/AnimNextAnimationGraph.cpp
+++ b/Engine/Plugins/Experimental/AnimNextAnimGraph/Source/AnimNextAnimGraph/Private/Graph/AnimNextAnimationGraph.cpp
@@ -166,7 +166,7 @@ bool UAnimNextAnimationGraph::LoadFromArchiveBuffer(const TArray<uint8>& InShare
 
 	// Reconstruct our graph shared data
 	FMemoryReader GraphSharedDataArchive(InSharedDataArchiveBuffer);
-	FTraitReader TraitReader(GraphReferencedObjects, GraphReferencedSoftObjects, GraphSharedDataArchive);
+	FTraitReader TraitReader(GraphReferencedObjects, GraphSharedDataArchive);
 
 	const FTraitReader::EErrorState ErrorState = TraitReader.ReadGraph(SharedDataBuffer);
 	if (ErrorState == FTraitReader::EErrorState::None)
diff --git a/Engine/Plugins/Experimental/AnimNextAnimGraph/Source/AnimNextAnimGraph/Private/Graph/AnimNextGraphInstance.cpp b/Engine/Plugins/Experimental/AnimNextAnimGraph/Source/AnimNextAnimGraph/Private/Graph/AnimNextGraphInstance.cpp
--- a/Engine/Plugins/Experimental/AnimNextAnimGraph/Source/AnimNextAnimGraph/Private/Graph/AnimNextGraphInstance.cpp
+++ b/Engine/Plugins/Experimental/AnimNextAnimGraph/Source/AnimNextAnimGraph/Private/Graph/AnimNextGraphInstance.cpp
@@ -7,7 +7,6 @@
 #include "DataInterface/DataInterfaceStructAdapter.h"
 #include "Graph/AnimNextAnimationGraph.h"
 #include "Graph/AnimNextGraphContextData.h"
-#include "Graph/AnimNextGraphLatentPropertiesContextData.h"
 #include "TraitCore/ExecutionContext.h"
 #include "Graph/GC_GraphInstanceComponent.h"
 #include "Graph/RigUnit_AnimNextShimRoot.h"
@@ -182,7 +181,7 @@ void FAnimNextGraphInstance::ExecuteLatentPins(const TConstArrayView<UE::AnimNex
 		FAnimNextExecuteContext& AnimNextContext = ExtendedExecuteContext.GetPublicDataSafe<FAnimNextExecuteContext>();
 
 		// Insert our context data for the scope of execution
-		FAnimNextGraphLatentPropertiesContextData ContextData(ModuleInstance, this, LatentHandles, DestinationBasePtr, bIsFrozen);
+		FAnimNextGraphContextData ContextData(ModuleInstance, this, LatentHandles, DestinationBasePtr, bIsFrozen);
 		UE::AnimNext::FScopedExecuteContextData ContextDataScope(AnimNextContext, ContextData);
 
 		VM->ExecuteVM(ExtendedExecuteContext, FRigUnit_AnimNextShimRoot::EventName);
diff --git a/Engine/Plugins/Experimental/AnimNextAnimGraph/Source/AnimNextAnimGraph/Private/Graph/RigUnit_AnimNextGraphEvaluator.cpp b/Engine/Plugins/Experimental/AnimNextAnimGraph/Source/AnimNextAnimGraph/Private/Graph/RigUnit_AnimNextGraphEvaluator.cpp
--- a/Engine/Plugins/Experimental/AnimNextAnimGraph/Source/AnimNextAnimGraph/Private/Graph/RigUnit_AnimNextGraphEvaluator.cpp
+++ b/Engine/Plugins/Experimental/AnimNextAnimGraph/Source/AnimNextAnimGraph/Private/Graph/RigUnit_AnimNextGraphEvaluator.cpp
@@ -2,7 +2,6 @@
 
 #include "Graph/RigUnit_AnimNextGraphEvaluator.h"
 #include "Graph/AnimNextGraphContextData.h"
-#include "Graph/AnimNextGraphLatentPropertiesContextData.h"
 #include "TraitCore/LatentPropertyHandle.h"
 
 #include UE_INLINE_GENERATED_CPP_BY_NAME(RigUnit_AnimNextGraphEvaluator)
@@ -28,11 +27,11 @@ namespace UE::AnimNext::Private
 void FRigUnit_AnimNextGraphEvaluator::StaticExecute(FRigVMExtendedExecuteContext& RigVMExecuteContext, FRigVMMemoryHandleArray RigVMMemoryHandles, FRigVMPredicateBranchArray RigVMBranches)
 {
 	const FAnimNextExecuteContext& VMExecuteContext = RigVMExecuteContext.GetPublicData<FAnimNextExecuteContext>();
-	const FAnimNextGraphLatentPropertiesContextData& LatentPropertiesContextData = VMExecuteContext.GetContextData<FAnimNextGraphLatentPropertiesContextData>();
+	const FAnimNextGraphContextData& GraphContextData = VMExecuteContext.GetContextData<FAnimNextGraphContextData>();
 
-	const TConstArrayView<UE::AnimNext::FLatentPropertyHandle>& LatentHandles = LatentPropertiesContextData.GetLatentHandles();
-	uint8* DestinationBasePtr = (uint8*)LatentPropertiesContextData.GetDestinationBasePtr();
-	const bool bIsFrozen = LatentPropertiesContextData.IsFrozen();
+	const TConstArrayView<UE::AnimNext::FLatentPropertyHandle>& LatentHandles = GraphContextData.GetLatentHandles();
+	uint8* DestinationBasePtr = (uint8*)GraphContextData.GetDestinationBasePtr();
+	const bool bIsFrozen = GraphContextData.IsFrozen();
 
 	for (UE::AnimNext::FLatentPropertyHandle Handle : LatentHandles)
 	{
diff --git a/Engine/Plugins/Experimental/AnimNextAnimGraph/Source/AnimNextAnimGraph/Private/Graph/RigUnit_AnimNextRunAnimationGraph.cpp b/Engine/Plugins/Experimental/AnimNextAnimGraph/Source/AnimNextAnimGraph/Private/Graph/RigUnit_AnimNextRunAnimationGraph.cpp
--- a/Engine/Plugins/Experimental/AnimNextAnimGraph/Source/AnimNextAnimGraph/Private/Graph/RigUnit_AnimNextRunAnimationGraph.cpp
+++ b/Engine/Plugins/Experimental/AnimNextAnimGraph/Source/AnimNextAnimGraph/Private/Graph/RigUnit_AnimNextRunAnimationGraph.cpp
@@ -77,10 +77,6 @@ FRigUnit_AnimNextRunAnimationGraph_Execute()
 		Instance.BindPublicVariables(Hosts);
 	}
 
-	// Propagate delta time
-	FAnimNextExecuteContext& AnimNextExecuteContext = Instance.GetExtendedExecuteContext().GetPublicDataSafe<FAnimNextExecuteContext>();
-	AnimNextExecuteContext.SetDeltaTime(ExecuteContext.GetDeltaTime());
-
 	const UE::AnimNext::FReferencePose& RefPose = ReferencePose.ReferencePose.GetRef<UE::AnimNext::FReferencePose>();
 
 	// TODO: Currently forcing additive flag to false here
diff --git a/Engine/Plugins/Experimental/AnimNextAnimGraph/Source/AnimNextAnimGraph/Private/TraitCore/TraitReader.cpp b/Engine/Plugins/Experimental/AnimNextAnimGraph/Source/AnimNextAnimGraph/Private/TraitCore/TraitReader.cpp
--- a/Engine/Plugins/Experimental/AnimNextAnimGraph/Source/AnimNextAnimGraph/Private/TraitCore/TraitReader.cpp
+++ b/Engine/Plugins/Experimental/AnimNextAnimGraph/Source/AnimNextAnimGraph/Private/TraitCore/TraitReader.cpp
@@ -12,10 +12,9 @@
 
 namespace UE::AnimNext
 {
-	FTraitReader::FTraitReader(const TArray<TObjectPtr<UObject>>& InGraphReferencedObjects, const TArray<FSoftObjectPath>& InGraphReferencedSoftObjects, FArchive& Ar)
+	FTraitReader::FTraitReader(const TArray<TObjectPtr<UObject>>& InGraphReferencedObjects, FArchive& Ar)
 		: FArchiveProxy(Ar)
 		, GraphReferencedObjects(InGraphReferencedObjects)
-		, GraphReferencedSoftObjects(InGraphReferencedSoftObjects)
 	{
 	}
 
@@ -228,44 +227,6 @@ namespace UE::AnimNext
 		return FArchiveUObject::SerializeObjectPtr(*this, Obj);
 	}
 
-	FArchive& FTraitReader::operator<<(FSoftObjectPath& Value)
-	{
-		// Load the object index
-		int32 SoftObjectIndex = INDEX_NONE;
-		*this << SoftObjectIndex;
-
-		if (ensure(GraphReferencedSoftObjects.IsValidIndex(SoftObjectIndex)))
-		{
-			Value = GraphReferencedSoftObjects[SoftObjectIndex];
-		}
-		else
-		{
-			// Something went wrong, the reference list must have gotten out of sync which shouldn't happen
-			Value.Reset();
-		}
-
-		return *this;
-	}
-
-	FArchive& FTraitReader::operator<<(FSoftObjectPtr& Value)
-	{
-		// Load the object index
-		int32 SoftObjectIndex = INDEX_NONE;
-		*this << SoftObjectIndex;
-
-		if (ensure(GraphReferencedSoftObjects.IsValidIndex(SoftObjectIndex)))
-		{
-			Value = FSoftObjectPtr(GraphReferencedSoftObjects[SoftObjectIndex]);
-		}
-		else
-		{
-			// Something went wrong, the reference list must have gotten out of sync which shouldn't happen
-			Value.Reset();
-		}
-
-		return *this;
-	}
-
 	FNodeHandle FTraitReader::ResolveNodeHandle(FNodeHandle NodeHandle) const
 	{
 		if (!NodeHandle.IsValid())
diff --git a/Engine/Plugins/Experimental/AnimNextAnimGraph/Source/AnimNextAnimGraph/Private/TraitCore/TraitWriter.cpp b/Engine/Plugins/Experimental/AnimNextAnimGraph/Source/AnimNextAnimGraph/Private/TraitCore/TraitWriter.cpp
--- a/Engine/Plugins/Experimental/AnimNextAnimGraph/Source/AnimNextAnimGraph/Private/TraitCore/TraitWriter.cpp
+++ b/Engine/Plugins/Experimental/AnimNextAnimGraph/Source/AnimNextAnimGraph/Private/TraitCore/TraitWriter.cpp
@@ -75,7 +75,6 @@ namespace UE::AnimNext
 
 		bIsNodeWriting = true;
 		GraphReferencedObjects.Reset();
-		GraphReferencedSoftObjects.Reset();
 
 		// Serialize the node templates
 		TArray<FNodeTemplateRegistryHandle> NodeTemplateHandles;
@@ -230,11 +229,6 @@ namespace UE::AnimNext
 		return GraphReferencedObjects;
 	}
 
-	const TArray<FSoftObjectPath>& FTraitWriter::GetGraphReferencedSoftObjects() const
-	{
-		return GraphReferencedSoftObjects;
-	}
-
 	FArchive& FTraitWriter::operator<<(UObject*& Obj)
 	{
 		// Add our object for tracking
@@ -250,27 +244,5 @@ namespace UE::AnimNext
 	{
 		return FArchiveUObject::SerializeObjectPtr(*this, Obj);
 	}
-
-	FArchive& FTraitWriter::operator<<(FSoftObjectPath& Value)
-	{
-		// Add our object for tracking
-		int32 SoftObjectIndex = GraphReferencedSoftObjects.AddUnique(Value);
-
-		// Save our index, we'll use it to resolve the soft object on load
-		*this << SoftObjectIndex;
-
-		return *this;
-	}
-
-	FArchive& FTraitWriter::operator<<(FSoftObjectPtr& Value)
-	{
-		// Add our object for tracking
-		int32 SoftObjectIndex = GraphReferencedSoftObjects.AddUnique(Value.ToSoftObjectPath());
-
-		// Save our index, we'll use it to resolve the soft object on load
-		*this << SoftObjectIndex;
-
-		return *this;
-	}
 }
 #endif
diff --git a/Engine/Plugins/Experimental/AnimNextAnimGraph/Source/AnimNextAnimGraph/Private/Traits/CallFunction.cpp b/Engine/Plugins/Experimental/AnimNextAnimGraph/Source/AnimNextAnimGraph/Private/Traits/CallFunction.cpp
deleted file mode 100644
diff --git a/Engine/Plugins/Experimental/AnimNextAnimGraph/Source/AnimNextAnimGraph/Public/AnimNextAnimGraphStats.h b/Engine/Plugins/Experimental/AnimNextAnimGraph/Source/AnimNextAnimGraph/Public/AnimNextAnimGraphStats.h
--- a/Engine/Plugins/Experimental/AnimNextAnimGraph/Source/AnimNextAnimGraph/Public/AnimNextAnimGraphStats.h
+++ b/Engine/Plugins/Experimental/AnimNextAnimGraph/Source/AnimNextAnimGraph/Public/AnimNextAnimGraphStats.h
@@ -14,4 +14,3 @@ DECLARE_CYCLE_STAT_EXTERN(TEXT("AnimNext: Allocate Graph Instance"), STAT_AnimNe
 DECLARE_CYCLE_STAT_EXTERN(TEXT("AnimNext: Graph RigVM"), STAT_AnimNext_Graph_RigVM, STATGROUP_AnimNext, ANIMNEXTANIMGRAPH_API);
 DECLARE_CYCLE_STAT_EXTERN(TEXT("AnimNext: Execute Evaluation Program"), STAT_AnimNext_EvaluationProgram_Execute, STATGROUP_AnimNext, ANIMNEXTANIMGRAPH_API);
 
-DECLARE_CYCLE_STAT_EXTERN(TEXT("AnimNext: Call Function"), STAT_AnimNext_CallFunction, STATGROUP_AnimNext, ANIMNEXTANIMGRAPH_API);
diff --git a/Engine/Plugins/Experimental/AnimNextAnimGraph/Source/AnimNextAnimGraphEditor/Private/AnimNextAnimGraphEditorModule.cpp b/Engine/Plugins/Experimental/AnimNextAnimGraph/Source/AnimNextAnimGraphEditor/Private/AnimNextAnimGraphEditorModule.cpp
--- a/Engine/Plugins/Experimental/AnimNextAnimGraph/Source/AnimNextAnimGraphEditor/Private/AnimNextAnimGraphEditorModule.cpp
+++ b/Engine/Plugins/Experimental/AnimNextAnimGraph/Source/AnimNextAnimGraphEditor/Private/AnimNextAnimGraphEditorModule.cpp
@@ -8,7 +8,6 @@
 #include "EdGraphUtilities.h"
 #include "IAnimNextEditorModule.h"
 #include "IWorkspaceEditorModule.h"
-#include "PropertyEditorModule.h"
 #include "TraitStackEditor.h"
 #include "UncookedOnlyUtils.h"
 #include "Entries/AnimNextAnimationGraphEntry.h"
@@ -22,8 +21,6 @@
 #include "Graph/RigUnit_AnimNextGraphRoot.h"
 #include "Graph/RigVMTrait_AnimNextPublicVariablesEditor.h"
 #include "Graph/TraitEditorTabSummoner.h"
-#include "Traits/AnimNextCallFunctionSharedDataDetails.h"
-#include "Traits/CallFunction.h"
 #include "WorkflowOrientedApp/WorkflowTabManager.h"
 
 #define LOCTEXT_NAMESPACE "FAnimNextAnimGraphEditorModule"
@@ -44,11 +41,6 @@ class FAnimNextAnimGraphEditorModule : public IModuleInterface
 
 		FPublicVariablesImpl::Register();
 
-		FPropertyEditorModule& PropertyModule = FModuleManager::LoadModuleChecked<FPropertyEditorModule>("PropertyEditor");
-		PropertyModule.RegisterCustomPropertyTypeLayout(
-			FAnimNextCallFunctionSharedData::StaticStruct()->GetFName(),
-			FOnGetPropertyTypeCustomizationInstance::CreateLambda([] { return MakeShared<FCallFunctionSharedDataDetails>(); }));
-
 		Workspace::IWorkspaceEditorModule& WorkspaceEditorModule = FModuleManager::Get().LoadModuleChecked<Workspace::IWorkspaceEditorModule>("WorkspaceEditor");
 
 		// --- AnimNextAnimationGraph ---
@@ -110,13 +102,7 @@ class FAnimNextAnimGraphEditorModule : public IModuleInterface
 				WorkspaceEditorModule.UnregisterWorkspaceItemDetails(Workspace::FOutlinerItemDetailsId(FAnimNextAnimationGraphOutlinerData::StaticStruct()->GetFName()));
 			}
 		}
-
-		if(FModuleManager::Get().IsModuleLoaded("PropertyEditor"))
-		{
-			FPropertyEditorModule& PropertyModule = FModuleManager::LoadModuleChecked<FPropertyEditorModule>("PropertyEditor");
-			PropertyModule.UnregisterCustomClassLayout("AnimNextCallFunctionSharedData");
-		}
-
+		
 		IModularFeatures::Get().UnregisterModularFeature(ITraitStackEditor::ModularFeatureName, TraitStackEditor.Get());
 		TraitStackEditor.Reset();
 
diff --git a/Engine/Plugins/Experimental/AnimNextAnimGraph/Source/AnimNextAnimGraphEditor/Private/Traits/AnimNextCallFunctionSharedDataDetails.cpp b/Engine/Plugins/Experimental/AnimNextAnimGraph/Source/AnimNextAnimGraphEditor/Private/Traits/AnimNextCallFunctionSharedDataDetails.cpp
deleted file mode 100644
diff --git a/Engine/Plugins/Experimental/AnimNextAnimGraph/Source/AnimNextAnimGraphEditor/Private/Traits/AnimNextCallFunctionSharedDataDetails.h b/Engine/Plugins/Experimental/AnimNextAnimGraph/Source/AnimNextAnimGraphEditor/Private/Traits/AnimNextCallFunctionSharedDataDetails.h
deleted file mode 100644
diff --git a/Engine/Plugins/Experimental/AnimNextAnimGraph/Source/AnimNextAnimGraphTestSuite/Private/AnimNextAnimGraphTraitEventTest.cpp b/Engine/Plugins/Experimental/AnimNextAnimGraph/Source/AnimNextAnimGraphTestSuite/Private/AnimNextAnimGraphTraitEventTest.cpp
--- a/Engine/Plugins/Experimental/AnimNextAnimGraph/Source/AnimNextAnimGraphTestSuite/Private/AnimNextAnimGraphTraitEventTest.cpp
+++ b/Engine/Plugins/Experimental/AnimNextAnimGraph/Source/AnimNextAnimGraphTestSuite/Private/AnimNextAnimGraphTraitEventTest.cpp
@@ -228,7 +228,6 @@ bool FAnimationAnimNextRuntimeTest_GraphTraitEvent::RunTest(const FString& InPar
 		// Write our graph
 		TArray<uint8> GraphSharedDataArchiveBuffer;
 		TArray<TObjectPtr<UObject>> GraphReferencedObjects;
-		TArray<FSoftObjectPath> GraphReferencedSoftObjects;
 		{
 			FTraitWriter TraitWriter;
 
@@ -294,7 +293,6 @@ bool FAnimationAnimNextRuntimeTest_GraphTraitEvent::RunTest(const FString& InPar
 			AddErrorIfFalse(TraitWriter.GetErrorState() == FTraitWriter::EErrorState::None, "FAnimationAnimNextRuntimeTest_GraphTraitEvent -> Failed to write traits");
 			GraphSharedDataArchiveBuffer = TraitWriter.GetGraphSharedData();
 			GraphReferencedObjects = TraitWriter.GetGraphReferencedObjects();
-			GraphReferencedSoftObjects = TraitWriter.GetGraphReferencedSoftObjects();
 		}
 
 		// Read our graph
diff --git a/Engine/Plugins/Experimental/AnimNextAnimGraph/Source/AnimNextAnimGraphTestSuite/Private/AnimNextAnimGraphTraitInterfacesTest.cpp b/Engine/Plugins/Experimental/AnimNextAnimGraph/Source/AnimNextAnimGraphTestSuite/Private/AnimNextAnimGraphTraitInterfacesTest.cpp
--- a/Engine/Plugins/Experimental/AnimNextAnimGraph/Source/AnimNextAnimGraphTestSuite/Private/AnimNextAnimGraphTraitInterfacesTest.cpp
+++ b/Engine/Plugins/Experimental/AnimNextAnimGraph/Source/AnimNextAnimGraphTestSuite/Private/AnimNextAnimGraphTraitInterfacesTest.cpp
@@ -520,7 +520,6 @@ bool FAnimationAnimNextRuntimeTest_IHierarchy::RunTest(const FString& InParamete
 		// Write our graph
 		TArray<uint8> GraphSharedDataArchiveBuffer;
 		TArray<TObjectPtr<UObject>> GraphReferencedObjects;
-		TArray<FSoftObjectPath> GraphReferencedSoftObjects;
 		{
 			FTraitWriter TraitWriter;
 
@@ -589,7 +588,6 @@ bool FAnimationAnimNextRuntimeTest_IHierarchy::RunTest(const FString& InParamete
 			AddErrorIfFalse(TraitWriter.GetErrorState() == FTraitWriter::EErrorState::None, "FAnimationAnimNextRuntimeTest_IHierarchy -> Failed to write traits");
 			GraphSharedDataArchiveBuffer = TraitWriter.GetGraphSharedData();
 			GraphReferencedObjects = TraitWriter.GetGraphReferencedObjects();
-			GraphReferencedSoftObjects = TraitWriter.GetGraphReferencedSoftObjects();
 		}
 
 		// Read our graph
@@ -735,7 +733,6 @@ bool FAnimationAnimNextRuntimeTest_IUpdate::RunTest(const FString& InParameters)
 		// Write our graph
 		TArray<uint8> GraphSharedDataArchiveBuffer;
 		TArray<TObjectPtr<UObject>> GraphReferencedObjects;
-		TArray<FSoftObjectPath> GraphReferencedSoftObjects;
 		{
 			FTraitWriter TraitWriter;
 
@@ -789,7 +786,6 @@ bool FAnimationAnimNextRuntimeTest_IUpdate::RunTest(const FString& InParameters)
 			AddErrorIfFalse(TraitWriter.GetErrorState() == FTraitWriter::EErrorState::None, "FAnimationAnimNextRuntimeTest_IUpdate -> Failed to write traits");
 			GraphSharedDataArchiveBuffer = TraitWriter.GetGraphSharedData();
 			GraphReferencedObjects = TraitWriter.GetGraphReferencedObjects();
-			GraphReferencedSoftObjects = TraitWriter.GetGraphReferencedSoftObjects();
 		}
 
 		// Read our graph
@@ -878,7 +874,6 @@ bool FAnimationAnimNextRuntimeTest_IEvaluate::RunTest(const FString& InParameter
 		// Write our graph
 		TArray<uint8> GraphSharedDataArchiveBuffer;
 		TArray<TObjectPtr<UObject>> GraphReferencedObjects;
-		TArray<FSoftObjectPath> GraphReferencedSoftObjects;
 		{
 			FTraitWriter TraitWriter;
 
@@ -933,7 +928,6 @@ bool FAnimationAnimNextRuntimeTest_IEvaluate::RunTest(const FString& InParameter
 			AddErrorIfFalse(TraitWriter.GetErrorState() == FTraitWriter::EErrorState::None, "FAnimationAnimNextRuntimeTest_IEvaluate -> Failed to write traits");
 			GraphSharedDataArchiveBuffer = TraitWriter.GetGraphSharedData();
 			GraphReferencedObjects = TraitWriter.GetGraphReferencedObjects();
-			GraphReferencedSoftObjects = TraitWriter.GetGraphReferencedSoftObjects();
 		}
 
 		// Read our graph
@@ -1017,7 +1011,6 @@ bool FAnimationAnimNextRuntimeTest_IScopedInterface::RunTest(const FString& InPa
 		// Write our graph
 		TArray<uint8> GraphSharedDataArchiveBuffer;
 		TArray<TObjectPtr<UObject>> GraphReferencedObjects;
-		TArray<FSoftObjectPath> GraphReferencedSoftObjects;
 		{
 			FTraitWriter TraitWriter;
 
@@ -1071,7 +1064,6 @@ bool FAnimationAnimNextRuntimeTest_IScopedInterface::RunTest(const FString& InPa
 			AddErrorIfFalse(TraitWriter.GetErrorState() == FTraitWriter::EErrorState::None, "FAnimationAnimNextRuntimeTest_IScopedInterface -> Failed to write traits");
 			GraphSharedDataArchiveBuffer = TraitWriter.GetGraphSharedData();
 			GraphReferencedObjects = TraitWriter.GetGraphReferencedObjects();
-			GraphReferencedSoftObjects = TraitWriter.GetGraphReferencedSoftObjects();
 		}
 
 		// Read our graph
diff --git a/Engine/Plugins/Experimental/AnimNextAnimGraph/Source/AnimNextAnimGraphTestSuite/Private/AnimNextRuntimeTest.cpp b/Engine/Plugins/Experimental/AnimNextAnimGraph/Source/AnimNextAnimGraphTestSuite/Private/AnimNextRuntimeTest.cpp
--- a/Engine/Plugins/Experimental/AnimNextAnimGraph/Source/AnimNextAnimGraphTestSuite/Private/AnimNextRuntimeTest.cpp
+++ b/Engine/Plugins/Experimental/AnimNextAnimGraph/Source/AnimNextAnimGraphTestSuite/Private/AnimNextRuntimeTest.cpp
@@ -40,11 +40,10 @@ namespace UE::AnimNext
 		AnimationGraph.ExecuteDefinition = ExecuteDefinition;
 		AnimationGraph.SharedDataArchiveBuffer = SharedDataArchiveBuffer;
 		AnimationGraph.GraphReferencedObjects.Empty();
-		AnimationGraph.GraphReferencedSoftObjects.Empty();
 
 		// Reconstruct our graph shared data
 		FMemoryReader GraphSharedDataArchive(SharedDataArchiveBuffer);
-		FTraitReader TraitReader(AnimationGraph.GraphReferencedObjects, AnimationGraph.GraphReferencedSoftObjects, GraphSharedDataArchive);
+		FTraitReader TraitReader(AnimationGraph.GraphReferencedObjects, GraphSharedDataArchive);
 
 		const FTraitReader::EErrorState ErrorState = TraitReader.ReadGraph(AnimationGraph.SharedDataBuffer);
 		if (ErrorState == FTraitReader::EErrorState::None)
diff --git a/Engine/Plugins/Experimental/AnimNextAnimGraph/Source/AnimNextAnimGraphTestSuite/Private/AnimNextTraitBaseTest.cpp b/Engine/Plugins/Experimental/AnimNextAnimGraph/Source/AnimNextAnimGraphTestSuite/Private/AnimNextTraitBaseTest.cpp
--- a/Engine/Plugins/Experimental/AnimNextAnimGraph/Source/AnimNextAnimGraphTestSuite/Private/AnimNextTraitBaseTest.cpp
+++ b/Engine/Plugins/Experimental/AnimNextAnimGraph/Source/AnimNextAnimGraphTestSuite/Private/AnimNextTraitBaseTest.cpp
@@ -547,7 +547,6 @@ bool FAnimationAnimNextRuntimeTest_NodeLifetime::RunTest(const FString& InParame
 		// Write our graph
 		TArray<uint8> GraphSharedDataArchiveBuffer;
 		TArray<TObjectPtr<UObject>> GraphReferencedObjects;
-		TArray<FSoftObjectPath> GraphReferencedSoftObjects;
 		{
 			FTraitWriter TraitWriter;
 
@@ -580,7 +579,6 @@ bool FAnimationAnimNextRuntimeTest_NodeLifetime::RunTest(const FString& InParame
 			AddErrorIfFalse(TraitWriter.GetErrorState() == FTraitWriter::EErrorState::None, "FAnimationAnimNextRuntimeTest_NodeLifetime -> Failed to write traits");
 			GraphSharedDataArchiveBuffer = TraitWriter.GetGraphSharedData();
 			GraphReferencedObjects = TraitWriter.GetGraphReferencedObjects();
-			GraphReferencedSoftObjects = TraitWriter.GetGraphReferencedSoftObjects();
 		}
 
 		// Read our graph
@@ -758,7 +756,6 @@ bool FAnimationAnimNextRuntimeTest_TraitStackBinding::RunTest(const FString& InP
 		// Write our graph
 		TArray<uint8> GraphSharedDataArchiveBuffer;
 		TArray<TObjectPtr<UObject>> GraphReferencedObjects;
-		TArray<FSoftObjectPath> GraphReferencedSoftObjects;
 		{
 			FTraitWriter TraitWriter;
 
@@ -781,7 +778,6 @@ bool FAnimationAnimNextRuntimeTest_TraitStackBinding::RunTest(const FString& InP
 			AddErrorIfFalse(TraitWriter.GetErrorState() == FTraitWriter::EErrorState::None, "FAnimationAnimNextRuntimeTest_TraitStackBinding -> Failed to write traits");
 			GraphSharedDataArchiveBuffer = TraitWriter.GetGraphSharedData();
 			GraphReferencedObjects = TraitWriter.GetGraphReferencedObjects();
-			GraphReferencedSoftObjects = TraitWriter.GetGraphReferencedSoftObjects();
 		}
 
 		// Read our graph
@@ -1261,7 +1257,6 @@ bool FAnimationAnimNextRuntimeTest_TraitSerialization::RunTest(const FString& In
 		// Write our graph
 		TArray<uint8> GraphSharedDataArchiveBuffer;
 		TArray<TObjectPtr<UObject>> GraphReferencedObjects;
-		TArray<FSoftObjectPath> GraphReferencedSoftObjects;
 		{
 			FTraitWriter TraitWriter;
 
@@ -1353,7 +1348,6 @@ bool FAnimationAnimNextRuntimeTest_TraitSerialization::RunTest(const FString& In
 			AddErrorIfFalse(TraitWriter.GetErrorState() == FTraitWriter::EErrorState::None, "FAnimationAnimNextRuntimeTest_TraitSerialization -> Failed to write traits");
 			GraphSharedDataArchiveBuffer = TraitWriter.GetGraphSharedData();
 			GraphReferencedObjects = TraitWriter.GetGraphReferencedObjects();
-			GraphReferencedSoftObjects = TraitWriter.GetGraphReferencedSoftObjects();
 		}
 
 		// Clear out the node template registry to test registration on load
diff --git a/Engine/Plugins/Experimental/AnimNextAnimGraph/Source/AnimNextAnimGraphTestSuite/Private/AnimNextTraitEventTest.cpp b/Engine/Plugins/Experimental/AnimNextAnimGraph/Source/AnimNextAnimGraphTestSuite/Private/AnimNextTraitEventTest.cpp
--- a/Engine/Plugins/Experimental/AnimNextAnimGraph/Source/AnimNextAnimGraphTestSuite/Private/AnimNextTraitEventTest.cpp
+++ b/Engine/Plugins/Experimental/AnimNextAnimGraph/Source/AnimNextAnimGraphTestSuite/Private/AnimNextTraitEventTest.cpp
@@ -268,7 +268,6 @@ bool FAnimationAnimNextRuntimeTest_TraitEventRaising::RunTest(const FString& InP
 		// Write our graph
 		TArray<uint8> GraphSharedDataArchiveBuffer;
 		TArray<TObjectPtr<UObject>> GraphReferencedObjects;
-		TArray<FSoftObjectPath> GraphReferencedSoftObjects;
 		{
 			FTraitWriter TraitWriter;
 
@@ -291,7 +290,6 @@ bool FAnimationAnimNextRuntimeTest_TraitEventRaising::RunTest(const FString& InP
 			AddErrorIfFalse(TraitWriter.GetErrorState() == FTraitWriter::EErrorState::None, "FAnimationAnimNextRuntimeTest_TraitEventRaising -> Failed to write traits");
 			GraphSharedDataArchiveBuffer = TraitWriter.GetGraphSharedData();
 			GraphReferencedObjects = TraitWriter.GetGraphReferencedObjects();
-			GraphReferencedSoftObjects = TraitWriter.GetGraphReferencedSoftObjects();
 		}
 
 		// Read our graph
diff --git a/Engine/Plugins/Experimental/AnimNextAnimGraph/Source/AnimNextAnimGraphUncookedOnly/Internal/Graph/AnimNextAnimationGraph_EditorData.h b/Engine/Plugins/Experimental/AnimNextAnimGraph/Source/AnimNextAnimGraphUncookedOnly/Internal/Graph/AnimNextAnimationGraph_EditorData.h
--- a/Engine/Plugins/Experimental/AnimNextAnimGraph/Source/AnimNextAnimGraphUncookedOnly/Internal/Graph/AnimNextAnimationGraph_EditorData.h
+++ b/Engine/Plugins/Experimental/AnimNextAnimGraph/Source/AnimNextAnimGraphUncookedOnly/Internal/Graph/AnimNextAnimationGraph_EditorData.h
@@ -71,5 +71,4 @@ protected:
 	virtual bool CanAddNewEntry(TSubclassOf<UAnimNextRigVMAssetEntry> InClass) const override;
 	virtual TSubclassOf<UAssetUserData> GetAssetUserDataClass() const override;
 	virtual void InitializeAssetUserData() override;
-	virtual void GetProgrammaticGraphs(const FRigVMCompileSettings& InSettings, TArray<URigVMGraph*>& OutGraphs) override;
 };
diff --git a/Engine/Plugins/Experimental/AnimNextAnimGraph/Source/AnimNextAnimGraphUncookedOnly/Private/Graph/AnimNextAnimationGraph_EditorData.cpp b/Engine/Plugins/Experimental/AnimNextAnimGraph/Source/AnimNextAnimGraphUncookedOnly/Private/Graph/AnimNextAnimationGraph_EditorData.cpp
--- a/Engine/Plugins/Experimental/AnimNextAnimGraph/Source/AnimNextAnimGraphUncookedOnly/Private/Graph/AnimNextAnimationGraph_EditorData.cpp
+++ b/Engine/Plugins/Experimental/AnimNextAnimGraph/Source/AnimNextAnimGraphUncookedOnly/Private/Graph/AnimNextAnimationGraph_EditorData.cpp
@@ -17,15 +17,12 @@
 #include "Graph/RigUnit_AnimNextBeginExecution.h"
 #include "Graph/RigUnit_AnimNextShimRoot.h"
 #include "Graph/RigUnit_AnimNextTraitStack.h"
-#include "Logging/LogScopedVerbosityOverride.h"
 #include "Logging/MessageLog.h"
-#include "RigVMFunctions/Execution/RigVMFunction_UserDefinedEvent.h"
 #include "TraitCore/NodeTemplateBuilder.h"
 #include "TraitCore/TraitRegistry.h"
 #include "TraitCore/TraitWriter.h"
 #include "UObject/AssetRegistryTagsContext.h"
 #include "UObject/LinkerLoad.h"
-#include "Traits/CallFunction.h"
 
 namespace UE::AnimNext::UncookedOnly::Private
 {
@@ -133,10 +130,9 @@ namespace UE::AnimNext::UncookedOnly::Private
 		const TArray<URigVMPin*>& Pins = TraitStack.TraitEntries[TraitIndex].DecoratorEntryPin->GetSubPins();
 		for (const URigVMPin* Pin : Pins)
 		{
-			if (Pin->GetDirection() != ERigVMPinDirection::Input &&
-				Pin->GetDirection() != ERigVMPinDirection::Hidden)
+			if (Pin->GetDirection() != ERigVMPinDirection::Input)
 			{
-				continue;	// We only look for input or hidden pins
+				continue;	// We only look for input pins
 			}
 
 			if (Pin->GetFName() == PropertyName)
@@ -502,7 +498,6 @@ void UAnimNextAnimationGraph_EditorData::RecompileVM()
 		AnimationGraph->ExecuteDefinition = FAnimNextGraphEvaluatorExecuteDefinition();
 		AnimationGraph->SharedDataBuffer.Empty();
 		AnimationGraph->GraphReferencedObjects.Empty();
-		AnimationGraph->GraphReferencedSoftObjects.Empty();
 		AnimationGraph->DefaultEntryPoint = NAME_None;
 
 		FRigVMClient* VMClient = GetRigVMClient();
@@ -521,39 +516,35 @@ void UAnimNextAnimationGraph_EditorData::RecompileVM()
 			return;
 		}
 
-		TArray<URigVMGraph*> AnimGraphs;
-		TArray<URigVMGraph*> NonAnimGraphs;
-		for(URigVMGraph* SourceGraph : AllGraphs)
+		TArray<URigVMGraph*> TempGraphs;
+		for(const URigVMGraph* SourceGraph : AllGraphs)
 		{
 			// We use a temporary graph models to build our final graphs that we'll compile
-			if(SourceGraph->GetSchemaClass() == UAnimNextAnimationGraphSchema::StaticClass())
-			{
-				URigVMGraph* TempGraph = CastChecked<URigVMGraph>(StaticDuplicateObject(SourceGraph, this, NAME_None, RF_Transient));
-				UAnimNextController* TempController = CastChecked<UAnimNextController>(VMClient->GetOrCreateController(TempGraph));
-				TempGraph->SetFlags(RF_Transient);
-				AnimGraphs.Add(TempGraph);
-			}
-			else
-			{
-				NonAnimGraphs.Add(SourceGraph);
-			}
+			URigVMGraph* TempGraph = CastChecked<URigVMGraph>(StaticDuplicateObject(SourceGraph, GetTransientPackage(), NAME_None, RF_Transient));
+			TempGraph->SetFlags(RF_Transient);
+			TempGraphs.Add(TempGraph);
 		}
 
-		if(AnimGraphs.Num() > 0)
+		if(TempGraphs.Num() == 0)
 		{
-			UAnimNextController* TempController = CastChecked<UAnimNextController>(VMClient->GetOrCreateController(AnimGraphs[0]));
+			return;
+		}
 
-			UE::AnimNext::FTraitWriter TraitWriter;
+		UAnimNextController* TempController = CastChecked<UAnimNextController>(VMClient->GetOrCreateController(TempGraphs[0]));
 
-			FRigVMPinInfoArray LatentPins;
-			TMap<FName, URigVMPin*> LatentPinMapping;
-			TArray<Private::FTraitGraph> TraitGraphs;
+		UE::AnimNext::FTraitWriter TraitWriter;
 
-			// Build entry points and extract their required latent pins
-			for(const URigVMGraph* AnimGraph : AnimGraphs)
+		FRigVMPinInfoArray LatentPins;
+		TMap<FName, URigVMPin*> LatentPinMapping;
+		TArray<Private::FTraitGraph> TraitGraphs;
+
+		// Build entry points and extract their required latent pins
+		for(const URigVMGraph* TempGraph : TempGraphs)
+		{
+			if(TempGraph->GetSchemaClass() == UAnimNextAnimationGraphSchema::StaticClass())
 			{
 				// Gather our trait stacks
-				Private::FTraitGraph& TraitGraph = TraitGraphs.Add_GetRef(Private::CollectGraphInfo(AnimationGraph, AnimGraph, TempController->GetControllerForGraph(AnimGraph)));
+				Private::FTraitGraph& TraitGraph = TraitGraphs.Add_GetRef(Private::CollectGraphInfo(AnimationGraph, TempGraph, TempController->GetControllerForGraph(TempGraph)));
 				check(!TraitGraph.TraitStackNodes.IsEmpty());
 
 				FAnimNextGraphEntryPoint& EntryPoint = AnimationGraph->EntryPoints.AddDefaulted_GetRef();
@@ -574,73 +565,68 @@ void UAnimNextAnimationGraph_EditorData::RecompileVM()
 					EntryPoint.RootTraitHandle = FAnimNextEntryPointHandle(TraitGraph.TraitStackNodes[0].TraitStackNodeHandle);
 				}
 			}
+		}
 
-			// Set default entry point
-			if(AnimationGraph->EntryPoints.Num() > 0)
-			{
-				AnimationGraph->DefaultEntryPoint = AnimationGraph->EntryPoints[0].EntryPointName;
-			}
+		// Set default entry point
+		if(AnimationGraph->EntryPoints.Num() > 0)
+		{
+			AnimationGraph->DefaultEntryPoint = AnimationGraph->EntryPoints[0].EntryPointName;
+		}
 
-			// Remove our old root nodes
-			for (Private::FTraitGraph& TraitGraph : TraitGraphs)
-			{
-				URigVMController* GraphController = TempController->GetControllerForGraph(TraitGraph.RootNode->GetGraph());
-				GraphController->RemoveNode(TraitGraph.RootNode, false, false);
-			}
+		// Remove our old root nodes
+		for (Private::FTraitGraph& TraitGraph : TraitGraphs)
+		{
+			URigVMController* GraphController = TempController->GetControllerForGraph(TraitGraph.RootNode->GetGraph());
+			GraphController->RemoveNode(TraitGraph.RootNode, false, false);
+		}
 
-			if(LatentPins.Num() > 0)
-			{
-				// We need a unique method name to match our unique argument list
-				AnimationGraph->ExecuteDefinition = Private::GetGraphEvaluatorExecuteMethod(LatentPins);
+		if(LatentPins.Num() > 0)
+		{
+			// We need a unique method name to match our unique argument list
+			AnimationGraph->ExecuteDefinition = Private::GetGraphEvaluatorExecuteMethod(LatentPins);
+
+			// Add our runtime shim root node
+			URigVMUnitNode* TempShimRootNode = TempController->AddUnitNode(FRigUnit_AnimNextShimRoot::StaticStruct(), FRigUnit_AnimNextShimRoot::EventName, FVector2D::ZeroVector, FString(), false);
+			URigVMUnitNode* GraphEvaluatorNode = TempController->AddUnitNodeWithPins(FRigUnit_AnimNextGraphEvaluator::StaticStruct(), LatentPins, *AnimationGraph->ExecuteDefinition.MethodName, FVector2D::ZeroVector, FString(), false);
 
-				// Add our runtime shim root node
-				URigVMUnitNode* TempShimRootNode = TempController->AddUnitNode(FRigUnit_AnimNextShimRoot::StaticStruct(), FRigUnit_AnimNextShimRoot::EventName, FVector2D::ZeroVector, FString(), false);
-				URigVMUnitNode* GraphEvaluatorNode = TempController->AddUnitNodeWithPins(FRigUnit_AnimNextGraphEvaluator::StaticStruct(), LatentPins, *AnimationGraph->ExecuteDefinition.MethodName, FVector2D::ZeroVector, FString(), false);
+			// Link our shim and evaluator nodes together using the execution context
+			TempController->AddLink(
+				TempShimRootNode->FindPin(GET_MEMBER_NAME_STRING_CHECKED(FRigUnit_AnimNextShimRoot, ExecuteContext)),
+				GraphEvaluatorNode->FindPin(GET_MEMBER_NAME_STRING_CHECKED(FRigUnit_AnimNextGraphEvaluator, ExecuteContext)),
+				false);
 
-				// Link our shim and evaluator nodes together using the execution context
+			// Link our latent pins
+			for (const FRigVMPinInfo& LatentPin : LatentPins)
+			{
 				TempController->AddLink(
-					TempShimRootNode->FindPin(GET_MEMBER_NAME_STRING_CHECKED(FRigUnit_AnimNextShimRoot, ExecuteContext)),
-					GraphEvaluatorNode->FindPin(GET_MEMBER_NAME_STRING_CHECKED(FRigUnit_AnimNextGraphEvaluator, ExecuteContext)),
+					LatentPinMapping[LatentPin.Name],
+					GraphEvaluatorNode->FindPin(LatentPin.Name.ToString()),
 					false);
-
-				// Link our latent pins
-				for (const FRigVMPinInfo& LatentPin : LatentPins)
-				{
-					TempController->AddLink(
-						LatentPinMapping[LatentPin.Name],
-						GraphEvaluatorNode->FindPin(LatentPin.Name.ToString()),
-						false);
-				}
 			}
+		}
 
-			// Write our node shared data
-			TraitWriter.BeginNodeWriting();
+		// Write our node shared data
+		TraitWriter.BeginNodeWriting();
 
-			for(Private::FTraitGraph& TraitGraph : TraitGraphs)
+		for(Private::FTraitGraph& TraitGraph : TraitGraphs)
+		{
+			for (const Private::FTraitStackMapping& NodeMapping : TraitGraph.TraitStackNodes)
 			{
-				for (const Private::FTraitStackMapping& NodeMapping : TraitGraph.TraitStackNodes)
-				{
-					Private::WriteTraitProperties(TraitWriter, NodeMapping, TraitGraph.TraitStackNodes);
-				}
+				Private::WriteTraitProperties(TraitWriter, NodeMapping, TraitGraph.TraitStackNodes);
 			}
+		}
 
-			TraitWriter.EndNodeWriting();
-
-			// Cache our compiled metadata
-			AnimationGraph->SharedDataArchiveBuffer = TraitWriter.GetGraphSharedData();
-			AnimationGraph->GraphReferencedObjects = TraitWriter.GetGraphReferencedObjects();
-			AnimationGraph->GraphReferencedSoftObjects = TraitWriter.GetGraphReferencedSoftObjects();
+		TraitWriter.EndNodeWriting();
 
-			// Populate our runtime metadata
-			AnimationGraph->LoadFromArchiveBuffer(AnimationGraph->SharedDataArchiveBuffer);
-		}
+		// Cache our compiled metadata
+		AnimationGraph->SharedDataArchiveBuffer = TraitWriter.GetGraphSharedData();
+		AnimationGraph->GraphReferencedObjects = TraitWriter.GetGraphReferencedObjects();
 
-		TArray<URigVMGraph*> GraphsToCompile = MoveTemp(AnimGraphs);
-		GraphsToCompile.Append(NonAnimGraphs);
+		// Populate our runtime metadata
+		AnimationGraph->LoadFromArchiveBuffer(AnimationGraph->SharedDataArchiveBuffer);
 
-		UAnimNextController* Controller = CastChecked<UAnimNextController>(VMClient->GetOrCreateController(GraphsToCompile[0]));
 		URigVMCompiler* Compiler = URigVMCompiler::StaticClass()->GetDefaultObject<URigVMCompiler>();
-		Compiler->Compile(Settings, GraphsToCompile, Controller, AnimationGraph->VM, AnimationGraph->ExtendedExecuteContext, AnimationGraph->GetExternalVariables(), &PinToOperandMap);
+		Compiler->Compile(Settings, TempGraphs, TempController, AnimationGraph->VM, AnimationGraph->ExtendedExecuteContext, AnimationGraph->GetExternalVariables(), &PinToOperandMap);
 
 		// Initialize right away, in packaged builds we initialize during PostLoad
 		AnimationGraph->VM->Initialize(AnimationGraph->ExtendedExecuteContext);
@@ -663,24 +649,23 @@ void UAnimNextAnimationGraph_EditorData::RecompileVM()
 			RigVMCompiledEvent.Broadcast(AnimationGraph, AnimationGraph->VM, AnimationGraph->ExtendedExecuteContext);
 		}
 
-		// Remove the temp anim graphs
-		RemoveProgrammaticGraphs(AnimGraphs);
+		for(URigVMGraph* TempGraph : TempGraphs)
+		{
+			VMClient->RemoveController(TempGraph);
+		}
 
 		// Now that the graph has been re-compiled, re-allocate the previous live instances
 		AnimationGraph->ThawGraphInstances();
 	}
-
+		
 #if WITH_EDITOR
 	// Display programmatic graphs
 	if(CVarDumpProgrammaticGraphs.GetValueOnGameThread())
 	{
 		FUtils::OpenProgrammaticGraphs(this, ProgrammaticGraphs);
 	}
-	else
 #endif
-	{
-		RemoveProgrammaticGraphs(ProgrammaticGraphs);
-	}
+
 
 #if WITH_EDITOR
 //	RefreshBreakpoints(EditorData);
@@ -827,120 +812,4 @@ void UAnimNextAnimationGraph_EditorData::InitializeAssetUserData()
 	}
 
 	Super::InitializeAssetUserData();
-}
-
-void UAnimNextAnimationGraph_EditorData::GetProgrammaticGraphs(const FRigVMCompileSettings& InSettings, TArray<URigVMGraph*>& OutGraphs)
-{
-	using namespace UE::AnimNext::UncookedOnly;
-
-	Super::GetProgrammaticGraphs(InSettings, OutGraphs);
-
-	// Gather all 'call function' traits and create shim-calls for them.
-	// For the compiler to pick them up if they are not public we need a calling reference to the function from a graph
-	TArray<FRigVMGraphFunctionHeader> FunctionHeaders;
-
-	UAnimNextAnimationGraph* Asset = FUtils::GetAsset<UAnimNextAnimationGraph>(this);
-	FRigVMClient* VMClient = GetRigVMClient();
-	for (const URigVMGraph* Graph : VMClient->GetAllModels(false, false))
-	{
-		for(const URigVMNode* Node : Graph->GetNodes())
-		{
-			for(const URigVMPin* TraitPin : Node->GetTraitPins())
-			{
-				if (TraitPin->IsExecuteContext())
-				{
-					continue;
-				}
-
-				TSharedPtr<FStructOnScope> ScopedTrait = Node->GetTraitInstance(TraitPin->GetFName());
-				if (!ScopedTrait.IsValid())
-				{
-					continue;
-				}
-
-				const FRigVMTrait* Trait = (FRigVMTrait*)ScopedTrait->GetStructMemory();
-				UScriptStruct* TraitSharedInstanceData = Trait->GetTraitSharedDataStruct();
-				if (TraitSharedInstanceData == nullptr)
-				{
-					continue;
-				}
-
-				if(!TraitSharedInstanceData->IsChildOf(FAnimNextCallFunctionSharedData::StaticStruct()))
-				{
-					continue;
-				}
-
-				const FString DefaultValue = TraitPin->GetDefaultValue();
-				TInstancedStruct<FAnimNextCallFunctionSharedData> InstancedStruct = TInstancedStruct<FAnimNextCallFunctionSharedData>::Make();
-				FRigVMPinDefaultValueImportErrorContext ErrorPipe(ELogVerbosity::Verbose);
-				LOG_SCOPE_VERBOSITY_OVERRIDE(LogExec, ErrorPipe.GetMaxVerbosity());
-				TraitSharedInstanceData->ImportText(*DefaultValue, InstancedStruct.GetMutableMemory(), nullptr, PPF_SerializedAsImportText, &ErrorPipe, TraitSharedInstanceData->GetName());
-
-				const FRigVMGraphFunctionHeader& FunctionHeader = InstancedStruct.Get<FAnimNextCallFunctionSharedData>().FunctionHeader;
-				if(FunctionHeader.IsValid())
-				{
-					FunctionHeaders.Add(FunctionHeader);
-				}
-			}
-		}
-	}
-
-	if(FunctionHeaders.Num() > 0)
-	{
-		// Create all shim events for our traits to call
-		const bool bSetupUndoRedo = false;
-		URigVMGraph* WrapperGraph = NewObject<URigVMGraph>(this, NAME_None, RF_Transient);
-		UAnimNextController* Controller = CastChecked<UAnimNextController>(VMClient->GetOrCreateController(WrapperGraph));
-		FRigVMControllerNotifGuard NotifGuard(Controller);
-		bool bAddedWrapperEvent = true;
-		for(const FRigVMGraphFunctionHeader& FunctionHeader : FunctionHeaders)
-		{
-			URigVMLibraryNode* LibraryNode = Cast<URigVMLibraryNode>(FunctionHeader.LibraryPointer.GetNodeSoftPath().TryLoad());
-			if(LibraryNode == nullptr)
-			{
-				InSettings.ReportError(FString::Printf(TEXT("Could not find function '%s'"), *FunctionHeader.Name.ToString()));
-				continue;
-			}
-
-			// Create user-defined entry point
-			FString WrapperEventName = FUtils::MakeFunctionWrapperEventName(FunctionHeader.Name);
-			URigVMUnitNode* EventNode = Controller->AddUnitNode(FRigVMFunction_UserDefinedEvent::StaticStruct(), TEXT("Execute"), FVector2D::ZeroVector, FunctionHeader.Name.ToString(), bSetupUndoRedo);
-			if(EventNode == nullptr)
-			{
-				InSettings.ReportError(FString::Printf(TEXT("Could not spawn event node for function '%s'"), *FunctionHeader.Name.ToString()));
-				continue;
-			}
-			URigVMPin* EventNamePin = EventNode->FindPin(GET_MEMBER_NAME_STRING_CHECKED(FRigVMFunction_UserDefinedEvent, EventName));
-			if(EventNamePin == nullptr)
-			{
-				InSettings.ReportError(FString::Printf(TEXT("Could not find custom event name pin")));
-				continue;
-			}
-			Controller->SetPinDefaultValue(EventNamePin->GetPinPath(), WrapperEventName, true, bSetupUndoRedo);
-
-			// Call function
-			URigVMFunctionReferenceNode* FunctionNode = Controller->AddFunctionReferenceNode(LibraryNode, FVector2D::ZeroVector, FunctionHeader.Name.ToString(), bSetupUndoRedo);
-			if(FunctionNode == nullptr)
-			{
-				InSettings.ReportError(FString::Printf(TEXT("Could not spawn function node for function '%s'"), *FunctionHeader.Name.ToString()));
-				continue;
-			}
-
-			// Link up nodes
-			URigVMPin* ExecuteOutputPin = EventNode->FindPin(FRigVMStruct::ExecuteContextName.ToString());
-			URigVMPin* ExecuteInputPin = FunctionNode->FindPin(FRigVMStruct::ExecuteContextName.ToString());
-			if(!Controller->AddLink(ExecuteOutputPin, ExecuteInputPin, bSetupUndoRedo))
-			{
-				InSettings.ReportError(FString::Printf(TEXT("Could not link execute pins for function '%s'"), *FunctionHeader.Name.ToString()));
-				continue;
-			}
-
-			bAddedWrapperEvent = true;
-		}
-
-		if(bAddedWrapperEvent)
-		{
-			OutGraphs.Add(WrapperGraph);
-		}
-	}
 }
\ No newline at end of file
diff --git a/Engine/Plugins/Runtime/RigVM/Source/RigVM/Private/RigVMCore/RigVMStruct.cpp b/Engine/Plugins/Runtime/RigVM/Source/RigVM/Private/RigVMCore/RigVMStruct.cpp
--- a/Engine/Plugins/Runtime/RigVM/Source/RigVM/Private/RigVMCore/RigVMStruct.cpp
+++ b/Engine/Plugins/Runtime/RigVM/Source/RigVM/Private/RigVMCore/RigVMStruct.cpp
@@ -607,12 +607,10 @@ FString FRigVMStruct::ExportToFullyQualifiedText(const FProperty* InMemberProper
 	}
 	else
 	{
-		// Note we pass PPF_Delimited here to correctly escape strings
-		InMemberProperty->ExportTextItem_Direct(DefaultValue, InMemberMemoryPtr, nullptr, nullptr, PPF_Delimited);
+		InMemberProperty->ExportTextItem_Direct(DefaultValue, InMemberMemoryPtr, nullptr, nullptr, PPF_None);
 
 		if (CastField<FNameProperty>(InMemberProperty) != nullptr ||
-			CastField<FStrProperty>(InMemberProperty) != nullptr ||
-			CastField<FTextProperty>(InMemberProperty) != nullptr)
+			CastField<FStrProperty>(InMemberProperty) != nullptr)
 		{
 			if(bUseQuotes)
 			{
@@ -620,15 +618,11 @@ FString FRigVMStruct::ExportToFullyQualifiedText(const FProperty* InMemberProper
 				{
 					DefaultValue = TEXT("\"\"");
 				}
-				else if (DefaultValue[0] != TCHAR('\"') && DefaultValue[DefaultValue.Len() - 1] != TCHAR('\"'))
+				else
 				{
 					DefaultValue = FString::Printf(TEXT("\"%s\""), *DefaultValue);
 				}
 			}
-			else
-			{
-				DefaultValue = DefaultValue.TrimQuotes();
-			}
 		}
 	}
 
@@ -640,38 +634,26 @@ FString FRigVMStruct::ExportToFullyQualifiedText(const UScriptStruct* InStruct,
 	check(InStruct);
 	check(InStructMemoryPtr);
 
-	// If a struct has a custom text import/export, we must use its ExportText call rather than going per-property
-	if((InStruct->StructFlags & (STRUCT_ImportTextItemNative | STRUCT_ExportTextItemNative)) != 0)
-	{
-		FInstancedStruct Defaults;
-		Defaults.InitializeAs(InStruct);
-		FString Value;
-		InStruct->ExportText(Value, InStructMemoryPtr, Defaults.GetMemory(), nullptr, PPF_None, nullptr);
-		return Value;
-	}
-	else
+	TArray<FString> FieldValues;
+	for (TFieldIterator<FProperty> It(InStruct); It; ++It)
 	{
-		TArray<FString> FieldValues;
-		for (TFieldIterator<FProperty> It(InStruct); It; ++It)
-		{
-			if(It->HasAnyPropertyFlags(CPF_Transient))
-			{
-				continue;
-			}
-		
-			FString PropertyName = It->GetName();
-			const uint8* StructMemberMemoryPtr = It->ContainerPtrToValuePtr<uint8>(InStructMemoryPtr);
-			FString DefaultValue = ExportToFullyQualifiedText(*It, StructMemberMemoryPtr, bUseQuotes);
-			FieldValues.Add(FString::Printf(TEXT("%s=%s"), *PropertyName, *DefaultValue));
-		}
-
-		if (FieldValues.Num() == 0)
+		if(It->HasAnyPropertyFlags(CPF_Transient))
 		{
-			return TEXT("()");
+			continue;
 		}
+		
+		FString PropertyName = It->GetName();
+		const uint8* StructMemberMemoryPtr = It->ContainerPtrToValuePtr<uint8>(InStructMemoryPtr);
+		FString DefaultValue = ExportToFullyQualifiedText(*It, StructMemberMemoryPtr, bUseQuotes);
+		FieldValues.Add(FString::Printf(TEXT("%s=%s"), *PropertyName, *DefaultValue));
+	}
 
-		return FString::Printf(TEXT("(%s)"), *FString::Join(FieldValues, TEXT(",")));
+	if (FieldValues.Num() == 0)
+	{
+		return TEXT("()");
 	}
+
+	return FString::Printf(TEXT("(%s)"), *FString::Join(FieldValues, TEXT(",")));
 }
 
 FString FRigVMStruct::ExportToFullyQualifiedText(const UScriptStruct* InScriptStruct, const FName& InPropertyName, const uint8* InStructMemoryPointer, bool bUseQuotes) const
@@ -730,32 +712,28 @@ FName FRigVMStruct::GetNextAggregateName(const FName& InLastAggregatePinName) co
 TMap<FName, FString> FRigVMStruct::GetDefaultValues(UScriptStruct* InScriptStruct) const
 {
 	check(InScriptStruct);
-
+	
 	TMap<FName, FString> DefaultValues;
-	// If a struct has a custom text import/export, we cannot extract its property values individually
-	if((InScriptStruct->StructFlags & (STRUCT_ImportTextItemNative | STRUCT_ExportTextItemNative)) == 0)
+	for (TFieldIterator<FProperty> It(InScriptStruct); It; ++It)
 	{
-		for (TFieldIterator<FProperty> It(InScriptStruct); It; ++It)
+		if(It->HasAnyPropertyFlags(CPF_Transient))
 		{
-			if(It->HasAnyPropertyFlags(CPF_Transient))
-			{
-				continue;
-			}
+			continue;
+		}
 
 #if WITH_EDITOR
-			if(!It->HasMetaData(FRigVMStruct::InputMetaName) &&
-				!It->HasMetaData(FRigVMStruct::OutputMetaName) &&
-				!It->HasMetaData(FRigVMStruct::VisibleMetaName))
-			{
-				continue;
-			}
-#endif
-			
-			const FName PropertyName = It->GetFName();
-			const uint8* StructMemberMemoryPtr = It->ContainerPtrToValuePtr<uint8>(this);
-			const FString DefaultValue = ExportToFullyQualifiedText(*It, StructMemberMemoryPtr, false);
-			DefaultValues.Add(PropertyName, DefaultValue);
+		if(!It->HasMetaData(FRigVMStruct::InputMetaName) &&
+			!It->HasMetaData(FRigVMStruct::OutputMetaName) &&
+			!It->HasMetaData(FRigVMStruct::VisibleMetaName))
+		{
+			continue;
 		}
+#endif
+		
+		const FName PropertyName = It->GetFName();
+		const uint8* StructMemberMemoryPtr = It->ContainerPtrToValuePtr<uint8>(this);
+		const FString DefaultValue = ExportToFullyQualifiedText(*It, StructMemberMemoryPtr, false);
+		DefaultValues.Add(PropertyName, DefaultValue);
 	}
 
 	return DefaultValues;
diff --git a/Engine/Plugins/Runtime/RigVM/Source/RigVMDeveloper/Private/RigVMModel/RigVMBuildData.cpp b/Engine/Plugins/Runtime/RigVM/Source/RigVMDeveloper/Private/RigVMModel/RigVMBuildData.cpp
--- a/Engine/Plugins/Runtime/RigVM/Source/RigVMDeveloper/Private/RigVMModel/RigVMBuildData.cpp
+++ b/Engine/Plugins/Runtime/RigVM/Source/RigVMDeveloper/Private/RigVMModel/RigVMBuildData.cpp
@@ -491,7 +491,7 @@ void URigVMBuildData::ClearInvalidReferences()
 		return;
 	}
 	
-	Modify(false);
+	Modify();
 	
 	// check each function's each reference
 	int32 NumRemoved = 0;
diff --git a/Engine/Plugins/Runtime/RigVM/Source/RigVMDeveloper/Private/RigVMModel/RigVMController.cpp b/Engine/Plugins/Runtime/RigVM/Source/RigVMDeveloper/Private/RigVMModel/RigVMController.cpp
--- a/Engine/Plugins/Runtime/RigVM/Source/RigVMDeveloper/Private/RigVMModel/RigVMController.cpp
+++ b/Engine/Plugins/Runtime/RigVM/Source/RigVMDeveloper/Private/RigVMModel/RigVMController.cpp
@@ -19154,11 +19154,7 @@ void URigVMController::PostProcessDefaultValue(const URigVMPin* Pin, FString& Ou
 	}
 	if (Pin->IsStruct() || Pin->IsArray())
 	{
-		// We cannot sanitize struct text defaults that have native import/export, as they do not have to follow conventional syntax 
-		const UScriptStruct* ScriptStruct = Pin->GetScriptStruct();
-		const bool bHasNativeImportExport = ScriptStruct && ((ScriptStruct->StructFlags & (STRUCT_ImportTextItemNative | STRUCT_ExportTextItemNative)) != 0);
-
-		if (!OutDefaultValue.IsEmpty() && !bHasNativeImportExport)
+		if (!OutDefaultValue.IsEmpty())
 		{
 			if (OutDefaultValue[0] != TCHAR('(') || OutDefaultValue[OutDefaultValue.Len()-1] != TCHAR(')'))
 			{
diff --git a/Engine/Plugins/Runtime/RigVM/Source/RigVMDeveloper/Private/RigVMModel/RigVMPin.cpp b/Engine/Plugins/Runtime/RigVM/Source/RigVMDeveloper/Private/RigVMModel/RigVMPin.cpp
--- a/Engine/Plugins/Runtime/RigVM/Source/RigVMDeveloper/Private/RigVMModel/RigVMPin.cpp
+++ b/Engine/Plugins/Runtime/RigVM/Source/RigVMDeveloper/Private/RigVMModel/RigVMPin.cpp
@@ -870,48 +870,40 @@ FString URigVMPin::GetDefaultValue(const URigVMPin::FPinOverride& InOverride, bo
 	{
 		static const FString EmptyStructDefaultValue = TEXT("()");
 		
-		if((GetScriptStruct()->StructFlags & (STRUCT_ImportTextItemNative | STRUCT_ExportTextItemNative)) != 0)
+		// for trait pins, there are cases where a pin is not created for a property (see ShouldCreatePinForProperty())
+		// so we store the value of that property in the default value of the struct pin containing that property
+		// as a result, to retrieve the default value we need to combine the default value on the struct pin, with additional overrides in the available sub pins
+		if (SubPins.Num() > 0 || IsTraitPin())
 		{
-			// If the struct has a native import/export, then its default value must be used verbatim (subpins are not displayed)
-			return DefaultValue;
-		}
-		else
-		{
-			// for trait pins, there are cases where a pin is not created for a property (see ShouldCreatePinForProperty())
-			// so we store the value of that property in the default value of the struct pin containing that property
-			// as a result, to retrieve the default value we need to combine the default value on the struct pin, with additional overrides in the available sub pins
-			if (SubPins.Num() > 0 || IsTraitPin())
+			FString FinalDefaultValue = DefaultValue;
+			
+			// root trait pin store their default value in a separate property bag so that
+			// things like soft object ptr can be used and tracked in a uproperty
+			if (IsTraitPin() && IsRootPin())
 			{
-				FString FinalDefaultValue = DefaultValue;
-				
-				// root trait pin store their default value in a separate property bag so that
-				// things like soft object ptr can be used and tracked in a uproperty
-				if (IsTraitPin() && IsRootPin())
+				FRigVMTraitDefaultValueStruct* DefaultValueStructPtr = GetNode()->TraitDefaultValues.Find(GetName());
+				if (ensure(DefaultValueStructPtr))
 				{
-					FRigVMTraitDefaultValueStruct* DefaultValueStructPtr = GetNode()->TraitDefaultValues.Find(GetName());
-					if (ensure(DefaultValueStructPtr))
-					{
-						FinalDefaultValue = DefaultValueStructPtr->GetValue();
-					}
+					FinalDefaultValue = DefaultValueStructPtr->GetValue();
 				}
+			}
 
-				for (const URigVMPin* SubPin : SubPins)
+			for (const URigVMPin* SubPin : SubPins)
+			{
+				FString MemberDefaultValue = SubPin->GetDefaultValue(InOverride, bAdaptValueForPinType);
+				if (SubPin->IsStringType() && !MemberDefaultValue.IsEmpty())
 				{
-					FString MemberDefaultValue = SubPin->GetDefaultValue(InOverride, bAdaptValueForPinType);
-					if (SubPin->IsStringType() && !MemberDefaultValue.IsEmpty())
-					{
-						MemberDefaultValue = TEXT("\"") + MemberDefaultValue + TEXT("\"");
-					}
-					else if (MemberDefaultValue.IsEmpty() || MemberDefaultValue == TEXT("()"))
-					{
-						continue;
-					}
-
-					URigVMController::OverrideDefaultValueMember(SubPin->GetName(), MemberDefaultValue, FinalDefaultValue);
+					MemberDefaultValue = TEXT("\"") + MemberDefaultValue + TEXT("\"");
+				}
+				else if (MemberDefaultValue.IsEmpty() || MemberDefaultValue == TEXT("()"))
+				{
+					continue;
 				}
 
-				return !FinalDefaultValue.IsEmpty() ? FinalDefaultValue : EmptyStructDefaultValue;
+				URigVMController::OverrideDefaultValueMember(SubPin->GetName(), MemberDefaultValue, FinalDefaultValue);
 			}
+
+			return !FinalDefaultValue.IsEmpty() ? FinalDefaultValue : EmptyStructDefaultValue;
 		}
 
 		// special case certain pin types to adapt their values from
@@ -1065,9 +1057,8 @@ bool URigVMPin::IsValidDefaultValue(const FString& InDefaultValue) const
 			FRigVMPinDefaultValueImportErrorContext ErrorPipe(ELogVerbosity::Verbose);
 			{
 				// force logging to the error pipe for error detection
-				LOG_SCOPE_VERBOSITY_OVERRIDE(LogExec, ErrorPipe.GetMaxVerbosity());
-				// Note we use PPF_UseDeprecatedProperties as if a deprecated property is encountered this call will otherwise fail via the ErrorPipe
-				ScriptStruct->ImportText(*Value, TempStructBuffer.GetData(), nullptr, PPF_UseDeprecatedProperties, &ErrorPipe, ScriptStruct->GetName()); 
+				LOG_SCOPE_VERBOSITY_OVERRIDE(LogExec, ErrorPipe.GetMaxVerbosity()); 
+				ScriptStruct->ImportText(*Value, TempStructBuffer.GetData(), nullptr, PPF_None, &ErrorPipe, ScriptStruct->GetName()); 
 			}
 
 			ScriptStruct->DestroyStruct(TempStructBuffer.GetData());
diff --git a/Engine/Plugins/Runtime/RigVM/Source/RigVMDeveloper/Private/RigVMModel/RigVMSchema.cpp b/Engine/Plugins/Runtime/RigVM/Source/RigVMDeveloper/Private/RigVMModel/RigVMSchema.cpp
--- a/Engine/Plugins/Runtime/RigVM/Source/RigVMDeveloper/Private/RigVMModel/RigVMSchema.cpp
+++ b/Engine/Plugins/Runtime/RigVM/Source/RigVMDeveloper/Private/RigVMModel/RigVMSchema.cpp
@@ -245,14 +245,6 @@ bool URigVMSchema::ShouldUnfoldStruct(URigVMController* InController, const UStr
 	{
 		return false;
 	}
-	if(const UScriptStruct* ScriptStruct = Cast<UScriptStruct>(InStruct))
-	{
-		// We cannot unfold structs that have native import/exports, as their members may not correspond directly to pins
-		if ((ScriptStruct->StructFlags & (STRUCT_ImportTextItemNative | STRUCT_ExportTextItemNative)) != 0)
-		{
-			return false;
-		}
-	}
 
 	return true;
 }
diff --git a/Engine/Plugins/Runtime/RigVM/Source/RigVMEditor/Private/Editor/RigVMVariantDetailCustomization.cpp b/Engine/Plugins/Runtime/RigVM/Source/RigVMEditor/Private/Editor/RigVMVariantDetailCustomization.cpp
--- a/Engine/Plugins/Runtime/RigVM/Source/RigVMEditor/Private/Editor/RigVMVariantDetailCustomization.cpp
+++ b/Engine/Plugins/Runtime/RigVM/Source/RigVMEditor/Private/Editor/RigVMVariantDetailCustomization.cpp
@@ -27,6 +27,7 @@ void FRigVMVariantDetailCustomization::CustomizeHeader(TSharedRef<IPropertyHandl
 		
 	TArray<UObject*> Objects;
 	InStructPropertyHandle->GetOuterObjects(Objects);
+	ensure(Objects.Num() == 1); // This is in here to ensure we are only showing the modifier details in the blueprint editor
 
 	for (UObject* Object : Objects)
 	{
