AnimNext: Ensure that skeletal mesh component is a registered object type

diff --git a/Engine/Plugins/Experimental/AnimNextAnimGraph/Source/AnimNextAnimGraph/Private/AnimNextAnimGraphModule.cpp b/Engine/Plugins/Experimental/AnimNextAnimGraph/Source/AnimNextAnimGraph/Private/AnimNextAnimGraphModule.cpp
--- a/Engine/Plugins/Experimental/AnimNextAnimGraph/Source/AnimNextAnimGraph/Private/AnimNextAnimGraphModule.cpp
+++ b/Engine/Plugins/Experimental/AnimNextAnimGraph/Source/AnimNextAnimGraph/Private/AnimNextAnimGraphModule.cpp
@@ -8,6 +8,7 @@
 #include "Modules/ModuleManager.h"
 #include "RigVMCore/RigVMRegistry.h"
 #include "Animation/AnimSequence.h"
+#include "Components/SkeletalMeshComponent.h"
 #include "TraitInterfaces/IUpdate.h"
 #include "Graph/AnimNextAnimationGraph.h"
 #include "Graph/AnimNext_LODPose.h"
@@ -53,6 +54,7 @@ void FAnimNextAnimGraphModule::StartupModule()
 	static TPair<UClass*, FRigVMRegistry::ERegisterObjectOperation> const AllowedObjectTypes[] =
 	{
 		{ UAnimNextAnimationGraph::StaticClass(), FRigVMRegistry::ERegisterObjectOperation::Class },
+		{ USkeletalMeshComponent::StaticClass(), FRigVMRegistry::ERegisterObjectOperation::ClassAndParents },
 	};
 
 	RigVMRegistry.RegisterObjectTypes(AllowedObjectTypes);
