AnimNext: Cleaned up a little AnimNext DeadBlending ispc code using templated vector types. Fixed bug in ispc QuatLog implementation.

diff --git a/Engine/Plugins/Experimental/AnimNextAnimGraph/Source/AnimNextAnimGraph/Private/EvaluationVM/Tasks/DeadBlending.cpp b/Engine/Plugins/Experimental/AnimNextAnimGraph/Source/AnimNextAnimGraph/Private/EvaluationVM/Tasks/DeadBlending.cpp
--- a/Engine/Plugins/Experimental/AnimNextAnimGraph/Source/AnimNextAnimGraph/Private/EvaluationVM/Tasks/DeadBlending.cpp
+++ b/Engine/Plugins/Experimental/AnimNextAnimGraph/Source/AnimNextAnimGraph/Private/EvaluationVM/Tasks/DeadBlending.cpp
@@ -208,25 +208,25 @@ namespace UE::AnimNext::DeadBlending::Private
 
 #if UE_ANIM_NEXT_DEAD_BLENDING_ISPC
 		ispc::AnimNextDeadBlendingTransition(
-			(ispc::FVector4f*)BoneRotationDirections.GetData(),
-			(ispc::FVector*)Source.Translations.GetData(),
-			(ispc::FVector4*)Source.Rotations.GetData(),
-			(ispc::FVector*)Source.Scales3D.GetData(),
-			(ispc::FVector3f*)SourceBoneTranslationVelocities.GetData(),
-			(ispc::FVector3f*)SourceBoneRotationVelocities.GetData(),
-			(ispc::FVector3f*)SourceBoneScaleVelocities.GetData(),
-			(ispc::FVector3f*)SourceBoneTranslationDecayHalfLives.GetData(),
-			(ispc::FVector3f*)SourceBoneRotationDecayHalfLives.GetData(),
-			(ispc::FVector3f*)SourceBoneScaleDecayHalfLives.GetData(),
-			(ispc::FVector*)Dest.Translations.GetData(),
-			(ispc::FVector4*)Dest.Rotations.GetData(),
-			(ispc::FVector*)Dest.Scales3D.GetData(),
-			(ispc::FVector*)Curr.Translations.GetData(),
-			(ispc::FVector4*)Curr.Rotations.GetData(),
-			(ispc::FVector*)Curr.Scales3D.GetData(),
-			(ispc::FVector*)Prev.Translations.GetData(),
-			(ispc::FVector4*)Prev.Rotations.GetData(),
-			(ispc::FVector*)Prev.Scales3D.GetData(),
+			(float*)BoneRotationDirections.GetData(),
+			(double*)Source.Translations.GetData(),
+			(double*)Source.Rotations.GetData(),
+			(double*)Source.Scales3D.GetData(),
+			(float*)SourceBoneTranslationVelocities.GetData(),
+			(float*)SourceBoneRotationVelocities.GetData(),
+			(float*)SourceBoneScaleVelocities.GetData(),
+			(float*)SourceBoneTranslationDecayHalfLives.GetData(),
+			(float*)SourceBoneRotationDecayHalfLives.GetData(),
+			(float*)SourceBoneScaleDecayHalfLives.GetData(),
+			(double*)Dest.Translations.GetData(),
+			(double*)Dest.Rotations.GetData(),
+			(double*)Dest.Scales3D.GetData(),
+			(double*)Curr.Translations.GetData(),
+			(double*)Curr.Rotations.GetData(),
+			(double*)Curr.Scales3D.GetData(),
+			(double*)Prev.Translations.GetData(),
+			(double*)Prev.Rotations.GetData(),
+			(double*)Prev.Scales3D.GetData(),
 			LODBoneNum,
 			DeltaTime,
 			Parameters.ExtrapolationHalfLife,
@@ -320,19 +320,19 @@ namespace UE::AnimNext::DeadBlending::Private
 
 #if UE_ANIM_NEXT_DEAD_BLENDING_ISPC
 		ispc::AnimNextDeadBlendingTransitionStatic(
-			(ispc::FVector4f*)BoneRotationDirections.GetData(),
-			(ispc::FVector*)Source.Translations.GetData(),
-			(ispc::FVector4*)Source.Rotations.GetData(),
-			(ispc::FVector*)Source.Scales3D.GetData(),
-			(ispc::FVector3f*)SourceBoneTranslationVelocities.GetData(),
-			(ispc::FVector3f*)SourceBoneRotationVelocities.GetData(),
-			(ispc::FVector3f*)SourceBoneScaleVelocities.GetData(),
-			(ispc::FVector3f*)SourceBoneTranslationDecayHalfLives.GetData(),
-			(ispc::FVector3f*)SourceBoneRotationDecayHalfLives.GetData(),
-			(ispc::FVector3f*)SourceBoneScaleDecayHalfLives.GetData(),
-			(ispc::FVector*)Curr.Translations.GetData(),
-			(ispc::FVector4*)Curr.Rotations.GetData(),
-			(ispc::FVector*)Curr.Scales3D.GetData(),
+			(float*)BoneRotationDirections.GetData(),
+			(double*)Source.Translations.GetData(),
+			(double*)Source.Rotations.GetData(),
+			(double*)Source.Scales3D.GetData(),
+			(float*)SourceBoneTranslationVelocities.GetData(),
+			(float*)SourceBoneRotationVelocities.GetData(),
+			(float*)SourceBoneScaleVelocities.GetData(),
+			(float*)SourceBoneTranslationDecayHalfLives.GetData(),
+			(float*)SourceBoneRotationDecayHalfLives.GetData(),
+			(float*)SourceBoneScaleDecayHalfLives.GetData(),
+			(double*)Curr.Translations.GetData(),
+			(double*)Curr.Rotations.GetData(),
+			(double*)Curr.Scales3D.GetData(),
 			LODBoneNum,
 			Parameters.ExtrapolationHalfLifeMin);
 #else
@@ -385,19 +385,19 @@ namespace UE::AnimNext::DeadBlending::Private
 
 #if UE_ANIM_NEXT_DEAD_BLENDING_ISPC
 		ispc::AnimNextDeadBlendingApply(
-			(ispc::FVector*)Dest.Translations.GetData(),
-			(ispc::FVector4*)Dest.Rotations.GetData(),
-			(ispc::FVector*)Dest.Scales3D.GetData(),
-			(ispc::FVector4f*)BoneRotationDirections.GetData(),
-			(ispc::FVector*)Source.Translations.GetData(),
-			(ispc::FVector4*)Source.Rotations.GetData(),
-			(ispc::FVector*)Source.Scales3D.GetData(),
-			(ispc::FVector3f*)SourceBoneTranslationVelocities.GetData(),
-			(ispc::FVector3f*)SourceBoneRotationVelocities.GetData(),
-			(ispc::FVector3f*)SourceBoneScaleVelocities.GetData(),
-			(ispc::FVector3f*)SourceBoneTranslationDecayHalfLives.GetData(),
-			(ispc::FVector3f*)SourceBoneRotationDecayHalfLives.GetData(),
-			(ispc::FVector3f*)SourceBoneScaleDecayHalfLives.GetData(),
+			(double*)Dest.Translations.GetData(),
+			(double*)Dest.Rotations.GetData(),
+			(double*)Dest.Scales3D.GetData(),
+			(float*)BoneRotationDirections.GetData(),
+			(double*)Source.Translations.GetData(),
+			(double*)Source.Rotations.GetData(),
+			(double*)Source.Scales3D.GetData(),
+			(float*)SourceBoneTranslationVelocities.GetData(),
+			(float*)SourceBoneRotationVelocities.GetData(),
+			(float*)SourceBoneScaleVelocities.GetData(),
+			(float*)SourceBoneTranslationDecayHalfLives.GetData(),
+			(float*)SourceBoneRotationDecayHalfLives.GetData(),
+			(float*)SourceBoneScaleDecayHalfLives.GetData(),
 			LODBoneNum,
 			Alpha,
 			TimeSinceTransition);
diff --git a/Engine/Plugins/Experimental/AnimNextAnimGraph/Source/AnimNextAnimGraph/Private/EvaluationVM/Tasks/DeadBlending.ispc b/Engine/Plugins/Experimental/AnimNextAnimGraph/Source/AnimNextAnimGraph/Private/EvaluationVM/Tasks/DeadBlending.ispc
--- a/Engine/Plugins/Experimental/AnimNextAnimGraph/Source/AnimNextAnimGraph/Private/EvaluationVM/Tasks/DeadBlending.ispc
+++ b/Engine/Plugins/Experimental/AnimNextAnimGraph/Source/AnimNextAnimGraph/Private/EvaluationVM/Tasks/DeadBlending.ispc
@@ -1,160 +1,291 @@
 // Copyright Epic Games, Inc. All Rights Reserved.
 
-#include "Math/Vector.isph"
-#include "Math/Quat.isph"
-
 #define UE_LN2 (0.69314718056f)
 #define UE_SMALL_NUMBER (1.e-8f)
 #define UE_KINDA_SMALL_NUMBER (1.e-4f)
 
-inline FVector CastToVector(const FVector3f& V)
+typedef float<3> Vector3f;
+typedef float<4> Quat4f;
+typedef double<3> Vector3d;
+typedef double<4> Quat4d;
+
+inline Vector3f Vector3fLoad(uniform const float Data[], const int Index)
 {
-	const FVector Out = { {
-		V.V[0],
-		V.V[1],
-		V.V[2],
-	} };
+	const Vector3f X = {
+#pragma ignore warning(perf)
+		Data[Index * 3 + 0],
+#pragma ignore warning(perf)
+		Data[Index * 3 + 1],
+#pragma ignore warning(perf)
+		Data[Index * 3 + 2],
+	};
+
+	return X;
+}
 
-	return Out;
+inline Quat4f Quat4fLoad(uniform const float Data[], const int Index)
+{
+	const Quat4f X = {
+#pragma ignore warning(perf)
+		Data[Index * 4 + 0],
+#pragma ignore warning(perf)
+		Data[Index * 4 + 1],
+#pragma ignore warning(perf)
+		Data[Index * 4 + 2],
+#pragma ignore warning(perf)
+		Data[Index * 4 + 3],
+	};
+
+	return X;
 }
 
-inline FVector3f CastToVector3f(const FVector& V)
+inline Vector3d Vector3dLoad(uniform const double Data[], const int Index)
 {
-	const FVector3f Out = { {
-		V.V[0],
-		V.V[1],
-		V.V[2],
-	} };
-    
-	return Out;
+	const Vector3d X = {
+#pragma ignore warning(perf)
+		Data[Index * 3 + 0],
+#pragma ignore warning(perf)
+		Data[Index * 3 + 1],
+#pragma ignore warning(perf)
+		Data[Index * 3 + 2],
+	};
+
+	return X;
 }
 
-inline FVector4f CastToVector4f(const FVector4& V)
+inline Quat4d Quat4dLoad(uniform const double Data[], const int Index)
 {
-	const FVector4f Out = { {
-		V.V[0],
-		V.V[1],
-		V.V[2],
-		V.V[3],
-	} };
-    
-	return Out;
+	const Quat4d X = {
+#pragma ignore warning(perf)
+		Data[Index * 4 + 0],
+#pragma ignore warning(perf)
+		Data[Index * 4 + 1],
+#pragma ignore warning(perf)
+		Data[Index * 4 + 2],
+#pragma ignore warning(perf)
+		Data[Index * 4 + 3],
+	};
+
+	return X;
 }
 
-inline FVector4 CastToVector4(const FVector4f& V)
+inline void Vector3fSave(uniform float Data[], const int Index, const Vector3f X)
 {
-	const FVector4 Out = { {
-		V.V[0],
-		V.V[1],
-		V.V[2],
-		V.V[3],
-	} };
-    
-	return Out;
+#pragma ignore warning(perf)
+	Data[Index * 3 + 0] = X[0];
+#pragma ignore warning(perf)
+	Data[Index * 3 + 1] = X[1];
+#pragma ignore warning(perf)
+	Data[Index * 3 + 2] = X[2];
 }
 
-inline FVector3f VectorDivMax(const float S, const FVector3f& W, const uniform float Epsilon = UE_SMALL_NUMBER)
+inline void Quat4fSave(uniform float Data[], const int Index, const Quat4f X)
 {
-	const FVector3f Out = { {
-		S / max(W.V[0], Epsilon),
-		S / max(W.V[1], Epsilon),
-		S / max(W.V[2], Epsilon),
-	} };
-    
-	return Out;
+#pragma ignore warning(perf)
+	Data[Index * 4 + 0] = X[0];
+#pragma ignore warning(perf)
+	Data[Index * 4 + 1] = X[1];
+#pragma ignore warning(perf)
+	Data[Index * 4 + 2] = X[2];
+#pragma ignore warning(perf)
+	Data[Index * 4 + 3] = X[3];
 }
 
-inline FVector3f VectorDivMax(const FVector3f& V, const FVector3f& W, const uniform float Epsilon = UE_SMALL_NUMBER)
+inline void Vector3dSave(uniform double Data[], const int Index, const Vector3d X)
 {
-	const FVector3f Out = { {
-		V.V[0] / max(W.V[0], Epsilon),
-		V.V[1] / max(W.V[1], Epsilon),
-		V.V[2] / max(W.V[2], Epsilon),
-	} };
-    
-	return Out;
+#pragma ignore warning(perf)
+	Data[Index * 3 + 0] = X[0];
+#pragma ignore warning(perf)
+	Data[Index * 3 + 1] = X[1];
+#pragma ignore warning(perf)
+	Data[Index * 3 + 2] = X[2];
 }
 
-inline FVector VectorDivMax(const FVector& V, const FVector& W, const uniform double Epsilon = UE_SMALL_NUMBER)
+inline void Quat4dSave(uniform double Data[], const int Index, const Quat4d X)
 {
-	const FVector Out = { {
-		V.V[0] / max(W.V[0], Epsilon),
-		V.V[1] / max(W.V[1], Epsilon),
-		V.V[2] / max(W.V[2], Epsilon),
-	} };
-    
-	return Out;
+#pragma ignore warning(perf)
+	Data[Index * 4 + 0] = X[0];
+#pragma ignore warning(perf)
+	Data[Index * 4 + 1] = X[1];
+#pragma ignore warning(perf)
+	Data[Index * 4 + 2] = X[2];
+#pragma ignore warning(perf)
+	Data[Index * 4 + 3] = X[3];
 }
 
-inline FVector3f VectorInvExpApprox(const FVector3f& V)
+inline float Vector3fDot(const Vector3f V, const Vector3f W)
+{
+	return V.x * W.x + V.y * W.y + V.z * W.z;
+}
+
+inline float Vector3fLengthSquared(const Vector3f V)
+{
+	return Vector3fDot(V, V);
+}
+
+inline float Vector3fLength(const Vector3f V)
+{
+	return sqrt(Vector3fLengthSquared(V));
+}
+
+inline Vector3f Vector3fInvExpApprox(const Vector3f V)
 {
 	const uniform float A = 1.00746054f;
 	const uniform float B = 0.45053901f;
 	const uniform float C = 0.25724632f;
+	return 1.0f / (1.0f + A * V + B * V * V + C * V * V * V);
+}
 
-	const FVector3f Out = { {
-		1.0f / (1.0f + A * V.V[0] + B * V.V[0] * V.V[0] + C * V.V[0] * V.V[0] * V.V[0]),
-		1.0f / (1.0f + A * V.V[1] + B * V.V[1] * V.V[1] + C * V.V[1] * V.V[1] * V.V[1]),
-		1.0f / (1.0f + A * V.V[2] + B * V.V[2] * V.V[2] + C * V.V[2] * V.V[2] * V.V[2]),
-	} };
+inline Vector3f Vector3fDivMax(const Vector3f V, const Vector3f W, const uniform float Epsilon = UE_SMALL_NUMBER)
+{
+	const Vector3f O = {
+		V.x / max(W.x, Epsilon),
+		V.y / max(W.y, Epsilon),
+		V.z / max(W.z, Epsilon),
+	};
 
-	return Out;
+	return O;
 }
 
-inline FVector VectorEerp(const FVector& V, const FVector& W, const uniform double Alpha, const uniform double Epsilon = UE_SMALL_NUMBER)
+inline Vector3f Vector3fExp(const Vector3f V)
 {
-	const FVector Out = { {
-        pow(max(V.V[0], Epsilon), (1.0 - Alpha)) * pow(max(W.V[0], Epsilon), Alpha),
-        pow(max(V.V[1], Epsilon), (1.0 - Alpha)) * pow(max(W.V[0], Epsilon), Alpha),
-        pow(max(V.V[2], Epsilon), (1.0 - Alpha)) * pow(max(W.V[0], Epsilon), Alpha),
-	} };
+	const Vector3f O = {
+		exp(V.x),
+		exp(V.y),
+		exp(V.z),
+	};
 
-	return Out;
+	return O;
 }
 
-inline FVector VectorExp(const FVector& V)
+inline double Vector3dDot(const Vector3d V, const Vector3d W)
 {
-	const FVector Out = { {
-		exp(V.V[0]),
-		exp(V.V[1]),
-		exp(V.V[2]),
-	} };
+	return V.x * W.x + V.y * W.y + V.z * W.z;
+}
 
-	return Out;
+inline double Vector3dLengthSquared(const Vector3d V)
+{
+	return Vector3dDot(V, V);
 }
 
-inline FVector VectorLogMax(const FVector& V, const uniform double Epsilon = UE_SMALL_NUMBER)
+inline double Vector3dLength(const Vector3d V)
 {
-	const FVector Out = { {
-		log(max(V.V[0], Epsilon)),
-		log(max(V.V[1], Epsilon)),
-		log(max(V.V[2], Epsilon)),
-	} };
-    
+	return sqrt(Vector3dLengthSquared(V));
+}
+
+inline Vector3d Vector3dLerp(const Vector3d V, const Vector3d W, const uniform double Alpha)
+{
+	return V * (1.0 - Alpha) + W * Alpha;
+}
+
+inline Vector3d Vector3dEerp(const Vector3d V, const Vector3d W, const uniform double Alpha, const uniform double Epsilon = UE_SMALL_NUMBER)
+{
+	const Vector3d Out = {
+		pow(max(V.x, Epsilon), (1.0 - Alpha)) * pow(max(W.x, Epsilon), Alpha),
+		pow(max(V.y, Epsilon), (1.0 - Alpha)) * pow(max(W.y, Epsilon), Alpha),
+		pow(max(V.z, Epsilon), (1.0 - Alpha)) * pow(max(W.z, Epsilon), Alpha),
+	};
+
 	return Out;
 }
 
-inline FVector VectorClampToMaxSize(const FVector V, const uniform double MaxSize)
+inline Vector3d Vector3dDivMax(const Vector3d V, const Vector3d W, const uniform double Epsilon = UE_SMALL_NUMBER)
 {
-	const double L = VectorDot(V, V);
-    
+	const Vector3d O = {
+		V.x / max(W.x, Epsilon),
+		V.y / max(W.y, Epsilon),
+		V.z / max(W.z, Epsilon),
+	};
+
+	return O;
+}
+
+inline Vector3d Vector3dLogMax(const Vector3d V, const uniform double Epsilon = UE_SMALL_NUMBER)
+{
+	const Vector3d O = {
+		log(max(V.x, Epsilon)),
+		log(max(V.y, Epsilon)),
+		log(max(V.z, Epsilon)),
+	};
+
+	return O;
+}
+
+inline Vector3d Vector3dClampToMaxSize(const Vector3d V, const uniform double MaxSize)
+{
+	const double L = Vector3dLengthSquared(V);
 	return L > (MaxSize * MaxSize) ? V * (MaxSize * rsqrt(L)) : V;
 }
 
-inline FVector4 QuatExp(const FVector& V, const uniform double Epsilon = UE_SMALL_NUMBER)
+inline uniform Quat4f Quat4fIdentity()
+{
+	const uniform Quat4f O = { 0.0f, 0.0f, 0.0f, 1.0f };
+	return O;
+}
+
+inline Quat4d Quat4dAbs(const Quat4d Q)
+{
+	const Quat4d O = {
+		select(Q.w > 0.0, Q.x, -Q.x),
+		select(Q.w > 0.0, Q.y, -Q.y),
+		select(Q.w > 0.0, Q.z, -Q.z),
+		select(Q.w > 0.0, Q.w, -Q.w),
+	};
+
+	return O;
+}
+
+inline Quat4d Quat4dMul(const Quat4d Q, const Quat4d P)
 {
-	const double HalfAngleSq = VectorDot(V, V);
+	const Quat4d O = {
+		P.w * Q.x + P.x * Q.w - P.y * Q.z + P.z * Q.y,
+		P.w * Q.y + P.x * Q.z + P.y * Q.w - P.z * Q.x,
+		P.w * Q.z - P.x * Q.y + P.y * Q.x + P.z * Q.w,
+		P.w * Q.w - P.x * Q.x - P.y * Q.y - P.z * Q.z,
+	};
+
+	return O;
+}
 
-	cif (HalfAngleSq < Epsilon)
+inline Quat4d Quat4dInv(const Quat4d Q)
+{
+	const Quat4d O = { -Q.x, -Q.y, -Q.z, Q.w };
+	return O;
+}
+
+inline Quat4d Quat4dMulInv(const Quat4d Q, const Quat4d P)
+{
+	return Quat4dMul(Q, Quat4dInv(P));
+}
+
+inline double Quat4dDot(const Quat4d Q, const Quat4d P)
+{
+	return Q.x * P.x + Q.y * P.y + Q.z * P.z + Q.w * P.w;
+}
+
+inline double Quat4dLengthSquared(const Quat4d Q)
+{
+	return Quat4dDot(Q, Q);
+}
+
+inline double Quat4dLength(const Quat4d Q)
+{
+	return sqrt(Quat4dLengthSquared(Q));
+}
+
+inline Quat4d Quat4dNormalized(const Quat4d Q)
+{
+	return Q / Quat4dLength(Q);
+}
+
+inline Quat4d Quat4dExp(const Vector3d V, const uniform double Epsilon = UE_SMALL_NUMBER)
+{
+	const double HalfAngleSq = Vector3dLengthSquared(V);
+
+	cif(HalfAngleSq < Epsilon)
 	{
-		const FVector4 O = { {
-			V.V[0],
-			V.V[1],
-			V.V[2],
-			1.0,
-		} };
-
-		return VectorNormalizeQuaternion(O);
+		const Quat4d O = { V.x, V.y, V.z, 1.0 };
+		return Quat4dNormalized(O);
 	}
 	else
 	{
@@ -162,98 +293,83 @@ inline FVector4 QuatExp(const FVector& V, const uniform double Epsilon = UE_SMAL
 		const double C = cos(HalfAngle);
 		const double S = sin(HalfAngle) / HalfAngle;
 
-		const FVector4 O = { {
-			S * V.V[0],
-			S * V.V[1],
-			S * V.V[2],
-			C,
-		} };
-
+		const Quat4d O = { S * V.x, S * V.y, S * V.z, C };
 		return O;
 	}
 }
 
-inline FVector QuatLog(const FVector4& Q, const uniform double Epsilon = UE_SMALL_NUMBER)
+inline Vector3d Quat4dLog(const Quat4d Q, const uniform double Epsilon = UE_SMALL_NUMBER)
 {
-	const double LengthSq = VectorDot(Q, Q);
+	const Vector3d V = { Q.x, Q.y, Q.z };
+	const double LengthSq = Vector3dLengthSquared(V);
 
-	cif (LengthSq < Epsilon)
+	cif(LengthSq < Epsilon)
 	{
-		const FVector O = { { Q.V[0], Q.V[1], Q.V[2] } };
-		return O;
+		return V;
 	}
 	else
 	{
 		const double Length = sqrt(LengthSq);
-		const double HalfAngle = acos(clamp(Q.V[3], -1.0d, 1.0d));
-
-		const FVector O = { {
-			HalfAngle * (Q.V[0] / Length),
-			HalfAngle * (Q.V[1] / Length),
-			HalfAngle * (Q.V[2] / Length),
-		} };
-
-		return O;
+		const double HalfAngle = acos(clamp(Q.w, -1.0d, 1.0d));
+		return HalfAngle * (V / Length);
 	}
 }
 
-inline FVector4 QuatFromRotationVector(const FVector& V, const uniform double Epsilon = UE_SMALL_NUMBER)
+inline Vector3d Quat4dToRotationVector(const Quat4d Q, const uniform double Epsilon = UE_SMALL_NUMBER)
 {
-	return QuatExp(0.5 * V, Epsilon);
+	return 2.0 * Quat4dLog(Q, Epsilon);
 }
 
-inline FVector QuatToRotationVector(const FVector4& Q, const uniform double Epsilon = UE_SMALL_NUMBER)
+inline Quat4d Quat4dFromRotationVector(const Vector3d V, const uniform double Epsilon = UE_SMALL_NUMBER)
 {
-	return 2.0 * QuatLog(Q, Epsilon);
+	return Quat4dExp(V * 0.5, Epsilon);
 }
 
-inline FVector4 VectorNegate(const FVector4& Q)
+inline Quat4d Quat4dGetShortestArcWith(const Quat4d Q, const Quat4d P)
 {
-	const FVector4 O = { { -Q.V[0], -Q.V[1], -Q.V[2], -Q.V[3] } };
-	return O;
-}
+	const double D = Quat4dDot(Q, P);
 
-inline FVector4 VectorQuaternionEnforceShortestArcWith(const FVector4 Q, const FVector4 P)
-{
-	return VectorDot(Q, P) > 0.0 ? Q : VectorNegate(Q);
-}
+	const Quat4d O = {
+		select(D > 0.0, Q.x, -Q.x),
+		select(D > 0.0, Q.y, -Q.y),
+		select(D > 0.0, Q.z, -Q.z),
+		select(D > 0.0, Q.w, -Q.w),
+	};
 
-inline FVector4 QuatAbs(const FVector4 Q)
-{
-	return Q.V[3] > 0.0 ? Q : VectorNegate(Q);
+	return O;
 }
 
-inline FVector ExtrapolateTranslation(
-	const FVector& Translation,
-	const FVector3f& Velocity,
-	const FVector3f& DecayHalflife,
+inline Vector3d ExtrapolateTranslation(
+	const Vector3d Translation,
+	const Vector3f Velocity,
+	const Vector3f DecayHalflife,
 	const uniform float Time,
 	const uniform float Epsilon = UE_SMALL_NUMBER)
 {
-	const FVector3f C = VectorDivMax(UE_LN2, DecayHalflife, Epsilon);
-	return Translation + CastToVector(VectorDivMax(Velocity, C, Epsilon) * (FloatOneVector - VectorInvExpApprox(C * Time)));
+	const Vector3f C = Vector3fDivMax(UE_LN2, DecayHalflife, Epsilon);
+	return Translation + Vector3fDivMax(Velocity, C, Epsilon) * (1.0f - Vector3fInvExpApprox(C * Time));
 }
 
-inline FVector4 ExtrapolateRotation(
-	const FVector4& Rotation,
-	const FVector3f& Velocity,
-	const FVector3f& DecayHalflife,
+inline Quat4d ExtrapolateRotation(
+	const Quat4d Rotation,
+	const Vector3f Velocity,
+	const Vector3f DecayHalflife,
 	const uniform float Time,
 	const uniform float Epsilon = UE_SMALL_NUMBER)
 {
-	const FVector3f C = VectorDivMax(UE_LN2, DecayHalflife, Epsilon);
-	return VectorQuaternionMultiply2(QuatFromRotationVector(CastToVector(VectorDivMax(Velocity, C, Epsilon) * (FloatOneVector - VectorInvExpApprox(C * Time)))), Rotation);
+	const Vector3f C = Vector3fDivMax(UE_LN2, DecayHalflife, Epsilon);
+	return Quat4dMul(Quat4dFromRotationVector(Vector3fDivMax(Velocity, C, Epsilon) * (1.0f - Vector3fInvExpApprox(C * Time))), Rotation);
 }
 
-inline FVector ExtrapolateScale(
-	const FVector& Scale,
-	const FVector3f& Velocity,
-	const FVector3f& DecayHalflife,
+inline Vector3d ExtrapolateScale(
+	const Vector3d Scale,
+	const Vector3f Velocity,
+	const Vector3f DecayHalflife,
 	const uniform float Time,
 	const uniform float Epsilon = UE_SMALL_NUMBER)
 {
-	const FVector3f C = VectorDivMax(UE_LN2, DecayHalflife, Epsilon);
-	return VectorExp(CastToVector(VectorDivMax(Velocity, C, Epsilon) * (FloatOneVector - VectorInvExpApprox(C * Time)))) * Scale;
+	const Vector3f C = Vector3fDivMax(UE_LN2, DecayHalflife, Epsilon);
+	return Vector3fExp(Vector3fDivMax(Velocity, C, Epsilon) * (1.0f - Vector3fInvExpApprox(C * Time))) * Scale;
 }
 
 inline float ClipMagnitudeToGreaterThanEpsilon(const float X, const uniform float Epsilon = UE_KINDA_SMALL_NUMBER)
@@ -274,43 +390,43 @@ inline float ComputeDecayHalfLifeFromDiffAndVelocity(
 	return clamp(HalfLife * (SrcDstDiff / ClipMagnitudeToGreaterThanEpsilon(SrcVelocity, Epsilon)), HalfLifeMin, HalfLifeMax);
 }
 
-inline FVector3f ComputeDecayHalfLifeFromDiffAndVelocity(
-	const FVector SrcDstDiff,
-	const FVector3f SrcVelocity,
+inline Vector3f ComputeDecayHalfLifeFromDiffAndVelocity(
+	const Vector3d SrcDstDiff,
+	const Vector3f SrcVelocity,
 	const uniform float HalfLife,
 	const uniform float HalfLifeMin,
 	const uniform float HalfLifeMax,
 	const uniform float Epsilon = UE_KINDA_SMALL_NUMBER)
 {
-	const FVector3f O = { {
-		ComputeDecayHalfLifeFromDiffAndVelocity((float)SrcDstDiff.V[0], SrcVelocity.V[0], HalfLife, HalfLifeMin, HalfLifeMax, Epsilon),
-		ComputeDecayHalfLifeFromDiffAndVelocity((float)SrcDstDiff.V[1], SrcVelocity.V[1], HalfLife, HalfLifeMin, HalfLifeMax, Epsilon),
-		ComputeDecayHalfLifeFromDiffAndVelocity((float)SrcDstDiff.V[2], SrcVelocity.V[2], HalfLife, HalfLifeMin, HalfLifeMax, Epsilon),
-	} };
+	const Vector3f O = {
+		ComputeDecayHalfLifeFromDiffAndVelocity((float)SrcDstDiff.x, SrcVelocity.x, HalfLife, HalfLifeMin, HalfLifeMax, Epsilon),
+		ComputeDecayHalfLifeFromDiffAndVelocity((float)SrcDstDiff.y, SrcVelocity.y, HalfLife, HalfLifeMin, HalfLifeMax, Epsilon),
+		ComputeDecayHalfLifeFromDiffAndVelocity((float)SrcDstDiff.z, SrcVelocity.z, HalfLife, HalfLifeMin, HalfLifeMax, Epsilon),
+	};
 
 	return O;
 }
 
 export void AnimNextDeadBlendingTransition(
-	uniform FVector4f BoneRotationDirections[],
-	uniform FVector SourceBoneTranslations[],
-	uniform FVector4 SourceBoneRotations[],
-	uniform FVector SourceBoneScales3D[],
-	uniform FVector3f SourceBoneTranslationVelocities[],
-	uniform FVector3f SourceBoneRotationVelocities[],
-	uniform FVector3f SourceBoneScaleVelocities[],
-	uniform FVector3f SourceBoneTranslationDecayHalfLives[],
-	uniform FVector3f SourceBoneRotationDecayHalfLives[],
-	uniform FVector3f SourceBoneScaleDecayHalfLives[],
-	const uniform FVector DestBoneTranslations[],
-	const uniform FVector4 DestBoneRotations[],
-	const uniform FVector DestBoneScales3D[],
-	const uniform FVector CurrBoneTranslations[],
-	const uniform FVector4 CurrBoneRotations[],
-	const uniform FVector CurrBoneScales3D[],
-	const uniform FVector PrevBoneTranslations[],
-	const uniform FVector4 PrevBoneRotations[],
-	const uniform FVector PrevBoneScales3D[],
+	uniform float BoneRotationDirections[],
+	uniform double SourceBoneTranslations[],
+	uniform double SourceBoneRotations[],
+	uniform double SourceBoneScales3D[],
+	uniform float SourceBoneTranslationVelocities[],
+	uniform float SourceBoneRotationVelocities[],
+	uniform float SourceBoneScaleVelocities[],
+	uniform float SourceBoneTranslationDecayHalfLives[],
+	uniform float SourceBoneRotationDecayHalfLives[],
+	uniform float SourceBoneScaleDecayHalfLives[],
+	const uniform double DestBoneTranslations[],
+	const uniform double DestBoneRotations[],
+	const uniform double DestBoneScales3D[],
+	const uniform double CurrBoneTranslations[],
+	const uniform double CurrBoneRotations[],
+	const uniform double CurrBoneScales3D[],
+	const uniform double PrevBoneTranslations[],
+	const uniform double PrevBoneRotations[],
+	const uniform double PrevBoneScales3D[],
 	const uniform int LODBoneNum,
 	const uniform float DeltaTime,
 	const uniform float ExtrapolationHalfLife,
@@ -320,59 +436,57 @@ export void AnimNextDeadBlendingTransition(
 	const uniform float MaximumRotationVelocity,
 	const uniform float MaximumScaleVelocity)
 {
-	const FVector4f IdentityRotation = Float0001;
-
-	foreach(i = 0 ... LODBoneNum)
+	foreach(BoneIdx = 0 ... LODBoneNum)
 	{
-		const uniform int Index = extract(i, 0);
-
-		VectorStore(&BoneRotationDirections[Index], IdentityRotation);
-
-		const FVector SrcTranslationCurr = VectorLoad(&CurrBoneTranslations[Index]);
-		const FVector4 SrcRotationCurr = VectorLoad(&CurrBoneRotations[Index]);
-		const FVector SrcScaleCurr = VectorLoad(&CurrBoneScales3D[Index]);
-
-		VectorStore(&SourceBoneTranslations[Index], SrcTranslationCurr);
-		VectorStore(&SourceBoneRotations[Index], SrcRotationCurr);
-		VectorStore(&SourceBoneScales3D[Index], SrcScaleCurr);
-
-		const FVector SrcTranslationPrev = VectorLoad(&PrevBoneTranslations[Index]);
-		const FVector4 SrcRotationPrev = VectorLoad(&PrevBoneRotations[Index]);
-		const FVector SrcScalePrev = VectorLoad(&PrevBoneScales3D[Index]);
-
-		const FVector TranslationDiff = SrcTranslationCurr - SrcTranslationPrev;
-		const FVector4 RotationDiff = QuatAbs(VectorQuaternionMultiply2(SrcRotationCurr, QuatInverse(SrcRotationPrev)));
-		const FVector ScaleDiff = VectorDivMax(SrcScaleCurr, SrcScalePrev);
-
-		VectorStore(&SourceBoneTranslationVelocities[Index], CastToVector3f(VectorClampToMaxSize(TranslationDiff / DeltaTime, MaximumTranslationVelocity)));
-		VectorStore(&SourceBoneRotationVelocities[Index], CastToVector3f(VectorClampToMaxSize(QuatToRotationVector(RotationDiff) / DeltaTime, MaximumRotationVelocity)));
-		VectorStore(&SourceBoneScaleVelocities[Index], CastToVector3f(VectorClampToMaxSize(VectorLogMax(ScaleDiff) / DeltaTime, MaximumScaleVelocity)));
+		Quat4fSave(BoneRotationDirections, BoneIdx, Quat4fIdentity());
+	}
 
-		const FVector DstTranslation = VectorLoad(&DestBoneTranslations[Index]);
-		const FVector4 DstRotation = VectorLoad(&DestBoneRotations[Index]);
-		const FVector DstScale = VectorLoad(&DestBoneScales3D[Index]);
+	foreach(BoneIdx = 0 ... LODBoneNum)
+	{
+		const Vector3d SrcTranslationCurr = Vector3dLoad(CurrBoneTranslations, BoneIdx);
+		Vector3dSave(SourceBoneTranslations, BoneIdx, SrcTranslationCurr);
 
-		const FVector TranslationSrcDstDiff = DstTranslation - SrcTranslationCurr;
-		const FVector4 RotationSrcDstDiff = QuatAbs(VectorQuaternionMultiply2(DstRotation, QuatInverse(SrcRotationCurr)));
-		const FVector ScaleSrcDstDiff = VectorDivMax(DstScale, SrcScaleCurr);
+		const Vector3d TranslationDiff = SrcTranslationCurr - Vector3dLoad(PrevBoneTranslations, BoneIdx);
+		const Vector3f TranslationVelocity = Vector3dClampToMaxSize(TranslationDiff / DeltaTime, MaximumTranslationVelocity);
+		Vector3fSave(SourceBoneTranslationVelocities, BoneIdx, TranslationVelocity);
 
-		VectorStore(&SourceBoneTranslationDecayHalfLives[Index], ComputeDecayHalfLifeFromDiffAndVelocity(
-			TranslationSrcDstDiff,
-			VectorLoad(&SourceBoneTranslationVelocities[Index]),
+		Vector3fSave(SourceBoneTranslationDecayHalfLives, BoneIdx, ComputeDecayHalfLifeFromDiffAndVelocity(
+			Vector3dLoad(DestBoneTranslations, BoneIdx) - SrcTranslationCurr,
+			TranslationVelocity,
 			ExtrapolationHalfLife,
 			ExtrapolationHalfLifeMin,
 			ExtrapolationHalfLifeMax));
+	}
+
+	foreach(BoneIdx = 0 ... LODBoneNum)
+	{
+		const Quat4d SrcRotationCurr = Quat4dLoad(CurrBoneRotations, BoneIdx);
+		Quat4dSave(SourceBoneRotations, BoneIdx, SrcRotationCurr);
+
+		const Quat4d RotationDiff = Quat4dAbs(Quat4dMulInv(SrcRotationCurr, Quat4dLoad(PrevBoneRotations, BoneIdx)));
+		const Vector3f RotationVelocity = Vector3dClampToMaxSize(Quat4dToRotationVector(RotationDiff) / DeltaTime, MaximumRotationVelocity);
+		Vector3fSave(SourceBoneRotationVelocities, BoneIdx, RotationVelocity);
 
-		VectorStore(&SourceBoneRotationDecayHalfLives[Index], ComputeDecayHalfLifeFromDiffAndVelocity(
-			QuatToRotationVector(RotationSrcDstDiff),
-			VectorLoad(&SourceBoneRotationVelocities[Index]),
+		Vector3fSave(SourceBoneRotationDecayHalfLives, BoneIdx, ComputeDecayHalfLifeFromDiffAndVelocity(
+			Quat4dToRotationVector(Quat4dAbs(Quat4dMulInv(Quat4dLoad(DestBoneRotations, BoneIdx), SrcRotationCurr))),
+			RotationVelocity,
 			ExtrapolationHalfLife,
 			ExtrapolationHalfLifeMin,
 			ExtrapolationHalfLifeMax));
+	}
 
-		VectorStore(&SourceBoneScaleDecayHalfLives[Index], ComputeDecayHalfLifeFromDiffAndVelocity(
-			ScaleSrcDstDiff,
-			VectorLoad(&SourceBoneScaleVelocities[Index]),
+	foreach(BoneIdx = 0 ... LODBoneNum)
+	{
+		const Vector3d SrcScaleCurr = Vector3dLoad(CurrBoneScales3D, BoneIdx);
+		Vector3dSave(SourceBoneScales3D, BoneIdx, SrcScaleCurr);
+
+		const Vector3d ScaleDiff = Vector3dDivMax(SrcScaleCurr, Vector3dLoad(PrevBoneScales3D, BoneIdx));
+		const Vector3f ScaleVelocity = Vector3dClampToMaxSize(Vector3dLogMax(ScaleDiff) / DeltaTime, MaximumScaleVelocity);
+		Vector3fSave(SourceBoneScaleVelocities, BoneIdx, ScaleVelocity);
+
+		Vector3fSave(SourceBoneScaleDecayHalfLives, BoneIdx, ComputeDecayHalfLifeFromDiffAndVelocity(
+			Vector3dDivMax(Vector3dLoad(DestBoneScales3D, BoneIdx), SrcScaleCurr),
+			ScaleVelocity,
 			ExtrapolationHalfLife,
 			ExtrapolationHalfLifeMin,
 			ExtrapolationHalfLifeMax));
@@ -381,97 +495,95 @@ export void AnimNextDeadBlendingTransition(
 
 
 export void AnimNextDeadBlendingTransitionStatic(
-	uniform FVector4f BoneRotationDirections[],
-	uniform FVector SourceBoneTranslations[],
-	uniform FVector4 SourceBoneRotations[],
-	uniform FVector SourceBoneScales3D[],
-	uniform FVector3f SourceBoneTranslationVelocities[],
-	uniform FVector3f SourceBoneRotationVelocities[],
-	uniform FVector3f SourceBoneScaleVelocities[],
-	uniform FVector3f SourceBoneTranslationDecayHalfLives[],
-	uniform FVector3f SourceBoneRotationDecayHalfLives[],
-	uniform FVector3f SourceBoneScaleDecayHalfLives[],
-	const uniform FVector CurrBoneTranslations[],
-	const uniform FVector4 CurrBoneRotations[],
-	const uniform FVector CurrBoneScales[],
+	uniform float BoneRotationDirections[],
+	uniform double SourceBoneTranslations[],
+	uniform double SourceBoneRotations[],
+	uniform double SourceBoneScales3D[],
+	uniform float SourceBoneTranslationVelocities[],
+	uniform float SourceBoneRotationVelocities[],
+	uniform float SourceBoneScaleVelocities[],
+	uniform float SourceBoneTranslationDecayHalfLives[],
+	uniform float SourceBoneRotationDecayHalfLives[],
+	uniform float SourceBoneScaleDecayHalfLives[],
+	const uniform double CurrBoneTranslations[],
+	const uniform double CurrBoneRotations[],
+	const uniform double CurrBoneScales[],
 	const uniform int LODBoneNum,
 	const uniform float ExtrapolationHalfLifeMin)
 {
-	const FVector4f IdentityRotation = Float0001;
-	const FVector3f VelocityZeroVector = FloatZeroVector;
-	const FVector3f HalfLifeVector = { { ExtrapolationHalfLifeMin, ExtrapolationHalfLifeMin, ExtrapolationHalfLifeMin } };
-
-	foreach(i = 0 ... LODBoneNum)
+	foreach(BoneIdx = 0 ... LODBoneNum)
 	{
-		const uniform int Index = extract(i, 0);
-
-		VectorStore(&BoneRotationDirections[Index], IdentityRotation);
-		VectorStore(&SourceBoneTranslations[Index], VectorLoad(&CurrBoneTranslations[Index]));
-		VectorStore(&SourceBoneRotations[Index], VectorLoad(&CurrBoneRotations[Index]));
-		VectorStore(&SourceBoneScales3D[Index], VectorLoad(&CurrBoneScales[Index]));
-
-		VectorStore(&SourceBoneTranslationVelocities[Index], VelocityZeroVector);
-		VectorStore(&SourceBoneRotationVelocities[Index], VelocityZeroVector);
-		VectorStore(&SourceBoneScaleVelocities[Index], VelocityZeroVector);
-		VectorStore(&SourceBoneTranslationDecayHalfLives[Index], HalfLifeVector);
-		VectorStore(&SourceBoneRotationDecayHalfLives[Index], HalfLifeVector);
-		VectorStore(&SourceBoneScaleDecayHalfLives[Index], HalfLifeVector);
+		Quat4fSave(BoneRotationDirections, BoneIdx, Quat4fIdentity());
+		Vector3dSave(SourceBoneTranslations, BoneIdx, Vector3dLoad(CurrBoneTranslations, BoneIdx));
+		Quat4dSave(SourceBoneRotations, BoneIdx, Quat4dLoad(CurrBoneRotations, BoneIdx));
+		Vector3dSave(SourceBoneScales3D, BoneIdx, Vector3dLoad(CurrBoneScales, BoneIdx));
+
+		Vector3fSave(SourceBoneTranslationVelocities, BoneIdx, 0.0f);
+		Vector3fSave(SourceBoneRotationVelocities, BoneIdx, 0.0f);
+		Vector3fSave(SourceBoneScaleVelocities, BoneIdx, 0.0f);
+		Vector3fSave(SourceBoneTranslationDecayHalfLives, BoneIdx, ExtrapolationHalfLifeMin);
+		Vector3fSave(SourceBoneRotationDecayHalfLives, BoneIdx, ExtrapolationHalfLifeMin);
+		Vector3fSave(SourceBoneScaleDecayHalfLives, BoneIdx, ExtrapolationHalfLifeMin);
 	}
 }
 
 export void AnimNextDeadBlendingApply(
-	uniform FVector DestTranslations[],
-	uniform FVector4 DestRotations[],
-	uniform FVector DestScales3D[],
-	uniform FVector4f BoneRotationDirections[],
-	const uniform FVector SourceTranslations[],
-	const uniform FVector4 SourceRotations[],
-	const uniform FVector SourceScales3D[],
-	const uniform FVector3f SourceBoneTranslationVelocities[],
-	const uniform FVector3f SourceBoneRotationVelocities[],
-	const uniform FVector3f SourceBoneScaleVelocities[],
-	const uniform FVector3f SourceBoneTranslationDecayHalfLives[],
-	const uniform FVector3f SourceBoneRotationDecayHalfLives[],
-	const uniform FVector3f SourceBoneScaleDecayHalfLives[],
+	uniform double DestTranslations[],
+	uniform double DestRotations[],
+	uniform double DestScales3D[],
+	uniform float BoneRotationDirections[],
+	const uniform double SourceTranslations[],
+	const uniform double SourceRotations[],
+	const uniform double SourceScales3D[],
+	const uniform float SourceBoneTranslationVelocities[],
+	const uniform float SourceBoneRotationVelocities[],
+	const uniform float SourceBoneScaleVelocities[],
+	const uniform float SourceBoneTranslationDecayHalfLives[],
+	const uniform float SourceBoneRotationDecayHalfLives[],
+	const uniform float SourceBoneScaleDecayHalfLives[],
 	const uniform int LODBoneNum,
 	const uniform float Alpha,
 	const uniform float TimeSinceTransition)
 {
-	foreach (i = 0 ... LODBoneNum)
+	foreach(BoneIdx = 0 ... LODBoneNum)
 	{
-		const uniform int Index = extract(i, 0);
-
-		const FVector ExtrapolatedTranslation = ExtrapolateTranslation(
-			VectorLoad(&SourceTranslations[Index]),
-			VectorLoad(&SourceBoneTranslationVelocities[Index]),
-			VectorLoad(&SourceBoneTranslationDecayHalfLives[Index]),
+		const Vector3d ExtrapolatedTranslation = ExtrapolateTranslation(
+			Vector3dLoad(SourceTranslations, BoneIdx),
+			Vector3fLoad(SourceBoneTranslationVelocities, BoneIdx),
+			Vector3fLoad(SourceBoneTranslationDecayHalfLives, BoneIdx),
 			TimeSinceTransition);
 
-		const FVector4 ExtrapolatedRotation = ExtrapolateRotation(
-			VectorLoad(&SourceRotations[Index]),
-			VectorLoad(&SourceBoneRotationVelocities[Index]),
-			VectorLoad(&SourceBoneRotationDecayHalfLives[Index]),
-			TimeSinceTransition);
+		Vector3dSave(DestTranslations, BoneIdx, Vector3dLerp(Vector3dLoad(DestTranslations, BoneIdx), ExtrapolatedTranslation, Alpha));
+	}
 
-		const FVector ExtrapolatedScale = ExtrapolateScale(
-			VectorLoad(&SourceScales3D[Index]),
-			VectorLoad(&SourceBoneScaleVelocities[Index]),
-			VectorLoad(&SourceBoneScaleDecayHalfLives[Index]),
+	foreach(BoneIdx = 0 ... LODBoneNum)
+	{
+		const Quat4d ExtrapolatedRotation = ExtrapolateRotation(
+			Quat4dLoad(SourceRotations, BoneIdx),
+			Vector3fLoad(SourceBoneRotationVelocities, BoneIdx),
+			Vector3fLoad(SourceBoneRotationDecayHalfLives, BoneIdx),
 			TimeSinceTransition);
 
-		const FVector4 RotationDiff = VectorQuaternionEnforceShortestArcWith(
-			VectorQuaternionMultiply2(ExtrapolatedRotation, QuatInverse(VectorLoad(&DestRotations[Index]))),
-			CastToVector4(VectorLoad(&BoneRotationDirections[Index])));
+		const Quat4d RotationDiff = Quat4dGetShortestArcWith(
+			Quat4dMulInv(ExtrapolatedRotation, Quat4dLoad(DestRotations, BoneIdx)),
+			Quat4fLoad(BoneRotationDirections, BoneIdx));
 
-		VectorStore(&BoneRotationDirections[Index], CastToVector4f(RotationDiff));
+		Quat4fSave(BoneRotationDirections, BoneIdx, RotationDiff);
 
-		VectorStore(&DestTranslations[Index], VectorLerp(VectorLoad(&DestTranslations[Index]), ExtrapolatedTranslation, Alpha));
+		Quat4dSave(DestRotations, BoneIdx,
+			Quat4dMul(
+				Quat4dFromRotationVector(Quat4dToRotationVector(RotationDiff) * Alpha),
+				Quat4dLoad(DestRotations, BoneIdx)));
+	}
 
-		VectorStore(&DestRotations[Index],
-			VectorQuaternionMultiply2(
-				QuatFromRotationVector(QuatToRotationVector(RotationDiff) * Alpha),
-				VectorLoad(&DestRotations[Index])));
+	foreach(BoneIdx = 0 ... LODBoneNum)
+	{
+		const Vector3d ExtrapolatedScale = ExtrapolateScale(
+			Vector3dLoad(SourceScales3D, BoneIdx),
+			Vector3fLoad(SourceBoneScaleVelocities, BoneIdx),
+			Vector3fLoad(SourceBoneScaleDecayHalfLives, BoneIdx),
+			TimeSinceTransition);
 
-		VectorStore(&DestScales3D[Index], VectorEerp(VectorLoad(&DestScales3D[Index]), ExtrapolatedScale, Alpha));
+		Vector3dSave(DestScales3D, BoneIdx, Vector3dEerp(Vector3dLoad(DestScales3D, BoneIdx), ExtrapolatedScale, Alpha));
 	}
 }
\ No newline at end of file
