AnimNext: Fix memory stomp in variables automated test

diff --git a/Engine/Plugins/Experimental/AnimNext/Source/AnimNext/Private/DataInterface/AnimNextDataInterfaceInstance.cpp b/Engine/Plugins/Experimental/AnimNext/Source/AnimNext/Private/DataInterface/AnimNextDataInterfaceInstance.cpp
--- a/Engine/Plugins/Experimental/AnimNext/Source/AnimNext/Private/DataInterface/AnimNextDataInterfaceInstance.cpp
+++ b/Engine/Plugins/Experimental/AnimNext/Source/AnimNext/Private/DataInterface/AnimNextDataInterfaceInstance.cpp
@@ -271,13 +271,15 @@ namespace UE::AnimNext::Private
 
 		const FStructProperty* StructProperty = CastFieldChecked<FStructProperty>(InProperty);
 		check(StructProperty->Struct);
+		check(RequestedStruct);
 
-		if (RequestedStruct != nullptr && CanCastTo(StructProperty->Struct, RequestedStruct) == false)
+		if (CanCastTo(StructProperty->Struct, RequestedStruct) == false)
 		{
 			return EPropertyBagResult::TypeMismatch;
 		}
 
-		StructProperty->CopyCompleteValue(OutValue, Address);
+		// We dont use the property here to avoid copying more than we need if we are 'casting' from derived to base
+		RequestedStruct->CopyScriptStruct(OutValue, Address, 1);
 
 		return EPropertyBagResult::Success;
 	}
diff --git a/Engine/Plugins/Experimental/AnimNextAnimGraph/Source/AnimNextAnimGraphTestSuite/Private/AnimNextAnimGraphTraitGraphTest.cpp b/Engine/Plugins/Experimental/AnimNextAnimGraph/Source/AnimNextAnimGraphTestSuite/Private/AnimNextAnimGraphTraitGraphTest.cpp
--- a/Engine/Plugins/Experimental/AnimNextAnimGraph/Source/AnimNextAnimGraphTestSuite/Private/AnimNextAnimGraphTraitGraphTest.cpp
+++ b/Engine/Plugins/Experimental/AnimNextAnimGraph/Source/AnimNextAnimGraphTestSuite/Private/AnimNextAnimGraphTraitGraphTest.cpp
@@ -845,6 +845,7 @@ bool FAnimationAnimNextRuntimeTest_Variables::RunTest(const FString& InParameter
 			UE_RETURN_ON_ERROR(GraphInstance->GetVariable("Struct", TestDerived) != EPropertyBagResult::Success, "FAnimationAnimNextRuntimeTest_Variables -> GetVariable succeeded");
 
 			// Test conversion from derived -> base
+			TestStruct = FVector::ZeroVector;
 			UE_RETURN_ON_ERROR(GraphInstance->GetVariable("DerivedStruct", TestStruct) == EPropertyBagResult::Success, "FAnimationAnimNextRuntimeTest_Variables -> GetVariable failed");
 			UE_RETURN_ON_ERROR(TestStruct == FVector::OneVector, "FAnimationAnimNextRuntimeTest_Variables -> Variable value did not match");
 		}
