AnimNext Silver Blend Profile trait
* Adds blend profile trait that uses a hierarchy tables backed by USkeleton (i.e. Skeleton table type) that use the Mask payload type. This is intended for the short term for Silver.
* In future will add abstract skeleton backed support for the long term future.

diff --git a/Engine/Plugins/Experimental/AnimNext/AnimNext.uplugin b/Engine/Plugins/Experimental/AnimNext/AnimNext.uplugin
--- a/Engine/Plugins/Experimental/AnimNext/AnimNext.uplugin
+++ b/Engine/Plugins/Experimental/AnimNext/AnimNext.uplugin
@@ -37,7 +37,15 @@
 			"Name": "EditorScriptingUtilities",
 			"Enabled": true,
 			"TargetAllowList": [ "Editor" ]
-		}
+		},
+		{
+			"Name": "HierarchyTable",
+			"Enabled": true
+		},
+		{
+			"Name": "HierarchyTableBuiltin",
+			"Enabled": true
+		},
 	],
 	"Modules" :
 	[
diff --git a/Engine/Plugins/Experimental/AnimNext/Source/AnimNext/AnimNext.Build.cs b/Engine/Plugins/Experimental/AnimNext/Source/AnimNext/AnimNext.Build.cs
--- a/Engine/Plugins/Experimental/AnimNext/Source/AnimNext/AnimNext.Build.cs
+++ b/Engine/Plugins/Experimental/AnimNext/Source/AnimNext/AnimNext.Build.cs
@@ -14,7 +14,8 @@ namespace UnrealBuildTool.Rules
 					"RigVM",
 					"ControlRig",
 					"Engine",
-					"RenderCore"
+					"RenderCore",
+					"HierarchyTableRuntime"
 				}
 			);
 
diff --git a/Engine/Plugins/Experimental/AnimNext/Source/AnimNext/Private/AnimNextModuleImpl.cpp b/Engine/Plugins/Experimental/AnimNext/Source/AnimNext/Private/AnimNextModuleImpl.cpp
--- a/Engine/Plugins/Experimental/AnimNext/Source/AnimNext/Private/AnimNextModuleImpl.cpp
+++ b/Engine/Plugins/Experimental/AnimNext/Source/AnimNext/Private/AnimNextModuleImpl.cpp
@@ -22,6 +22,7 @@
 #include "Variables/AnimNextFieldPath.h"
 #include "Variables/AnimNextSoftFunctionPtr.h"
 #include "RigVMCore/RigVMGraphFunctionDefinition.h"
+#include "HierarchyTable.h"
 
 #define LOCTEXT_NAMESPACE "AnimNextModule"
 
@@ -40,6 +41,7 @@ namespace UE::AnimNext
 			{ UCurveFloat::StaticClass(), FRigVMRegistry::ERegisterObjectOperation::Class },
 			{ UAnimNextComponent::StaticClass(), FRigVMRegistry::ERegisterObjectOperation::Class },
 			{ UAnimNextModule::StaticClass(), FRigVMRegistry::ERegisterObjectOperation::Class },
+			{ UHierarchyTable::StaticClass(), FRigVMRegistry::ERegisterObjectOperation::Class },
 		};
 
 		RigVMRegistry.RegisterObjectTypes(AllowedObjectTypes);
diff --git a/Engine/Plugins/Experimental/AnimNext/Source/AnimNext/Private/TransformArrayOperations_AoS.cpp b/Engine/Plugins/Experimental/AnimNext/Source/AnimNext/Private/TransformArrayOperations_AoS.cpp
--- a/Engine/Plugins/Experimental/AnimNext/Source/AnimNext/Private/TransformArrayOperations_AoS.cpp
+++ b/Engine/Plugins/Experimental/AnimNext/Source/AnimNext/Private/TransformArrayOperations_AoS.cpp
@@ -82,18 +82,30 @@ namespace UE::AnimNext
 
 	void BlendOverwritePerBoneWithScale(
 		const FTransformArrayAoSView& Dest, const FTransformArrayAoSConstView& Source,
-		const TArrayView<const int32>& LODBoneIndexToWeightIndexMap, const TArrayView<const float>& BoneWeights, const float DefaultScaleWeight)
+		const TArrayView<const int32>& LODBoneIndexToWeightIndexMap, const TArrayView<const float>& BoneWeights, const float DefaultScaleWeight, const bool bInvert)
 	{
 		const ScalarRegister VDefaultScaleWeight(DefaultScaleWeight);
 		const int32 NumTransforms = Source.Num();
 
 		check(Dest.Num() >= NumTransforms);
-
-		for (int32 LODBoneIndex = 0; LODBoneIndex < NumTransforms; ++LODBoneIndex)
+		
+		if (bInvert)
 		{
-			const int32 PerBoneIndex = LODBoneIndexToWeightIndexMap[LODBoneIndex];
-			const ScalarRegister VScaleWeight = BoneWeights.IsValidIndex(PerBoneIndex) ? ScalarRegister(BoneWeights[PerBoneIndex]) : VDefaultScaleWeight;
-			Dest[LODBoneIndex] = Source[LODBoneIndex] * VScaleWeight;
+			for (int32 LODBoneIndex = 0; LODBoneIndex < NumTransforms; ++LODBoneIndex)
+			{
+				const int32 PerBoneIndex = LODBoneIndexToWeightIndexMap[LODBoneIndex];
+				const ScalarRegister VScaleWeight = BoneWeights.IsValidIndex(PerBoneIndex) ? ScalarRegister(1.0f - BoneWeights[PerBoneIndex]) : VDefaultScaleWeight;
+				Dest[LODBoneIndex] = Source[LODBoneIndex] * VScaleWeight;
+			}
+		}
+		else
+		{
+			for (int32 LODBoneIndex = 0; LODBoneIndex < NumTransforms; ++LODBoneIndex)
+			{
+				const int32 PerBoneIndex = LODBoneIndexToWeightIndexMap[LODBoneIndex];
+				const ScalarRegister VScaleWeight = BoneWeights.IsValidIndex(PerBoneIndex) ? ScalarRegister(BoneWeights[PerBoneIndex]) : VDefaultScaleWeight;
+				Dest[LODBoneIndex] = Source[LODBoneIndex] * VScaleWeight;
+			}
 		}
 	}
 
diff --git a/Engine/Plugins/Experimental/AnimNext/Source/AnimNext/Private/TransformArrayOperations_SoA.cpp b/Engine/Plugins/Experimental/AnimNext/Source/AnimNext/Private/TransformArrayOperations_SoA.cpp
--- a/Engine/Plugins/Experimental/AnimNext/Source/AnimNext/Private/TransformArrayOperations_SoA.cpp
+++ b/Engine/Plugins/Experimental/AnimNext/Source/AnimNext/Private/TransformArrayOperations_SoA.cpp
@@ -194,20 +194,35 @@ namespace UE::AnimNext
 
 	void BlendOverwritePerBoneWithScale(
 		const FTransformArraySoAView& Dest, const FTransformArraySoAConstView& Source,
-		const TArrayView<const int32>& LODBoneIndexToWeightIndexMap, const TArrayView<const float>& BoneWeights, const float DefaultScaleWeight)
+		const TArrayView<const int32>& LODBoneIndexToWeightIndexMap, const TArrayView<const float>& BoneWeights, const float DefaultScaleWeight, const bool bInvert)
 	{
 		const int32 NumTransforms = Source.Num();
 
 		check(Dest.Num() >= NumTransforms);
 
-		for (int32 LODBoneIndex = 0; LODBoneIndex < NumTransforms; ++LODBoneIndex)
+		if (bInvert)
 		{
-			const int32 PerBoneIndex = LODBoneIndexToWeightIndexMap[LODBoneIndex];
-			const float ScaleWeight = BoneWeights.IsValidIndex(PerBoneIndex) ? BoneWeights[PerBoneIndex] : DefaultScaleWeight;
-
-			Dest.Translations[LODBoneIndex] = Source.Translations[LODBoneIndex] * ScaleWeight;
-			Dest.Rotations[LODBoneIndex] = Source.Rotations[LODBoneIndex] * ScaleWeight;
-			Dest.Scales3D[LODBoneIndex] = Source.Scales3D[LODBoneIndex] * ScaleWeight;
+			for (int32 LODBoneIndex = 0; LODBoneIndex < NumTransforms; ++LODBoneIndex)
+			{
+				const int32 PerBoneIndex = LODBoneIndexToWeightIndexMap[LODBoneIndex];
+				const float ScaleWeight = BoneWeights.IsValidIndex(PerBoneIndex) ? (1.0f - BoneWeights[PerBoneIndex]) : DefaultScaleWeight;
+
+				Dest.Translations[LODBoneIndex] = Source.Translations[LODBoneIndex] * ScaleWeight;
+				Dest.Rotations[LODBoneIndex] = Source.Rotations[LODBoneIndex] * ScaleWeight;
+				Dest.Scales3D[LODBoneIndex] = Source.Scales3D[LODBoneIndex] * ScaleWeight;
+			}
+		}
+		else
+		{
+			for (int32 LODBoneIndex = 0; LODBoneIndex < NumTransforms; ++LODBoneIndex)
+			{
+				const int32 PerBoneIndex = LODBoneIndexToWeightIndexMap[LODBoneIndex];
+				const float ScaleWeight = BoneWeights.IsValidIndex(PerBoneIndex) ? BoneWeights[PerBoneIndex] : DefaultScaleWeight;
+
+				Dest.Translations[LODBoneIndex] = Source.Translations[LODBoneIndex] * ScaleWeight;
+				Dest.Rotations[LODBoneIndex] = Source.Rotations[LODBoneIndex] * ScaleWeight;
+				Dest.Scales3D[LODBoneIndex] = Source.Scales3D[LODBoneIndex] * ScaleWeight;
+			}
 		}
 	}
 
diff --git a/Engine/Plugins/Experimental/AnimNext/Source/AnimNext/Public/TransformArrayOperations.h b/Engine/Plugins/Experimental/AnimNext/Source/AnimNext/Public/TransformArrayOperations.h
--- a/Engine/Plugins/Experimental/AnimNext/Source/AnimNext/Public/TransformArrayOperations.h
+++ b/Engine/Plugins/Experimental/AnimNext/Source/AnimNext/Public/TransformArrayOperations.h
@@ -38,13 +38,14 @@ namespace UE::AnimNext
 	ANIMNEXT_API void BlendAddWithScale(const FTransformArraySoAView& Dest, const FTransformArraySoAConstView& Source, const float ScaleWeight);
 
 	// The source transforms are scaled by the provided per bone weight and the result is written in the destination
+	// If bInvert is set, the bone weights are set to (1 - weight), DefaultScaleWeight remains unchanged, i.e. it is not inverted
 	// Dest = Source * (WeightIndex != INDEX_NONE ? Weights[WeightIndex] : DefaultScaleWeight)
 	ANIMNEXT_API void BlendOverwritePerBoneWithScale(
 		const FTransformArrayAoSView& Dest, const FTransformArrayAoSConstView& Source,
-		const TArrayView<const int32>& LODBoneIndexToWeightIndexMap, const TArrayView<const float>& BoneWeights, const float DefaultScaleWeight);
+		const TArrayView<const int32>& LODBoneIndexToWeightIndexMap, const TArrayView<const float>& BoneWeights, const float DefaultScaleWeight, const bool bInvert = false);
 	ANIMNEXT_API void BlendOverwritePerBoneWithScale(
 		const FTransformArraySoAView& Dest, const FTransformArraySoAConstView& Source,
-		const TArrayView<const int32>& LODBoneIndexToWeightIndexMap, const TArrayView<const float>& BoneWeights, const float DefaultScaleWeight);
+		const TArrayView<const int32>& LODBoneIndexToWeightIndexMap, const TArrayView<const float>& BoneWeights, const float DefaultScaleWeight, const bool bInvert = false);
 
 	// The source transforms are scaled by the provided per bone weight and the result is added to the destination
 	// Dest = Dest + (Source * (WeightIndex != INDEX_NONE ? Weights[WeightIndex] : DefaultScaleWeight))
diff --git a/Engine/Plugins/Experimental/AnimNextAnimGraph/Source/AnimNextAnimGraph/AnimNextAnimGraph.Build.cs b/Engine/Plugins/Experimental/AnimNextAnimGraph/Source/AnimNextAnimGraph/AnimNextAnimGraph.Build.cs
--- a/Engine/Plugins/Experimental/AnimNextAnimGraph/Source/AnimNextAnimGraph/AnimNextAnimGraph.Build.cs
+++ b/Engine/Plugins/Experimental/AnimNextAnimGraph/Source/AnimNextAnimGraph/AnimNextAnimGraph.Build.cs
@@ -21,6 +21,8 @@ namespace UnrealBuildTool.Rules
 					"RigVM",
 					"Engine",
 					"AnimNext",
+					"HierarchyTableRuntime",
+					"HierarchyTableBuiltinRuntime"
 				}
 			);
 
diff --git a/Engine/Plugins/Experimental/AnimNextAnimGraph/Source/AnimNextAnimGraph/Private/EvaluationVM/Tasks/BlendKeyframesPerBone.cpp b/Engine/Plugins/Experimental/AnimNextAnimGraph/Source/AnimNextAnimGraph/Private/EvaluationVM/Tasks/BlendKeyframesPerBone.cpp
--- a/Engine/Plugins/Experimental/AnimNextAnimGraph/Source/AnimNextAnimGraph/Private/EvaluationVM/Tasks/BlendKeyframesPerBone.cpp
+++ b/Engine/Plugins/Experimental/AnimNextAnimGraph/Source/AnimNextAnimGraph/Private/EvaluationVM/Tasks/BlendKeyframesPerBone.cpp
@@ -3,6 +3,8 @@
 #include "EvaluationVM/Tasks/BlendKeyframesPerBone.h"
 
 #include "AnimationRuntime.h"
+#include "Animation/AnimCurveUtils.h"
+#include "Animation/AttributeTypes.h"
 #include "Animation/BlendProfile.h"
 #include "Animation/SkeletonRemapping.h"
 #include "Animation/SkeletonRemappingRegistry.h"
@@ -11,9 +13,21 @@
 #include "TransformArrayOperations.h"
 
 FAnimNextBlendOverwriteKeyframePerBoneWithScaleTask FAnimNextBlendOverwriteKeyframePerBoneWithScaleTask::Make(TWeakObjectPtr<const UBlendProfile> BlendProfile, const FBlendSampleData& BlendData, float ScaleFactor)
+{
+	FAnimNextBlendOverwriteKeyframePerBoneWithScaleTask Task;
+	Task.BlendProfile = BlendProfile.Get();
+	Task.SourceSkeleton = BlendProfile->OwningSkeleton;
+	Task.BlendData = &BlendData;
+	Task.ScaleFactor = ScaleFactor;
+
+	return Task;
+}
+
+FAnimNextBlendOverwriteKeyframePerBoneWithScaleTask FAnimNextBlendOverwriteKeyframePerBoneWithScaleTask::Make(const IInterpolationIndexProvider* BlendProfile, const USkeleton* Skeleton, const FBlendSampleData& BlendData, float ScaleFactor)
 {
 	FAnimNextBlendOverwriteKeyframePerBoneWithScaleTask Task;
 	Task.BlendProfile = BlendProfile;
+	Task.SourceSkeleton = Skeleton;
 	Task.BlendData = &BlendData;
 	Task.ScaleFactor = ScaleFactor;
 
@@ -24,7 +38,7 @@ void FAnimNextBlendOverwriteKeyframePerBoneWithScaleTask::Execute(UE::AnimNext::
 {
 	using namespace UE::AnimNext;
 
-	const UBlendProfile* BlendProfilePtr = BlendProfile.Get();
+	const IInterpolationIndexProvider* BlendProfilePtr = BlendProfile;
 	if (BlendProfilePtr == nullptr)
 	{
 		// If we don't have a blend profile, we blend the whole pose
@@ -42,7 +56,6 @@ void FAnimNextBlendOverwriteKeyframePerBoneWithScaleTask::Execute(UE::AnimNext::
 	const USkeleton* TargetSkeleton = Keyframe->Pose.GetSkeletonAsset();
 	TSharedPtr<IInterpolationIndexProvider::FPerBoneInterpolationData> Data = BlendProfilePtr->GetPerBoneInterpolationData(TargetSkeleton);
 
-	const USkeleton* SourceSkeleton = BlendProfilePtr->OwningSkeleton;
 	const FSkeletonRemapping& SkeletonRemapping = UE::Anim::FSkeletonRemappingRegistry::Get().GetRemapping(SourceSkeleton, TargetSkeleton);
 
 	const TArrayView<const FBoneIndexType> LODBoneIndexToSkeletonBoneIndexMap = Keyframe->Pose.GetLODBoneIndexToSkeletonBoneIndexMap();
@@ -194,3 +207,189 @@ void FAnimNextBlendAddKeyframePerBoneWithScaleTask::Execute(UE::AnimNext::FEvalu
 
 	VM.PushValue(KEYFRAME_STACK_NAME, MoveTemp(KeyframeB));
 }
+
+FAnimNextBlendKeyframePerBoneWithScaleTask FAnimNextBlendKeyframePerBoneWithScaleTask::Make(TSharedPtr<IInterpolationIndexProvider> BlendProfile, const USkeleton* Skeleton, const TArray<float>& BoneMaskWeights, const UE::Anim::TNamedValueArray<FDefaultAllocator, UE::Anim::FCurveElement>& CurveMaskWeights, const TArray<FNamedFloat>& AttributeMaskWeights, float ScaleFactor)
+{
+	FAnimNextBlendKeyframePerBoneWithScaleTask Task;
+	Task.BlendProfile = BlendProfile;
+	Task.SourceSkeleton = Skeleton;
+	Task.BoneMaskWeights = BoneMaskWeights;
+	Task.CurveMaskWeights = CurveMaskWeights;
+	Task.AttributeMaskWeights = AttributeMaskWeights;
+	Task.ScaleFactor = ScaleFactor;
+
+	return Task;
+}
+
+void FAnimNextBlendKeyframePerBoneWithScaleTask::Execute(UE::AnimNext::FEvaluationVM& VM) const
+{
+	using namespace UE::AnimNext;
+
+	if (BlendProfile == nullptr)
+	{
+		// If we don't have a blend profile, we blend the whole pose
+		Super::Execute(VM);
+		return;
+	}
+
+	// Pop our top two poses, we'll re-use the top keyframe for our result
+
+	TUniquePtr<FKeyframeState> KeyframeB;
+	if (!VM.PopValue(KEYFRAME_STACK_NAME, KeyframeB))
+	{
+		// We have no inputs, nothing to do
+		return;
+	}
+
+	TUniquePtr<FKeyframeState> KeyframeA;
+	if (!VM.PopValue(KEYFRAME_STACK_NAME, KeyframeA))
+	{
+		// We have a single input, leave it on top of the stack
+		VM.PushValue(KEYFRAME_STACK_NAME, MoveTemp(KeyframeB));
+		return;
+	}
+
+	const USkeleton* TargetSkeleton = KeyframeB->Pose.GetSkeletonAsset();
+	TSharedPtr<IInterpolationIndexProvider::FPerBoneInterpolationData> Data = BlendProfile->GetPerBoneInterpolationData(TargetSkeleton);
+
+	const FSkeletonRemapping& SkeletonRemapping = UE::Anim::FSkeletonRemappingRegistry::Get().GetRemapping(SourceSkeleton, TargetSkeleton);
+
+	const TArrayView<const FBoneIndexType> LODBoneIndexToSkeletonBoneIndexMap = KeyframeB->Pose.GetLODBoneIndexToSkeletonBoneIndexMap();
+	const int32 NumLODBones = LODBoneIndexToSkeletonBoneIndexMap.Num();
+
+	TArray<int32, TMemStackAllocator<>> LODBoneIndexToWeightIndexMap;
+	LODBoneIndexToWeightIndexMap.AddUninitialized(NumLODBones);
+
+	if (SkeletonRemapping.IsValid())
+	{
+		for (int32 LODBoneIndex = 0; LODBoneIndex < NumLODBones; ++LODBoneIndex)
+		{
+			const FSkeletonPoseBoneIndex TargetSkeletonBoneIndex(LODBoneIndexToSkeletonBoneIndexMap[LODBoneIndex]);
+			const FSkeletonPoseBoneIndex SourceSkeletonBoneIndex(SkeletonRemapping.GetSourceSkeletonBoneIndex(TargetSkeletonBoneIndex.GetInt()));
+			LODBoneIndexToWeightIndexMap[LODBoneIndex] = SourceSkeletonBoneIndex.IsValid() ? BlendProfile->GetPerBoneInterpolationIndex(SourceSkeletonBoneIndex, TargetSkeleton, Data.Get()) : INDEX_NONE;
+		}
+	}
+	else
+	{
+		for (int32 LODBoneIndex = 0; LODBoneIndex < NumLODBones; ++LODBoneIndex)
+		{
+			const FSkeletonPoseBoneIndex SkeletonBoneIndex(LODBoneIndexToSkeletonBoneIndexMap[LODBoneIndex]);
+			LODBoneIndexToWeightIndexMap[LODBoneIndex] = BlendProfile->GetPerBoneInterpolationIndex(SkeletonBoneIndex, TargetSkeleton, Data.Get());
+		}
+	}
+
+	if (EnumHasAnyFlags(VM.GetFlags(), EEvaluationFlags::Bones))
+	{
+		check(KeyframeA->Pose.GetNumBones() == KeyframeB->Pose.GetNumBones());
+
+		BlendOverwritePerBoneWithScale(
+			KeyframeA->Pose.LocalTransforms.GetView(), KeyframeA->Pose.LocalTransforms.GetConstView(),
+			LODBoneIndexToWeightIndexMap, BoneMaskWeights, 1.0f, true);
+
+		BlendAddPerBoneWithScale(
+			KeyframeA->Pose.LocalTransforms.GetView(), KeyframeB->Pose.LocalTransforms.GetConstView(),
+			LODBoneIndexToWeightIndexMap, BoneMaskWeights, ScaleFactor);
+	}
+
+	if (EnumHasAnyFlags(VM.GetFlags(), EEvaluationFlags::Curves))
+	{
+		FBlendedCurve& OutCurve = KeyframeA->Curves;
+		FBlendedCurve& TargetCurve = KeyframeB->Curves;
+
+		if (FAnimWeight::IsRelevant(ScaleFactor))
+		{
+			FBlendedCurve FilteredCurves;
+
+			// Multiply per-curve blend weights by matching blend pose curves
+			UE::Anim::FNamedValueArrayUtils::Intersection(
+				TargetCurve,
+				CurveMaskWeights,
+				[this, &FilteredCurves](const UE::Anim::FCurveElement& InBlendElement, const UE::Anim::FCurveElement& InMaskElement) mutable
+				{
+					FilteredCurves.Add(InBlendElement.Name, InBlendElement.Value * InMaskElement.Value);
+				});
+
+			// Override blend curve values with premultipled curves
+			TargetCurve.Combine(FilteredCurves);
+
+			// Remove curves that have been filtered by the mask, curves with no mask value defined remain, even with 0.0 value
+			UE::Anim::FNamedValueArrayUtils::RemoveByPredicate(
+				TargetCurve,
+				CurveMaskWeights,
+				[](const UE::Anim::FCurveElement& InBaseElement, const UE::Anim::FCurveElement& InMaskElement)
+				{
+					const bool bKeep = InMaskElement.Value == 0.0;
+					return bKeep;
+				}
+			);
+
+			// Combine base and filtered pre-multiplied blend curves
+			UE::Anim::FNamedValueArrayUtils::Union(
+				OutCurve,
+				TargetCurve,
+				[this](UE::Anim::FCurveElement& InOutBaseElement, const UE::Anim::FCurveElement& InBlendElement, UE::Anim::ENamedValueUnionFlags InFlags)
+				{
+					InOutBaseElement.Value = FMath::Lerp(InOutBaseElement.Value, InBlendElement.Value, ScaleFactor);
+					InOutBaseElement.Flags |= InBlendElement.Flags;
+				});
+		}
+	}
+
+	if (EnumHasAnyFlags(VM.GetFlags(), EEvaluationFlags::Attributes))
+	{
+		using namespace UE::Anim;
+
+		FStackAttributeContainer& BaseAttributes = KeyframeA->Attributes;
+		FStackAttributeContainer& BlendAttributes = KeyframeB->Attributes;
+
+		// Premultiply masked blend attributes
+		const TArray<TWeakObjectPtr<UScriptStruct>>& BlendUniqueTypes = BlendAttributes.GetUniqueTypes();
+
+		//CachedAttributeMaskWeights.ForEachElement([&BlendUniqueTypes, &BlendAttributes](const FNamedFloat& AttributeMask)
+		for (const FNamedFloat& AttributeMask : AttributeMaskWeights)
+		{
+			for (int32 BlendUniqueTypeIndex = 0; BlendUniqueTypeIndex < BlendUniqueTypes.Num(); ++BlendUniqueTypeIndex)
+			{
+				const TArray<FAttributeId>& TypedAttributeIdentifiers = BlendAttributes.GetKeys(BlendUniqueTypeIndex);
+				const TArray<UE::Anim::TWrappedAttribute<FAnimStackAllocator>>& TypedAttributeValues = BlendAttributes.GetValues(BlendUniqueTypeIndex);
+
+				const TWeakObjectPtr<UScriptStruct> AttributeType = BlendUniqueTypes[BlendUniqueTypeIndex];
+				const IAttributeBlendOperator* Operator = UE::Anim::AttributeTypes::GetTypeOperator(AttributeType);
+
+				TWrappedAttribute<FAnimStackAllocator> DefaultData;
+				AttributeType->InitializeStruct(DefaultData.GetPtr<void>());
+
+				const int32 AttributeIndex = TypedAttributeIdentifiers.IndexOfByPredicate([&AttributeMask](const FAttributeId& InAttribute)
+					{
+						return InAttribute.GetName() == AttributeMask.Name;
+					});
+
+				if (AttributeIndex != INDEX_NONE)
+				{
+					const FAttributeId& AttributeId = TypedAttributeIdentifiers[AttributeIndex];
+
+					if (AttributeMask.Value == 0.0f)
+					{
+						BlendAttributes.Remove(AttributeType.Get(), AttributeId);
+					}
+					else
+					{
+						uint8* Value = BlendAttributes.Find(AttributeType.Get(), AttributeId);
+						check(Value != nullptr);
+
+						Operator->Interpolate(DefaultData.GetPtr<void>(), Value, AttributeMask.Value, Value);
+					}
+					break;
+				}
+			}
+		}
+
+		const TArray<const UE::Anim::FStackAttributeContainer> SourceAttributes{ BaseAttributes, BlendAttributes };
+		const TArray<const float> SourceWeights{ 1.0f - ScaleFactor, ScaleFactor };
+		static const TArray<const int32> WeightView{ 0, 1 };
+
+		UE::Anim::Attributes::BlendAttributes(SourceAttributes, SourceWeights, WeightView, BaseAttributes);
+	}
+
+	VM.PushValue(KEYFRAME_STACK_NAME, MoveTemp(KeyframeA));
+}
\ No newline at end of file
diff --git a/Engine/Plugins/Experimental/AnimNextAnimGraph/Source/AnimNextAnimGraph/Private/Traits/BlendLayer.cpp b/Engine/Plugins/Experimental/AnimNextAnimGraph/Source/AnimNextAnimGraph/Private/Traits/BlendLayer.cpp
new file mode 100644
--- /dev/null
+++ b/Engine/Plugins/Experimental/AnimNextAnimGraph/Source/AnimNextAnimGraph/Private/Traits/BlendLayer.cpp
@@ -0,0 +1,243 @@
+// Copyright Epic Games, Inc. All Rights Reserved.
+
+#include "Traits/BlendLayer.h"
+
+#include "AnimationRuntime.h"
+#include "Animation/AnimTypes.h"
+#include "TraitCore/ExecutionContext.h"
+#include "EvaluationVM/Tasks/BlendKeyframes.h"
+#include "EvaluationVM/Tasks/NormalizeRotations.h"
+#include "HierarchyTable.h"
+#include "MaskProfile/HierarchyTableTypeMask.h"
+#include "SkeletonHierarchyTableType.h"
+
+class FHierarchyTableBlendProfileProvider : public IInterpolationIndexProvider
+{
+public:
+	TObjectPtr<UHierarchyTable> HierarchyTable;
+
+	// IInterpolationIndexProvider
+	virtual TSharedPtr<FPerBoneInterpolationData> GetPerBoneInterpolationData(const USkeleton* Skeleton) const
+	{
+		return nullptr;
+	}
+
+	virtual int32 GetPerBoneInterpolationIndex(const FCompactPoseBoneIndex& InCompactPoseBoneIndex, const FBoneContainer& BoneContainer, const IInterpolationIndexProvider::FPerBoneInterpolationData* Data) const override
+	{
+		return 0;
+	}
+
+	virtual int32 GetPerBoneInterpolationIndex(const FSkeletonPoseBoneIndex InSkeletonBoneIndex, const USkeleton* TargetSkeleton, const IInterpolationIndexProvider::FPerBoneInterpolationData* Data) const override
+	{
+		if (HierarchyTable)
+		{
+			if (HierarchyTable->GetTableData().IsValidIndex(InSkeletonBoneIndex.GetInt()))
+			{
+				return InSkeletonBoneIndex.GetInt();
+			}
+		}
+		return INDEX_NONE;
+	}
+	// End IInterpolationIndexProvider
+};
+
+
+namespace UE::AnimNext
+{
+	AUTO_REGISTER_ANIM_TRAIT(FBlendLayerTrait)
+
+	#define TRAIT_INTERFACE_ENUMERATOR(GeneratorMacro) \
+		GeneratorMacro(IContinuousBlend) \
+		GeneratorMacro(IEvaluate) \
+		GeneratorMacro(IHierarchy) \
+		GeneratorMacro(IUpdate) \
+		GeneratorMacro(IUpdateTraversal) \
+
+	GENERATE_ANIM_TRAIT_IMPLEMENTATION(FBlendLayerTrait, TRAIT_INTERFACE_ENUMERATOR, NULL_ANIM_TRAIT_INTERFACE_ENUMERATOR, NULL_ANIM_TRAIT_EVENT_ENUMERATOR)
+	#undef TRAIT_INTERFACE_ENUMERATOR
+
+	void FBlendLayerTrait::PostEvaluate(FEvaluateTraversalContext& Context, const TTraitBinding<IEvaluate>& Binding) const
+	{
+		const FSharedData* SharedData = Binding.GetSharedData<FSharedData>();
+		const FInstanceData* InstanceData = Binding.GetInstanceData<FInstanceData>();
+
+		TObjectPtr<UHierarchyTable> HierarchyTable = SharedData->BlendProfile;
+
+		if (InstanceData->ChildBase.IsValid() && InstanceData->ChildBlend.IsValid() && HierarchyTable)
+		{
+			// We have two children, interpolate them
+
+			TTraitBinding<IContinuousBlend> ContinuousBlendTrait;
+			Binding.GetStackInterface(ContinuousBlendTrait);
+
+			TSharedPtr<FHierarchyTableBlendProfileProvider> InterpolationIndexProvider = MakeShared<FHierarchyTableBlendProfileProvider>();
+			InterpolationIndexProvider->HierarchyTable = HierarchyTable;
+
+			const FHierarchyTable_TableType_Skeleton& TableMetadata = HierarchyTable->GetTableMetadata<FHierarchyTable_TableType_Skeleton>();
+			const TObjectPtr<USkeleton> Skeleton = TableMetadata.Skeleton;
+
+			const float BlendWeight = ContinuousBlendTrait.GetBlendWeight(Context, 1);
+
+			Context.AppendTask(FAnimNextBlendKeyframePerBoneWithScaleTask::Make(InterpolationIndexProvider, Skeleton, InstanceData->BoneMaskWeights, InstanceData->CurveMaskWeights, InstanceData->AttributeMaskWeights, BlendWeight));
+			Context.AppendTask(FAnimNextNormalizeKeyframeRotationsTask());
+		}
+		else
+		{
+			// We have only one child that is active, do nothing
+		}
+	}
+
+	void FBlendLayerTrait::PreUpdate(FUpdateTraversalContext& Context, const TTraitBinding<IUpdate>& Binding, const FTraitUpdateState& TraitState) const
+	{
+		const FSharedData* SharedData = Binding.GetSharedData<FSharedData>();
+		FInstanceData* InstanceData = Binding.GetInstanceData<FInstanceData>();
+
+		TTraitBinding<IContinuousBlend> ContinuousBlendTrait;
+		Binding.GetStackInterface(ContinuousBlendTrait);
+
+		TObjectPtr<UHierarchyTable> HierarchyTable = SharedData->BlendProfile;
+
+		if (InstanceData->bBoneMaskWeightsNeedEvaluating)
+		{
+			InstanceData->BoneMaskWeights.Empty();
+
+			if (HierarchyTable)
+			{
+				const int32 NumBlendEntries = SharedData->BlendProfile->GetTableData().Num();
+				InstanceData->BoneMaskWeights.SetNum(NumBlendEntries);
+
+				for (int32 Index = 0; Index < NumBlendEntries; ++Index)
+				{
+					InstanceData->BoneMaskWeights[Index] = HierarchyTable->GetTableEntry(Index)->GetValue<FHierarchyTable_ElementType_Mask>()->Value;
+				}
+			}
+
+			InstanceData->bBoneMaskWeightsNeedEvaluating = false;
+		}
+
+		if (InstanceData->bCurveMaskWeightsNeedEvaluating)
+		{
+			InstanceData->CurveMaskWeights.Empty();
+
+			if (HierarchyTable)
+			{
+				for (const FHierarchyTableEntryData& TableEntry : HierarchyTable->GetTableData())
+				{
+					if (TableEntry.GetMetadata<FHierarchyTable_TablePayloadType_Skeleton>().EntryType == ESkeletonHierarchyTable_TablePayloadEntryType::Curve)
+					{
+						const float EntryWeight = TableEntry.GetValue<FHierarchyTable_ElementType_Mask>()->Value;
+						InstanceData->CurveMaskWeights.Add<UE::Anim::FCurveElement>({ TableEntry.Identifier, EntryWeight });
+					}
+				}
+			}
+
+			InstanceData->bCurveMaskWeightsNeedEvaluating = false;
+		}
+
+		if (InstanceData->bAttributeMaskWeightsNeedEvaluating)
+		{
+			InstanceData->AttributeMaskWeights.Empty();
+
+			if (HierarchyTable)
+			{
+				for (const FHierarchyTableEntryData& TableEntry : HierarchyTable->GetTableData())
+				{
+					if (TableEntry.GetMetadata<FHierarchyTable_TablePayloadType_Skeleton>().EntryType == ESkeletonHierarchyTable_TablePayloadEntryType::Attribute)
+					{
+						const float EntryWeight = TableEntry.GetValue<FHierarchyTable_ElementType_Mask>()->Value;
+						InstanceData->AttributeMaskWeights.Add(FAnimNextBlendKeyframePerBoneWithScaleTask::FNamedFloat(TableEntry.Identifier, EntryWeight));
+					}
+				}
+			}
+
+			InstanceData->bAttributeMaskWeightsNeedEvaluating = false;
+		}
+
+		if (!InstanceData->ChildBase.IsValid())
+		{
+			InstanceData->ChildBase = Context.AllocateNodeInstance(Binding, SharedData->ChildBase);
+		}
+		else
+		{
+			InstanceData->bWasChildBaseRelevant = true;
+		}
+
+		if (FAnimWeight::IsRelevant(SharedData->BlendWeight) && SharedData->BlendProfile)
+		{
+			if (!InstanceData->ChildBlend.IsValid())
+			{
+				// We need to blend a child that isn't instanced yet, allocate it
+				InstanceData->ChildBlend = Context.AllocateNodeInstance(Binding, SharedData->ChildBlend);
+			}
+			else
+			{
+				InstanceData->bWasChildBlendRelevant = true;
+			}
+		}
+		else
+		{
+			InstanceData->ChildBlend.Reset();
+		}
+	}
+
+	void FBlendLayerTrait::QueueChildrenForTraversal(FUpdateTraversalContext& Context, const TTraitBinding<IUpdateTraversal>& Binding, const FTraitUpdateState& TraitState, FUpdateTraversalQueue& TraversalQueue) const
+	{
+		const FSharedData* SharedData = Binding.GetSharedData<FSharedData>();
+		const FInstanceData* InstanceData = Binding.GetInstanceData<FInstanceData>();
+
+		TTraitBinding<IContinuousBlend> ContinuousBlendTrait;
+		Binding.GetStackInterface(ContinuousBlendTrait);
+
+		const float BlendWeightBlend = ContinuousBlendTrait.GetBlendWeight(Context, 1);
+		if (InstanceData->ChildBase.IsValid())
+		{
+			TraversalQueue.Push(InstanceData->ChildBase, TraitState.WithWeight(1.0f).AsNewlyRelevant(!InstanceData->bWasChildBaseRelevant));
+		}
+
+		if (InstanceData->ChildBlend.IsValid() && SharedData->BlendProfile)
+		{
+			TraversalQueue.Push(InstanceData->ChildBlend, TraitState.WithWeight(BlendWeightBlend).AsNewlyRelevant(!InstanceData->bWasChildBlendRelevant));
+		}
+	}
+
+	uint32 FBlendLayerTrait::GetNumChildren(const FExecutionContext& Context, const TTraitBinding<IHierarchy>& Binding) const
+	{
+		return 2;
+	}
+
+	void FBlendLayerTrait::GetChildren(const FExecutionContext& Context, const TTraitBinding<IHierarchy>& Binding, FChildrenArray& Children) const
+	{
+		const FSharedData* SharedData = Binding.GetSharedData<FSharedData>();
+		const FInstanceData* InstanceData = Binding.GetInstanceData<FInstanceData>();
+
+		// Add the two children, even if the handles are empty
+		Children.Add(InstanceData->ChildBase);
+		Children.Add(InstanceData->ChildBlend);
+	}
+
+	float FBlendLayerTrait::GetBlendWeight(const FExecutionContext& Context, const TTraitBinding<IContinuousBlend>& Binding, int32 ChildIndex) const
+	{
+		const FSharedData* SharedData = Binding.GetSharedData<FSharedData>();
+
+		const float BlendWeight = SharedData->GetBlendWeight(Binding);
+		float ClampedWeight = FMath::Clamp(BlendWeight, 0.0f, 1.0f);
+		if (!SharedData->BlendProfile)
+		{
+			ClampedWeight = 0.0f;
+		}
+
+		if (ChildIndex == 0)
+		{
+			return 1.0f - ClampedWeight;
+		}
+		else if (ChildIndex == 1)
+		{
+			return ClampedWeight;
+		}
+		else
+		{
+			// Invalid child index
+			return -1.0f;
+		}
+	}
+}
\ No newline at end of file
diff --git a/Engine/Plugins/Experimental/AnimNextAnimGraph/Source/AnimNextAnimGraph/Private/Traits/BlendLayer.h b/Engine/Plugins/Experimental/AnimNextAnimGraph/Source/AnimNextAnimGraph/Private/Traits/BlendLayer.h
new file mode 100644
--- /dev/null
+++ b/Engine/Plugins/Experimental/AnimNextAnimGraph/Source/AnimNextAnimGraph/Private/Traits/BlendLayer.h
@@ -0,0 +1,93 @@
+// Copyright Epic Games, Inc. All Rights Reserved.
+
+#pragma once
+
+#include "CoreMinimal.h"
+
+#include "EvaluationVM/Tasks/BlendKeyframesPerBone.h"
+#include "HierarchyTable.h"
+#include "TraitCore/Trait.h"
+#include "TraitInterfaces/IContinuousBlend.h"
+#include "TraitInterfaces/IEvaluate.h"
+#include "TraitInterfaces/IHierarchy.h"
+#include "TraitInterfaces/IUpdate.h"
+
+#include "BlendLayer.generated.h"
+
+USTRUCT(meta = (DisplayName = "Blend Layer"))
+struct FAnimNextBlendLayerTraitSharedData : public FAnimNextTraitSharedData
+{
+	GENERATED_BODY()
+
+	UPROPERTY()
+	FAnimNextTraitHandle ChildBase;
+
+	UPROPERTY()
+	FAnimNextTraitHandle ChildBlend;
+
+	UPROPERTY(EditAnywhere, Category = "Default")
+	float BlendWeight = 1.0f;
+
+	/** Blend profile that configures how fast to blend each bone. */
+	// TODO: Can't show list of blend profiles, we need to find a skeleton to perform the lookup with
+	UPROPERTY(EditAnywhere, Category = "Default", meta = (Inline))
+	TObjectPtr<UHierarchyTable> BlendProfile = nullptr;
+
+	#define TRAIT_LATENT_PROPERTIES_ENUMERATOR(GeneratorMacro) \
+		GeneratorMacro(BlendWeight) \
+
+	GENERATE_TRAIT_LATENT_PROPERTIES(FAnimNextBlendLayerTraitSharedData, TRAIT_LATENT_PROPERTIES_ENUMERATOR)
+	#undef TRAIT_LATENT_PROPERTIES_ENUMERATOR
+};
+
+namespace UE::AnimNext
+{
+	/**
+	 * FBlendLayerTrait
+	 * 
+	 * A trait that can blend a layer into an input.
+	 */
+	 struct FBlendLayerTrait : FBaseTrait, IEvaluate, IUpdate, IUpdateTraversal, IHierarchy, IContinuousBlend
+	 {
+		DECLARE_ANIM_TRAIT(FBlendLayerTrait, 0x57d56b7f, FBaseTrait)
+
+
+		using FSharedData = FAnimNextBlendLayerTraitSharedData;
+
+		struct FInstanceData : FTrait::FInstanceData
+		{
+			FTraitPtr ChildBase;
+			FTraitPtr ChildBlend;
+
+			bool bWasChildBaseRelevant = false;
+			bool bWasChildBlendRelevant = false;
+
+			bool bBoneMaskWeightsNeedEvaluating = true;
+			bool bCurveMaskWeightsNeedEvaluating = true;
+			bool bAttributeMaskWeightsNeedEvaluating = true;
+
+			TArray<float> BoneMaskWeights;
+
+			UE::Anim::TNamedValueArray<FDefaultAllocator, UE::Anim::FCurveElement> CurveMaskWeights;
+
+			TArray<FAnimNextBlendKeyframePerBoneWithScaleTask::FNamedFloat> AttributeMaskWeights;
+		};
+
+
+		// IEvaluate impl
+		virtual void PostEvaluate(FEvaluateTraversalContext& Context, const TTraitBinding<IEvaluate>& Binding) const override;
+
+		// IUpdate impl
+		virtual void PreUpdate(FUpdateTraversalContext& Context, const TTraitBinding<IUpdate>& Binding, const FTraitUpdateState& TraitState) const override;
+
+		// IUpdateTraversal impl
+		virtual void QueueChildrenForTraversal(FUpdateTraversalContext& Context, const TTraitBinding<IUpdateTraversal>& Binding, const FTraitUpdateState& TraitState, FUpdateTraversalQueue& TraversalQueue) const override;
+
+		// IHierarchy impl
+		virtual uint32 GetNumChildren(const FExecutionContext& Context, const TTraitBinding<IHierarchy>& Binding) const override;
+		virtual void GetChildren(const FExecutionContext& Context, const TTraitBinding<IHierarchy>& Binding, FChildrenArray& Children) const override;
+
+		// IContinuousBlend impl
+		virtual float GetBlendWeight(const FExecutionContext& Context, const TTraitBinding<IContinuousBlend>& Binding, int32 ChildIndex) const override;
+	 };
+}
\ No newline at end of file
diff --git a/Engine/Plugins/Experimental/AnimNextAnimGraph/Source/AnimNextAnimGraph/Public/EvaluationVM/Tasks/BlendKeyframesPerBone.h b/Engine/Plugins/Experimental/AnimNextAnimGraph/Source/AnimNextAnimGraph/Public/EvaluationVM/Tasks/BlendKeyframesPerBone.h
--- a/Engine/Plugins/Experimental/AnimNextAnimGraph/Source/AnimNextAnimGraph/Public/EvaluationVM/Tasks/BlendKeyframesPerBone.h
+++ b/Engine/Plugins/Experimental/AnimNextAnimGraph/Source/AnimNextAnimGraph/Public/EvaluationVM/Tasks/BlendKeyframesPerBone.h
@@ -2,6 +2,8 @@
 
 #pragma once
 
+#include "Animation/AnimCurveTypes.h"
+#include "Animation/NamedValueArray.h"
 #include "CoreMinimal.h"
 #include "EvaluationVM/EvaluationTask.h"
 #include "EvaluationVM/Tasks/BlendKeyframes.h"
@@ -10,6 +12,8 @@
 
 struct FBlendSampleData;
 class UBlendProfile;
+class IInterpolationIndexProvider;
+class USkeleton;
 
 /*
  * Blend Overwrite Keyframe Per Bone With Scale Task
@@ -30,6 +34,8 @@ struct ANIMNEXTANIMGRAPH_API FAnimNextBlendOverwriteKeyframePerBoneWithScaleTask
 
 	static FAnimNextBlendOverwriteKeyframePerBoneWithScaleTask Make(TWeakObjectPtr<const UBlendProfile> BlendProfile, const FBlendSampleData& BlendData, float ScaleFactor);
 
+	static FAnimNextBlendOverwriteKeyframePerBoneWithScaleTask Make(const IInterpolationIndexProvider* BlendProfile, const USkeleton* Skeleton, const FBlendSampleData& BlendData, float ScaleFactor);
+
 	// Task entry point
 	virtual void Execute(UE::AnimNext::FEvaluationVM& VM) const override;
 
@@ -40,8 +46,9 @@ struct ANIMNEXTANIMGRAPH_API FAnimNextBlendOverwriteKeyframePerBoneWithScaleTask
 
 	// The blend profile to use
 	// An optional blend profile provides a blend weight per bone
-	UPROPERTY()
-	TWeakObjectPtr<const UBlendProfile> BlendProfile;
+	const IInterpolationIndexProvider* BlendProfile = nullptr;
+
+	const USkeleton* SourceSkeleton = nullptr;
 };
 
 /*
@@ -82,3 +89,40 @@ struct ANIMNEXTANIMGRAPH_API FAnimNextBlendAddKeyframePerBoneWithScaleTask : pub
 	UPROPERTY()
 	TWeakObjectPtr<const UBlendProfile> BlendProfile;
 };
+
+
+USTRUCT()
+struct ANIMNEXTANIMGRAPH_API FAnimNextBlendKeyframePerBoneWithScaleTask : public FAnimNextBlendAddKeyframeWithScaleTask
+{
+	GENERATED_BODY()
+
+	DECLARE_ANIM_EVALUATION_TASK(FAnimNextBlendKeyframePerBoneWithScaleTask)
+
+	struct FNamedFloat
+	{
+		FName Name;
+		float Value;
+
+		FNamedFloat(const FName InName, const float InValue)
+			: Name(InName)
+			, Value(InValue)
+		{
+		}
+	};
+
+	static FAnimNextBlendKeyframePerBoneWithScaleTask Make(TSharedPtr<IInterpolationIndexProvider> BlendProfile, const USkeleton* Skeleton, const TArray<float>& BoneMaskWeights, const UE::Anim::TNamedValueArray<FDefaultAllocator, UE::Anim::FCurveElement>& CurveMaskWeights, const TArray<FNamedFloat>& AttributeMaskWeights, float ScaleFactor);
+
+	// Task entry point
+	virtual void Execute(UE::AnimNext::FEvaluationVM& VM) const override;
+
+private:
+	TArray<float> BoneMaskWeights;
+
+	UE::Anim::TNamedValueArray<FDefaultAllocator, UE::Anim::FCurveElement> CurveMaskWeights;
+
+	TArray<FNamedFloat> AttributeMaskWeights;
+
+	TSharedPtr<IInterpolationIndexProvider> BlendProfile = nullptr;
+
+	const USkeleton* SourceSkeleton = nullptr;
+};
diff --git a/Engine/Plugins/Experimental/Animation/HierarchyTable/Source/Runtime/Public/HierarchyTable.h b/Engine/Plugins/Experimental/Animation/HierarchyTable/Source/Runtime/Public/HierarchyTable.h
--- a/Engine/Plugins/Experimental/Animation/HierarchyTable/Source/Runtime/Public/HierarchyTable.h
+++ b/Engine/Plugins/Experimental/Animation/HierarchyTable/Source/Runtime/Public/HierarchyTable.h
@@ -66,6 +66,12 @@ public:
 		return StructPtr->GetPtr<T>();
 	}
 
+	template <typename T>
+	const T& GetMetadata() const
+	{
+		return TablePayload.Get<T>();
+	}
+
 	template <typename T>
 	T* GetMutableValue()
 	{
@@ -89,7 +95,7 @@ private:
 /**
  * A general-purpose container asset for storing typed hierarchical data.
  */
-UCLASS(EditInlineNew)
+UCLASS(EditInlineNew, BlueprintType)
 class HIERARCHYTABLERUNTIME_API UHierarchyTable : public UObject
 {
 	GENERATED_BODY()
