Add support for markers in AnimNext sync groups
Cleaned up AnimNext ITimeline interface and related things

#virtualized

diff --git a/Engine/Plugins/Experimental/AnimNextAnimGraph/Source/AnimNextAnimGraph/Internal/Traits/BlendStackTrait.h b/Engine/Plugins/Experimental/AnimNextAnimGraph/Source/AnimNextAnimGraph/Internal/Traits/BlendStackTrait.h
--- a/Engine/Plugins/Experimental/AnimNextAnimGraph/Source/AnimNextAnimGraph/Internal/Traits/BlendStackTrait.h
+++ b/Engine/Plugins/Experimental/AnimNextAnimGraph/Source/AnimNextAnimGraph/Internal/Traits/BlendStackTrait.h
@@ -154,8 +154,7 @@ namespace UE::AnimNext
 		virtual UCurveFloat* GetCustomBlendCurve(FExecutionContext& Context, const TTraitBinding<ISmoothBlend>& Binding, int32 ChildIndex) const override;
 
 		// ITimeline impl
-		virtual float GetPlayRate(FExecutionContext& Context, const TTraitBinding<ITimeline>& Binding) const;
-		virtual FTimelineProgress GetProgress(FExecutionContext& Context, const TTraitBinding<ITimeline>& Binding) const;
+		virtual FTimelineState GetState(FExecutionContext& Context, const TTraitBinding<ITimeline>& Binding) const override;
 
 		// IGarbageCollection impl
 		virtual void AddReferencedObjects(const FExecutionContext& Context, const TTraitBinding<IGarbageCollection>& Binding, FReferenceCollector& Collector) const override;
diff --git a/Engine/Plugins/Experimental/AnimNextAnimGraph/Source/AnimNextAnimGraph/Private/Graph/SyncGroup_GraphInstanceComponent.cpp b/Engine/Plugins/Experimental/AnimNextAnimGraph/Source/AnimNextAnimGraph/Private/Graph/SyncGroup_GraphInstanceComponent.cpp
--- a/Engine/Plugins/Experimental/AnimNextAnimGraph/Source/AnimNextAnimGraph/Private/Graph/SyncGroup_GraphInstanceComponent.cpp
+++ b/Engine/Plugins/Experimental/AnimNextAnimGraph/Source/AnimNextAnimGraph/Private/Graph/SyncGroup_GraphInstanceComponent.cpp
@@ -2,48 +2,140 @@
 
 #include "Graph/SyncGroup_GraphInstanceComponent.h"
 #include "TraitInterfaces/IGroupSynchronization.h"
+#include "TraitInterfaces/ITimeline.h"
+#include "Animation/AnimTypes.h"
+#include "Animation/AnimationAsset.h"
+#include "AnimationUtils.h"
+
+#define UE_LOG_ENABLE_MARKER_DEBUG 0
 
 namespace UE::AnimNext
 {
-	FSyncGroupGraphInstanceComponent::FSyncGroupGraphInstanceComponent(FAnimNextGraphInstance& InOwnerInstance)
-		: FGraphInstanceComponent(InOwnerInstance)
+	namespace Private
 	{
-	}
+		// Sync group member state as collected during graph traversal
+		struct FSyncGroupMember
+		{
+			FTraitUpdateState					TraitState;
+			FWeakTraitPtr						TraitPtr;
 
-	void FSyncGroupGraphInstanceComponent::RegisterWithGroup(FName GroupName, EAnimGroupRole::Type GroupRole, const FWeakTraitPtr& TraitPtr, const FTraitUpdateState& TraitState)
-	{
-		FSyncGroupState& GroupState = SyncGroupMap.FindOrAdd(GroupName);
-		GroupState.Members.Add(FSyncGroupMember{ TraitState, TraitPtr, GroupRole });
-	}
+			FSyncGroupParameters				GroupParameters;
+		};
 
-	void FSyncGroupGraphInstanceComponent::PreUpdate(FExecutionContext& Context)
-	{
-		// Reset our group state, we want to start fresh every update
-		SyncGroupMap.Reset();
-	}
+		// Sync group state as collected during graph traversal
+		struct FSyncGroupState
+		{
+			// The sync group name
+			FName GroupName;
 
-	void FSyncGroupGraphInstanceComponent::PostUpdate(FExecutionContext& Context)
-	{
-		FTraitStackBinding TraitStack;
-		TTraitBinding<IGroupSynchronization> GroupSyncTrait;
+			// The list of sync group members during the current update
+			TArray<FSyncGroupMember> Members;
 
-		// Now that we have discovered all groups and their memberships, we can perform synchronization
-		for (const TTuple<FName, FSyncGroupState>& It : SyncGroupMap)
+			// If this sync group was active during the previous update, this is the timeline progress of its leader
+			FTimelineState PreviousLeaderTimelineState;
+
+			// If this sync group was active during the previous update and using sync markers, this is the sync position of the leader
+			FMarkerSyncAnimPosition PreviousLeaderSyncPosition;
+
+			// Whether or not this sync group was active during the previous update
+			bool bWasActive = false;
+		};
+
+		// Sync group member context when performing group synchronization
+		struct FSyncGroupMemberContext
 		{
-			const FSyncGroupState& GroupState = It.Value;
+			const FSyncGroupMember* State = nullptr;
 
-			// Find our leader by looking at the total weight and the group role
+			FTraitStackBinding TraitStack;
+			TTraitBinding<IGroupSynchronization> GroupSyncTrait;
+			
+			FTimelineState TimelineState;
+
+			FTimelineSyncMarkerArray SyncMarkers;
+			bool bUseMarkerSyncing = false;
+
+			void Init(const FSyncGroupMember& InState)
+			{
+				State = &InState;
+				TraitStack.Reset();
+				GroupSyncTrait.Reset();
+				TimelineState.Reset();
+				SyncMarkers.Reset();
+				bUseMarkerSyncing = false;
+			}
+		};
+
+		// Sync group context when performing group synchronization
+		struct FSyncGroupContext
+		{
+			TArrayView<FSyncGroupMemberContext> Members;
+
+			// If this sync group was active during the previous update, this is the timeline progress of its leader
+			FTimelineState PreviousLeaderTimelineState;
+
+			// If this sync group was active during the previous update and using sync markers, this is the sync position of the leader
+			FMarkerSyncAnimPosition PreviousLeaderSyncPosition;
+
+			// Current leader progress, set after leader has advanced
+			FTimelineState LeaderTimelineState;
+			int32 LeaderIndex = INDEX_NONE;
+			float LeaderDeltaTime = 0.0f;
+
+			// Whether or not this sync group was active during the previous update
+			bool bWasActive = false;
+
+			bool bCanGroupUseMarkerSyncing = false;
+			FMarkerSyncAnimPosition SyncPositionStart;
+			FMarkerSyncAnimPosition SyncPositionEnd;
+			FMarkerPair LeaderStartMarker;
+			FMarkerPair LeaderEndMarker;
+
+			TSet<FName> ValidMarkers;
+			TSet<FName> CandidateMarkers;
+			TBitArray<TMemStackAllocator<>> GroupSeenMarkers;
+			TBitArray<TMemStackAllocator<>> MemberSeenMarkers;
+		};
+
+		static void InitGroup(const FSyncGroupState& GroupState, TArray<FSyncGroupMemberContext, TMemStackAllocator<>>& ContextPool, FSyncGroupContext& GroupContext)
+		{
 			const int32 NumMembers = GroupState.Members.Num();
 			check(NumMembers != 0);	// Groups should never be empty
 
+			TArrayView<FSyncGroupMemberContext> Contextes(ContextPool.GetData(), NumMembers);
+			for (int32 MemberIndex = 0; MemberIndex < NumMembers; ++MemberIndex)
+			{
+				Contextes[MemberIndex].Init(GroupState.Members[MemberIndex]);
+			}
+
+			GroupContext.Members = Contextes;
+			GroupContext.PreviousLeaderTimelineState = GroupState.PreviousLeaderTimelineState;
+			GroupContext.PreviousLeaderSyncPosition = GroupState.PreviousLeaderSyncPosition;
+			GroupContext.LeaderTimelineState.Reset();
+			GroupContext.LeaderIndex = INDEX_NONE;
+			GroupContext.LeaderDeltaTime = 0.0f;
+			GroupContext.bWasActive = GroupState.bWasActive;
+			GroupContext.bCanGroupUseMarkerSyncing = false;
+			GroupContext.ValidMarkers.Reset();
+			GroupContext.CandidateMarkers.Reset();
+			GroupContext.GroupSeenMarkers.Reset();
+			GroupContext.MemberSeenMarkers.Reset();
+		}
+
+		static void FindLeaderIndex(FSyncGroupContext& GroupContext)
+		{
+			// Find our leader by looking at the total weight and the group role
+			const int32 NumMembers = GroupContext.Members.Num();
+			check(NumMembers != 0);	// Groups should never be empty
+
 			int32 LeaderIndex = INDEX_NONE;
 			float LeaderTotalWeight = -1.0f;
 
 			for (int32 MemberIndex = 0; MemberIndex < NumMembers; ++MemberIndex)
 			{
-				const FSyncGroupMember& GroupMember = GroupState.Members[MemberIndex];
+				const FSyncGroupMemberContext& MemberContext = GroupContext.Members[MemberIndex];
+				const FSyncGroupMember& GroupMember = *MemberContext.State;
 
-				switch (GroupMember.GroupRole)
+				switch (GroupMember.GroupParameters.GroupRole)
 				{
 				case EAnimGroupRole::CanBeLeader:
 				case EAnimGroupRole::TransitionLeader:
@@ -75,19 +167,771 @@ namespace UE::AnimNext
 				LeaderIndex = 0;
 			}
 
-			// Now that we found our leader, advance it by the delta time
-			FTimelineProgress LeaderProgress;
+			GroupContext.LeaderIndex = LeaderIndex;
+		}
+
+		static void CheckSyncMarkersSorted(const FTimelineSyncMarkerArray& SyncMarkers)
+		{
+#if DO_CHECK
+			float PreviousMarkerPosition = -1.0f;
+			for (const FTimelineSyncMarker& SyncMarker : SyncMarkers)
 			{
-				const FSyncGroupMember& GroupLeader = GroupState.Members[LeaderIndex];
+				const float MarkerPosition = SyncMarker.GetPosition();
 
+				checkf(MarkerPosition >= 0.0, TEXT("Sync Markers should have a positive finite position"));
+				checkf(MarkerPosition > PreviousMarkerPosition, TEXT("Sync Markers should be sorted ascending by their position"));
+
+				PreviousMarkerPosition = MarkerPosition;
+			}
+#endif
+		}
+
+		static void CollectValidMarkers(FExecutionContext& Context, FSyncGroupContext& GroupContext)
+		{
+			const int32 NumMembers = GroupContext.Members.Num();
+			if (NumMembers == 1)
+			{
+				return;	// If the leader is alone, we have no need for syncing
+			}
+
+			const int32 LeaderIndex = GroupContext.LeaderIndex;
+
+			FSyncGroupMemberContext& GroupLeaderContext = GroupContext.Members[LeaderIndex];
+			const FSyncGroupMember& GroupLeader = *GroupLeaderContext.State;
+
+			FTraitStackBinding& LeaderTraitStack = GroupLeaderContext.TraitStack;
+
+			TTraitBinding<ITimeline> TimelineTrait;
+
+			bool bCanGroupUseMarkerSyncing = false;
+
+			// Get the group sync marker names from the leader
+			{
 				Context.BindTo(GroupLeader.TraitPtr);
-				ensure(Context.GetStack(GroupLeader.TraitPtr, TraitStack));
-				ensure(TraitStack.GetInterface(GroupSyncTrait));
+				ensure(Context.GetStack(GroupLeader.TraitPtr, LeaderTraitStack));
+
+				FTimelineSyncMarkerArray& SyncMarkers = GroupLeaderContext.SyncMarkers;
+
+				SyncMarkers.Reset();
+				if (LeaderTraitStack.GetInterface(TimelineTrait))
+				{
+					TimelineTrait.GetSyncMarkers(Context, SyncMarkers);
+					CheckSyncMarkersSorted(SyncMarkers);
+
+					bCanGroupUseMarkerSyncing = !SyncMarkers.IsEmpty();
+					GroupLeaderContext.bUseMarkerSyncing = bCanGroupUseMarkerSyncing;
+
+					GroupLeaderContext.TimelineState = TimelineTrait.GetState(Context);
+				}
+			}
+
+			// If our leader has sync markers, iterate over every member of the group and remove markers
+			// that they do not share
+			if (bCanGroupUseMarkerSyncing)
+			{
+				TSet<FName>& ValidMarkers = GroupContext.ValidMarkers;
+				TSet<FName>& CandidateMarkers = GroupContext.CandidateMarkers;
+				TBitArray<TMemStackAllocator<>>& GroupSeenMarkers = GroupContext.GroupSeenMarkers;
+				TBitArray<TMemStackAllocator<>>& MemberSeenMarkers = GroupContext.MemberSeenMarkers;
+
+				CandidateMarkers.Reset();
+				CandidateMarkers.Reserve(GroupLeaderContext.SyncMarkers.Num());
+				for (const FTimelineSyncMarker& SyncMarker : GroupLeaderContext.SyncMarkers)
+				{
+					CandidateMarkers.Add(SyncMarker.GetName());
+				}
+
+				// The group of seen markers starts with the list from the leader
+				// We'll combine this bitset for every follower and only the markers in common to each of them will be used
+				const int32 NumGroupSyncMarkerBits = CandidateMarkers.GetMaxIndex();
+				GroupSeenMarkers.Reset();
+				GroupSeenMarkers.Add(false, NumGroupSyncMarkerBits);
+
+				// Initialize seen markers using the leader
+				for (auto It = CandidateMarkers.CreateConstIterator(); It; ++It)
+				{
+					const FSetElementId SyncMarkerId = It.GetId();
+					check(SyncMarkerId.IsValidId());
+
+					GroupSeenMarkers[SyncMarkerId.AsInteger()] = true;
+				}
+
+				for (int32 MemberIndex = 0; MemberIndex < NumMembers; ++MemberIndex)
+				{
+					if (MemberIndex == LeaderIndex)
+					{
+						continue;	// Ignore the leader
+					}
+
+					FSyncGroupMemberContext& GroupMemberContext = GroupContext.Members[MemberIndex];
+					const FSyncGroupMember& GroupMember = *GroupMemberContext.State;
+
+					FTraitStackBinding& MemberTraitStack = GroupMemberContext.TraitStack;
+
+					Context.BindTo(GroupMember.TraitPtr);
+					ensure(Context.GetStack(GroupMember.TraitPtr, MemberTraitStack));
+
+					FTimelineSyncMarkerArray& SyncMarkers = GroupMemberContext.SyncMarkers;
+					SyncMarkers.Reset();
+					if (MemberTraitStack.GetInterface(TimelineTrait))
+					{
+						TimelineTrait.GetSyncMarkers(Context, SyncMarkers);
+						CheckSyncMarkersSorted(SyncMarkers);
+
+						GroupMemberContext.TimelineState = TimelineTrait.GetState(Context);
+					}
+
+					// Members without markers can still use time based syncing
+					const bool bCanMemberUseMarkerSyncing = !SyncMarkers.IsEmpty();
+					GroupMemberContext.bUseMarkerSyncing = bCanMemberUseMarkerSyncing;
+
+					if (bCanMemberUseMarkerSyncing)
+					{
+						MemberSeenMarkers.Reset();
+						MemberSeenMarkers.Add(false, NumGroupSyncMarkerBits);
+
+						// Increment the markers this member contains
+						for (const FTimelineSyncMarker& SyncMarker : SyncMarkers)
+						{
+							const FSetElementId SyncMarkerId = CandidateMarkers.FindId(SyncMarker.GetName());
+							if (SyncMarkerId.IsValidId())
+							{
+								MemberSeenMarkers[SyncMarkerId.AsInteger()] = true;
+							}
+						}
+
+						GroupSeenMarkers.CombineWithBitwiseAND(MemberSeenMarkers, EBitwiseOperatorFlags::MaintainSize);
+					}
+				}
+
+				// Build the list of valid markers
+				// A marker is valid if it is present in every member of the group
+				// We build a new set as it is faster than repeatedly removing entries
+				{
+					ValidMarkers.Reset();
+					ValidMarkers.Reserve(GroupSeenMarkers.Num());
+
+					FSetElementId SyncMarkerId = FSetElementId::FromInteger(GroupSeenMarkers.FindFrom(true, 0));
+					while (CandidateMarkers.IsValidId(SyncMarkerId))
+					{
+						ValidMarkers.Add(CandidateMarkers.Get(SyncMarkerId));
+
+						SyncMarkerId = FSetElementId::FromInteger(GroupSeenMarkers.FindFrom(true, SyncMarkerId.AsInteger() + 1));
+					}
+				}
+
+				bCanGroupUseMarkerSyncing = !ValidMarkers.IsEmpty();
+			}
+
+			GroupContext.bCanGroupUseMarkerSyncing = bCanGroupUseMarkerSyncing;
+		}
+
+		// Returns the sync marker index immediately after the specified position
+		static int32 GetNextSyncMarkerAt(const FTimelineSyncMarkerArray& SyncMarkers, float Position)
+		{
+			// Returns the index of the first sync marker greater than the specified position
+			return Algo::UpperBoundBy(SyncMarkers, Position, [](const FTimelineSyncMarker& Marker) { return Marker.GetPosition(); });
+		}
+
+		// Returns the boundary sync markers around the specified timeline position using the supplied predicate to determine marker validity
+		static void GetBoundarySyncMarkersAt(const TFunctionRef<bool(FName MarkerName, bool bIsNextMarker)>& MarkerValidityPredicate, const FTimelineSyncMarkerArray& SyncMarkers, const FTimelineState& TimelineState, FMarkerPair& OutPrevMarker, FMarkerPair& OutNextMarker)
+		{
+			const float CurrentTime = TimelineState.GetPosition();
+			const float Duration = TimelineState.GetDuration();
+
+			// Pick a guess to start our search
+			int32 NextMarkerIndex = GetNextSyncMarkerAt(SyncMarkers, CurrentTime);
+			int32 PrevMarkerIndex = NextMarkerIndex >= 0 ? (NextMarkerIndex - 1) : INDEX_NONE;
+
+			float PrevTimeToMarker = 0.0f;
+			float NextTimeToMarker = 0.0f;
+
+			const bool bIsLooping = TimelineState.IsLooping();
+
+			// Handle previous marker
+			{
+				// Refine our search using the valid markers
+				while (PrevMarkerIndex >= 0)
+				{
+					//if (ValidMarkers.Contains(SyncMarkers[PrevMarkerIndex].GetName()))
+					if (MarkerValidityPredicate(SyncMarkers[PrevMarkerIndex].GetName(), false))
+					{
+						// This marker is valid, use it
+						break;
+					}
+
+					PrevMarkerIndex--;
+				}
+
+				// If we reached the start of the marker list
+				if (PrevMarkerIndex == INDEX_NONE)
+				{
+					if (bIsLooping)
+					{
+						PrevMarkerIndex = SyncMarkers.Num() - 1;
+
+						while (PrevMarkerIndex >= 0)
+						{
+							//if (ValidMarkers.Contains(SyncMarkers[PrevMarkerIndex].GetName()))
+							if (MarkerValidityPredicate(SyncMarkers[PrevMarkerIndex].GetName(), false))
+							{
+								// This marker is valid, use it
+								break;
+							}
+
+							PrevMarkerIndex--;
+						}
+
+						check(PrevMarkerIndex != INDEX_NONE);
+
+						// Marker lives in previous loop iteration, in the 'past'
+						const float MarkerTime = SyncMarkers[PrevMarkerIndex].GetPosition() - Duration;
+						PrevTimeToMarker = MarkerTime - CurrentTime;
+					}
+					else
+					{
+						PrevMarkerIndex = MarkerIndexSpecialValues::AnimationBoundary;
+						PrevTimeToMarker = CurrentTime;
+					}
+				}
+				else
+				{
+					PrevTimeToMarker = SyncMarkers[PrevMarkerIndex].GetPosition() - CurrentTime;
+				}
+			}
+
+			// Handle next marker
+			{
+				// Refine our search using the valid markers
+				while (NextMarkerIndex < SyncMarkers.Num())
+				{
+					//if (ValidMarkers.Contains(SyncMarkers[NextMarkerIndex].GetName()))
+					if (MarkerValidityPredicate(SyncMarkers[NextMarkerIndex].GetName(), true))
+					{
+						// This marker is valid, use it
+						break;
+					}
+
+					NextMarkerIndex++;
+				}
+
+				// If we reached the end of the marker list
+				if (NextMarkerIndex == SyncMarkers.Num())
+				{
+					if (bIsLooping)
+					{
+						NextMarkerIndex = 0;
+
+						while (NextMarkerIndex < SyncMarkers.Num())
+						{
+							//if (ValidMarkers.Contains(SyncMarkers[NextMarkerIndex].GetName()))
+							if (MarkerValidityPredicate(SyncMarkers[NextMarkerIndex].GetName(), true))
+							{
+								// This marker is valid, use it
+								break;
+							}
+
+							NextMarkerIndex++;
+						}
+
+						check(NextMarkerIndex != SyncMarkers.Num());
+
+						// Marker lives in next loop iteration, in the 'future'
+						const float MarkerTime = SyncMarkers[NextMarkerIndex].GetPosition() + Duration;
+						NextTimeToMarker = MarkerTime - CurrentTime;
+					}
+					else
+					{
+						NextMarkerIndex = MarkerIndexSpecialValues::AnimationBoundary;
+						NextTimeToMarker = CurrentTime;
+					}
+				}
+				else
+				{
+					NextTimeToMarker = SyncMarkers[NextMarkerIndex].GetPosition() - CurrentTime;
+				}
+			}
+
+			OutPrevMarker.MarkerIndex = PrevMarkerIndex;
+			OutPrevMarker.TimeToMarker = PrevTimeToMarker;
+			OutNextMarker.MarkerIndex = NextMarkerIndex;
+			OutNextMarker.TimeToMarker = NextTimeToMarker;
+		}
+
+		// Returns the boundary sync markers around the specified timeline position
+		static void GetBoundarySyncMarkersAt(const TSet<FName>& ValidMarkers, const FTimelineSyncMarkerArray& SyncMarkers, const FTimelineState& TimelineState, FMarkerPair& OutPrevMarker, FMarkerPair& OutNextMarker)
+		{
+			const auto MarkerValidityPredicate =
+				[&ValidMarkers](FName MarkerName, bool bIsNextMarker)
+				{
+					return ValidMarkers.Contains(MarkerName);
+				};
+
+			return GetBoundarySyncMarkersAt(MarkerValidityPredicate, SyncMarkers, TimelineState, OutPrevMarker, OutNextMarker);
+		}
+
+		// Returns the boundary sync markers around the specified timeline position
+		static void GetBoundarySyncMarkersAt(const FMarkerSyncAnimPosition& SyncPosition, const FTimelineSyncMarkerArray& SyncMarkers, const FTimelineState& TimelineState, FMarkerPair& OutPrevMarker, FMarkerPair& OutNextMarker)
+		{
+			const auto MarkerValidityPredicate =
+				[&SyncPosition](FName MarkerName, bool bIsNextMarker)
+				{
+					return MarkerName == SyncPosition.PreviousMarkerName || MarkerName == SyncPosition.NextMarkerName;
+				};
+
+			return GetBoundarySyncMarkersAt(MarkerValidityPredicate, SyncMarkers, TimelineState, OutPrevMarker, OutNextMarker);
+		}
+
+		// Returns the boundary sync markers around the specified timeline position matching the supplied marker pair exactly
+		static void GetExactBoundarySyncMarkersAt(const FMarkerSyncAnimPosition& SyncPosition, const FTimelineSyncMarkerArray& SyncMarkers, const FTimelineState& TimelineState, FMarkerPair& OutPrevMarker, FMarkerPair& OutNextMarker)
+		{
+			const auto MarkerValidityPredicate =
+				[&SyncPosition](FName MarkerName, bool bIsNextMarker)
+				{
+					return bIsNextMarker ? (MarkerName == SyncPosition.NextMarkerName) : (MarkerName == SyncPosition.PreviousMarkerName);
+				};
+
+			return GetBoundarySyncMarkersAt(MarkerValidityPredicate, SyncMarkers, TimelineState, OutPrevMarker, OutNextMarker);
+		}
+
+		// Returns the boundary sync markers matching the sync position pair nearest the specified timeline position and the adjusted current time based on the sync position
+		static float GetNearestMatchingBoundarySyncMarkers(const FMarkerSyncAnimPosition& SyncPosition, const FTimelineSyncMarkerArray& SyncMarkers, const FTimelineState& TimelineState, FMarkerPair& OutPrevMarker, FMarkerPair& OutNextMarker)
+		{
+			check(SyncPosition.PreviousMarkerName != NAME_None || SyncPosition.NextMarkerName != NAME_None);
+
+			// If we are not looping or if we have no usable sync position
+			//		- We scan the marker list for the matching marker pair around the current time of the follower
+			//		- Current time is unchanged (no syncing)
+			//		- Done
+			if (!TimelineState.IsLooping())
+			{
+				GetBoundarySyncMarkersAt(SyncPosition, SyncMarkers, TimelineState, OutPrevMarker, OutNextMarker);
+				return TimelineState.GetPosition();
+			}
+
+			// If the sync position's previous marker is the start boundary
+			//		- We keep the previous marker at the boundary
+			//		- We look for the next marker, the first instance of the next marker's name
+			//		- We return the current time based on the sync position (see GetCurrentTimeFromMarkers)
+			//		- Done
+			if (SyncPosition.PreviousMarkerName == NAME_None)
+			{
+				const float PrevMarkerTime = 0.0f;	// Start of sequence
+				float CurrentTime = 0.0f;
+
+				OutPrevMarker.MarkerIndex = MarkerIndexSpecialValues::AnimationBoundary;
+
+				const int32 NumSyncMarkers = SyncMarkers.Num();
+				for (int32 MarkerIndex = 0; MarkerIndex < NumSyncMarkers; ++MarkerIndex)
+				{
+					if (SyncMarkers[MarkerIndex].GetName() == SyncPosition.NextMarkerName)
+					{
+						const float NextMarkerTime = SyncMarkers[MarkerIndex].GetPosition();
+						CurrentTime = FMath::LerpStable(PrevMarkerTime, NextMarkerTime, SyncPosition.PositionBetweenMarkers);
+
+						OutPrevMarker.TimeToMarker = PrevMarkerTime - CurrentTime;
+
+						OutNextMarker.MarkerIndex = MarkerIndex;
+						OutNextMarker.TimeToMarker = NextMarkerTime - CurrentTime;
+						break;
+					}
+				}
+
+				return CurrentTime;
+			}
+
+			// If the sync position's next marker is the end boundary
+			//		- We keep the next marker at the boundary
+			//		- We look for the previous marker, the last instance of the previous marker's name
+			//		- We return the current time based on the sync position (see GetCurrentTimeFromMarkers)
+			//		- Done
+			if (SyncPosition.NextMarkerName == NAME_None)
+			{
+				const float NextMarkerTime = TimelineState.GetDuration();	// End of sequence
+				float CurrentTime = 0.0f;
+
+				OutNextMarker.MarkerIndex = MarkerIndexSpecialValues::AnimationBoundary;
+
+				const int32 NumSyncMarkers = SyncMarkers.Num();
+				for (int32 MarkerIndex = NumSyncMarkers - 1; MarkerIndex >= 0; --MarkerIndex)
+				{
+					if (SyncMarkers[MarkerIndex].GetName() == SyncPosition.NextMarkerName)
+					{
+						const float PrevMarkerTime = SyncMarkers[MarkerIndex].GetPosition();
+						CurrentTime = FMath::LerpStable(PrevMarkerTime, NextMarkerTime, SyncPosition.PositionBetweenMarkers);
+
+						OutNextMarker.TimeToMarker = PrevMarkerTime - CurrentTime;
+
+						OutPrevMarker.MarkerIndex = MarkerIndex;
+						OutPrevMarker.TimeToMarker = NextMarkerTime - CurrentTime;
+						break;
+					}
+				}
+
+				return CurrentTime;
+			}
+
+			// Otherwise
+			//		- Scan the marker pairs that match the previous/next markers from the leader
+			//		- We look for the pair that yields an adjusted current time closest to the follower's current time, O(N^2) search
+			//		- Once we find that pair, we return the current time based on the sync position (see GetCurrentTimeFromMarkers)
+			//		- If no such pair is found, the follower must not contain the pair of markers present on the leader (not possible)
+			//		- Done
+			GetExactBoundarySyncMarkersAt(SyncPosition, SyncMarkers, TimelineState, OutPrevMarker, OutNextMarker);
+
+			float CurrentTime = TimelineState.GetPosition();
+			float PrevMarkerTime = SyncMarkers[OutPrevMarker.MarkerIndex].GetPosition();
+			float NextMarkerTime = SyncMarkers[OutNextMarker.MarkerIndex].GetPosition();
+
+			// Handle looping
+			if (PrevMarkerTime > CurrentTime)
+			{
+				PrevMarkerTime -= TimelineState.GetDuration();
+			}
+			else if (NextMarkerTime < CurrentTime)
+			{
+				NextMarkerTime += TimelineState.GetDuration();
+			}
+
+			CurrentTime = FMath::LerpStable(PrevMarkerTime, NextMarkerTime, SyncPosition.PositionBetweenMarkers);
+			return CurrentTime;
+		}
+
+		// Returns the relative sync position between the two specified markers
+		static FMarkerSyncAnimPosition CalculateSyncPosition(const FTimelineSyncMarkerArray& SyncMarkers, const FTimelineState& TimelineState, const FMarkerPair& PrevMarker, const FMarkerPair& NextMarker)
+		{
+			const float Duration = TimelineState.GetDuration();
+			float CurrentTime = TimelineState.GetPosition();
+
+			FMarkerSyncAnimPosition SyncPosition;
+			float PrevTime = 0.0f;
+			float NextTime = Duration;
+
+			// Get previous marker's time and name.
+			if (PrevMarker.MarkerIndex != MarkerIndexSpecialValues::AnimationBoundary && ensureAlwaysMsgf(SyncMarkers.IsValidIndex(PrevMarker.MarkerIndex),
+				TEXT("MarkerCount: %d, PrevMarker : %d, NextMarker: %d, CurrentTime : %0.2f"), SyncMarkers.Num(), PrevMarker.MarkerIndex, NextMarker.MarkerIndex, CurrentTime))
+			{
+				PrevTime = SyncMarkers[PrevMarker.MarkerIndex].GetPosition();
+				SyncPosition.PreviousMarkerName = SyncMarkers[PrevMarker.MarkerIndex].GetName();
+			}
+
+			// Get next marker's time and name.
+			if (NextMarker.MarkerIndex != MarkerIndexSpecialValues::AnimationBoundary && ensureAlwaysMsgf(SyncMarkers.IsValidIndex(NextMarker.MarkerIndex),
+				TEXT("MarkerCount: %d, PrevMarker : %d, NextMarker: %d, CurrentTime : %0.2f"), SyncMarkers.Num(), PrevMarker.MarkerIndex, NextMarker.MarkerIndex, CurrentTime))
+			{
+				NextTime = SyncMarkers[NextMarker.MarkerIndex].GetPosition();
+				SyncPosition.NextMarkerName = SyncMarkers[NextMarker.MarkerIndex].GetName();
+			}
+
+			// Account for looping
+			if (PrevTime > NextTime)
+			{
+				PrevTime = PrevTime > CurrentTime ? (PrevTime - Duration) : PrevTime;
+				NextTime = NextTime < CurrentTime ? (NextTime + Duration) : NextTime;
+			}
+			else if (PrevTime > CurrentTime)
+			{
+				CurrentTime += Duration;
+			}
+
+			if (PrevTime == NextTime)
+			{
+				PrevTime -= Duration;
+			}
+
+			check(NextTime > PrevTime);
+
+			// Store the encoded current time position as a ratio between markers
+			SyncPosition.PositionBetweenMarkers = (CurrentTime - PrevTime) / (NextTime - PrevTime);
+			return SyncPosition;
+		}
+
+		static float MatchLeaderMarkerPhase(const FSyncGroupContext& GroupContext, const FSyncGroupMemberContext& GroupLeaderContext, const FSyncGroupMemberContext& GroupMemberContext)
+		{
+			// Find the nearest markers matching our start position
+			FMarkerPair PrevMarker;
+			FMarkerPair NextMarker;
+			GetNearestMatchingBoundarySyncMarkers(GroupContext.SyncPositionStart, GroupMemberContext.SyncMarkers, GroupMemberContext.TimelineState, PrevMarker, NextMarker);
+
+			// Using the leader start/end sync position, skip as many markers as the leader skipped
+			int32 NumPassedMarkers;
+			if (GroupContext.LeaderStartMarker.MarkerIndex > GroupContext.LeaderEndMarker.MarkerIndex)
+			{
+				// We looped, add the trailing markers
+				NumPassedMarkers = GroupLeaderContext.SyncMarkers.Num() - GroupContext.LeaderStartMarker.MarkerIndex + GroupContext.LeaderEndMarker.MarkerIndex;
+			}
+			else
+			{
+				NumPassedMarkers = GroupContext.LeaderEndMarker.MarkerIndex - GroupContext.LeaderStartMarker.MarkerIndex;
+			}
+
+			int32 PrevMarkerIndex = PrevMarker.MarkerIndex + NumPassedMarkers;
+			float PrevMarkerTime = 0.0f;
+
+			{
+				if (PrevMarkerIndex >= GroupMemberContext.SyncMarkers.Num())
+				{
+					if (GroupMemberContext.TimelineState.IsLooping())
+					{
+						PrevMarkerIndex = 0;
+
+						while (PrevMarkerIndex < GroupMemberContext.SyncMarkers.Num())
+						{
+							if (GroupContext.SyncPositionEnd.PreviousMarkerName == GroupMemberContext.SyncMarkers[PrevMarkerIndex].GetName())
+							{
+								// This marker is valid, use it
+								break;
+							}
+
+							PrevMarkerIndex++;
+						}
+
+						check(PrevMarkerIndex != GroupMemberContext.SyncMarkers.Num());
+
+						// Marker lives in next loop iteration, in the 'future'
+						PrevMarkerTime = GroupMemberContext.SyncMarkers[PrevMarkerIndex].GetPosition();
+					}
+					else
+					{
+						PrevMarkerIndex = MarkerIndexSpecialValues::AnimationBoundary;
+						PrevMarkerTime = GroupMemberContext.TimelineState.GetDuration();
+					}
+				}
+				else
+				{
+					PrevMarkerTime = GroupMemberContext.SyncMarkers[PrevMarkerIndex].GetPosition();
+				}
+
+				check(GroupMemberContext.SyncMarkers[PrevMarkerIndex].GetName() == GroupContext.SyncPositionEnd.PreviousMarkerName);
+			}
+
+			int32 NextMarkerIndex = PrevMarkerIndex + 1;
+			float NextMarkerTime = PrevMarkerTime;
+
+			{
+				if (NextMarkerIndex >= GroupMemberContext.SyncMarkers.Num())
+				{
+					if (GroupMemberContext.TimelineState.IsLooping())
+					{
+						NextMarkerIndex = 0;
+
+						while (NextMarkerIndex < GroupMemberContext.SyncMarkers.Num())
+						{
+							if (GroupContext.SyncPositionEnd.NextMarkerName == GroupMemberContext.SyncMarkers[NextMarkerIndex].GetName())
+							{
+								// This marker is valid, use it
+								break;
+							}
+
+							NextMarkerIndex++;
+						}
+
+						check(NextMarkerIndex != GroupMemberContext.SyncMarkers.Num());
+
+						// Marker lives in next loop iteration, in the 'future'
+						NextMarkerTime = GroupMemberContext.SyncMarkers[NextMarkerIndex].GetPosition() + GroupMemberContext.TimelineState.GetDuration();
+					}
+					else
+					{
+						NextMarkerIndex = MarkerIndexSpecialValues::AnimationBoundary;
+						NextMarkerTime = GroupMemberContext.TimelineState.GetDuration();
+					}
+				}
+				else
+				{
+					NextMarkerTime = GroupMemberContext.SyncMarkers[NextMarkerIndex].GetPosition();
+				}
+
+				check(GroupMemberContext.SyncMarkers[NextMarkerIndex].GetName() == GroupContext.SyncPositionEnd.NextMarkerName);
+			}
+
+			const float CurrentTime = FMath::LerpStable(PrevMarkerTime, NextMarkerTime, GroupContext.SyncPositionEnd.PositionBetweenMarkers);
+
+			float DeltaTime = CurrentTime - GroupMemberContext.TimelineState.GetPosition();
+			if (DeltaTime >= GroupMemberContext.TimelineState.GetDuration())
+			{
+				DeltaTime -= GroupMemberContext.TimelineState.GetDuration();
+			}
+
+#if UE_LOG_ENABLE_MARKER_DEBUG
+			UE_LOG(LogAnimMarkerSync, Warning,
+				TEXT("[%p] Follower PreviousTime(%0.2f) CurrentTime(%0.2f) MoveDelta(%0.2f) Looping(%d) [(Prev '%s' @ %d | Next '%s' @ %d | %0.2f)]"),
+				GroupMemberContext.State->TraitPtr.GetNodeInstance(),
+				GroupMemberContext.TimelineState.GetPosition(),
+				CurrentTime,
+				DeltaTime,
+				GroupLeaderContext.TimelineState.IsLooping() ? 1 : 0,
+				*GroupMemberContext.SyncMarkers[PrevMarkerIndex].GetName().ToString(),
+				PrevMarkerIndex,
+				*GroupMemberContext.SyncMarkers[NextMarkerIndex].GetName().ToString(),
+				NextMarkerIndex,
+				GroupContext.SyncPositionEnd.PositionBetweenMarkers);
+#endif
+
+			return DeltaTime;
+		}
+
+		static void AdvanceLeader(FExecutionContext& Context, FSyncGroupContext& GroupContext)
+		{
+			const int32 LeaderIndex = GroupContext.LeaderIndex;
+
+			FSyncGroupMemberContext& GroupLeaderContext = GroupContext.Members[LeaderIndex];
+			const FSyncGroupMember& GroupLeader = *GroupLeaderContext.State;
+
+			FTraitStackBinding& LeaderTraitStack = GroupLeaderContext.TraitStack;
+
+			TTraitBinding<IGroupSynchronization> GroupSyncTrait;
+			TTraitBinding<ITimeline> TimelineTrait;
+
+			Context.BindTo(GroupLeader.TraitPtr);
+			ensure(Context.GetStack(GroupLeader.TraitPtr, LeaderTraitStack));
+			ensure(LeaderTraitStack.GetInterface(GroupSyncTrait));
+			ensure(LeaderTraitStack.GetInterface(TimelineTrait));
+
+			float DeltaTime = GroupLeader.TraitState.GetDeltaTime();
+
+			if (GroupContext.bCanGroupUseMarkerSyncing)
+			{
+				FTimelineState LeaderTimelineState;
+				FMarkerPair PrevMarkerStart;
+				FMarkerPair NextMarkerStart;
+
+				if (GroupContext.bWasActive)
+				{
+					if (GroupLeaderContext.State->GroupParameters.bMatchSyncPoint)
+					{
+						// This is a soft join where we wish to match the current sync group position
+						float CurrentTime = GroupContext.PreviousLeaderTimelineState.GetPositionRatio() * GroupLeaderContext.TimelineState.GetDuration();
 
-				LeaderProgress = GroupSyncTrait.AdvanceBy(Context, GroupLeader.TraitState.GetDeltaTime());
+						LeaderTimelineState = GroupLeaderContext.TimelineState.WithPosition(CurrentTime);
+
+						if (GroupContext.PreviousLeaderSyncPosition.IsValid() &&
+							GroupContext.ValidMarkers.Contains(GroupContext.PreviousLeaderSyncPosition.PreviousMarkerName) &&
+							GroupContext.ValidMarkers.Contains(GroupContext.PreviousLeaderSyncPosition.NextMarkerName))
+						{
+							// Previous sync position markers are valid, we can use them
+							CurrentTime = GetNearestMatchingBoundarySyncMarkers(GroupContext.PreviousLeaderSyncPosition, GroupLeaderContext.SyncMarkers, LeaderTimelineState, PrevMarkerStart, NextMarkerStart);
+
+							LeaderTimelineState = GroupLeaderContext.TimelineState.WithPosition(CurrentTime);
+						}
+						else
+						{
+							// Previous sync position markers are invalid, we can't use them
+							GetBoundarySyncMarkersAt(GroupContext.ValidMarkers, GroupLeaderContext.SyncMarkers, LeaderTimelineState, PrevMarkerStart, NextMarkerStart);
+						}
+					}
+					else
+					{
+						// This is a hard join where the leader dictates its new position
+						float CurrentTime = GroupLeaderContext.TimelineState.GetPosition();
+
+						LeaderTimelineState = GroupLeaderContext.TimelineState.WithPosition(CurrentTime);
+
+						GetBoundarySyncMarkersAt(GroupContext.ValidMarkers, GroupLeaderContext.SyncMarkers, LeaderTimelineState, PrevMarkerStart, NextMarkerStart);
+					}
+				}
+				else
+				{
+					// This is a new sync group, leader dictates its position
+					float CurrentTime = GroupLeaderContext.TimelineState.GetPosition();
+
+					LeaderTimelineState = GroupLeaderContext.TimelineState.WithPosition(CurrentTime);
+
+					GetBoundarySyncMarkersAt(GroupContext.ValidMarkers, GroupLeaderContext.SyncMarkers, LeaderTimelineState, PrevMarkerStart, NextMarkerStart);
+				}
+
+				// Cache the start sync position
+				GroupContext.SyncPositionStart = CalculateSyncPosition(GroupLeaderContext.SyncMarkers, LeaderTimelineState, PrevMarkerStart, NextMarkerStart);
+				GroupContext.LeaderStartMarker = PrevMarkerStart;
+
+				// We advance the leader by its desired delta time
+				GroupSyncTrait.AdvanceBy(Context, DeltaTime);
+
+				// Compute our new timeline state
+				const FTimelineState NewTimelineState = TimelineTrait.GetState(Context);
+
+				GroupContext.LeaderTimelineState = NewTimelineState;
+				GroupContext.LeaderDeltaTime = DeltaTime;
+
+				// Find our new markers
+				FMarkerPair PrevMarkerEnd;
+				FMarkerPair NextMarkerEnd;
+				GetBoundarySyncMarkersAt(GroupContext.ValidMarkers, GroupLeaderContext.SyncMarkers, NewTimelineState, PrevMarkerEnd, NextMarkerEnd);
+
+				// Cache the end sync position
+				GroupContext.SyncPositionEnd = CalculateSyncPosition(GroupLeaderContext.SyncMarkers, NewTimelineState, PrevMarkerEnd, NextMarkerEnd);
+				GroupContext.LeaderEndMarker = PrevMarkerEnd;
+
+#if UE_LOG_ENABLE_MARKER_DEBUG
+				UE_LOG(LogAnimMarkerSync, Warning,
+					TEXT("[%p] Leader CurrentTime(%0.2f, %0.2f%%) MoveDelta(%0.2f) Looping(%d) [(Prev '%s' @ %d | Next '%s' @ %d | %0.2f) -> (Prev '%s' @ %d | Next '%s' @ %d | %0.2f)]"),
+					GroupLeader.TraitPtr.GetNodeInstance(),
+					GroupContext.LeaderProgress.GetPosition(),
+					GroupContext.LeaderProgress.GetPositionRatio(),
+					DeltaTime,
+					LeaderTimelineState.IsLooping() ? 1 : 0,
+					*GroupContext.SyncPositionStart.PreviousMarkerName.ToString(),
+					PrevMarkerStart.MarkerIndex,
+					*GroupContext.SyncPositionStart.NextMarkerName.ToString(),
+					NextMarkerStart.MarkerIndex,
+					GroupContext.SyncPositionStart.PositionBetweenMarkers,
+					*GroupContext.SyncPositionEnd.PreviousMarkerName.ToString(),
+					PrevMarkerEnd.MarkerIndex,
+					*GroupContext.SyncPositionEnd.NextMarkerName.ToString(),
+					NextMarkerEnd.MarkerIndex,
+					GroupContext.SyncPositionEnd.PositionBetweenMarkers);
+#endif
+			}
+			else
+			{
+				if (GroupContext.bWasActive)
+				{
+					if (GroupLeaderContext.State->GroupParameters.bMatchSyncPoint)
+					{
+						// This is a soft join where we wish to match the current sync group position
+						const float CurrentPosition = GroupLeaderContext.TimelineState.GetPosition();
+						const float PreviousPosition = GroupContext.PreviousLeaderTimelineState.GetPositionRatio() * GroupLeaderContext.TimelineState.GetDuration();
+						DeltaTime = (CurrentPosition + DeltaTime) - PreviousPosition;
+					}
+					else
+					{
+						// This is a hard join where the leader dictates its new position
+					}
+				}
+				else
+				{
+					// This is a new sync group, leader dictates its position
+				}
+
+				GroupSyncTrait.AdvanceBy(Context, DeltaTime);
+
+				GroupContext.LeaderTimelineState = TimelineTrait.GetState(Context);
+				GroupContext.LeaderDeltaTime = DeltaTime;
+
+#if UE_LOG_ENABLE_MARKER_DEBUG
+				UE_LOG(LogAnimMarkerSync, Warning,
+					TEXT("[%p] Leader CurrentTime(%0.2f, %0.2f%%) MoveDelta(%0.2f)"),
+					GroupLeader.TraitPtr.GetNodeInstance(),
+					GroupContext.LeaderProgress.GetPosition(),
+					GroupContext.LeaderProgress.GetPositionRatio(),
+					DeltaTime);
+#endif
+			}
+		}
+
+		static void AdvanceFollowers(FExecutionContext& Context, FSyncGroupContext& GroupContext)
+		{
+			const int32 NumMembers = GroupContext.Members.Num();
+			if (NumMembers == 1)
+			{
+				return;	// No followers if we have a single leader in the group
 			}
 
-			const float LeaderProgressRatio = LeaderProgress.GetPositionRatio();
+			const int32 LeaderIndex = GroupContext.LeaderIndex;
+			const float LeaderProgressRatio = GroupContext.LeaderTimelineState.GetPositionRatio();
+			const FSyncGroupMemberContext& GroupLeaderContext = GroupContext.Members[LeaderIndex];
+
+			TTraitBinding<IGroupSynchronization> GroupSyncTrait;
 
 			// Advance every follower to the same progress ratio as the leader
 			for (int32 MemberIndex = 0; MemberIndex < NumMembers; ++MemberIndex)
@@ -97,14 +941,150 @@ namespace UE::AnimNext
 					continue;	// Ignore the leader, it already advanced
 				}
 
-				const FSyncGroupMember& GroupMember = GroupState.Members[MemberIndex];
+				FSyncGroupMemberContext& GroupMemberContext = GroupContext.Members[MemberIndex];
+				const FSyncGroupMember& GroupMember = *GroupMemberContext.State;
+
+				FTraitStackBinding& MemberTraitStack = GroupMemberContext.TraitStack;
 
 				Context.BindTo(GroupMember.TraitPtr);
-				ensure(Context.GetStack(GroupMember.TraitPtr, TraitStack));
-				ensure(TraitStack.GetInterface(GroupSyncTrait));
+				ensure(Context.GetStack(GroupMember.TraitPtr, MemberTraitStack));
+				ensure(MemberTraitStack.GetInterface(GroupSyncTrait));
+
+				float DeltaTime = 0.0f;
+
+				if (GroupMember.GroupParameters.GroupRole == EAnimGroupRole::ExclusiveAlwaysLeader)
+				{
+					// We asked to be a leader and we weren't picked as leader
+					// We don't want to follow anyone so we tick on our own
+					DeltaTime = GroupMember.TraitState.GetDeltaTime();
+				}
+				else if (GroupMemberContext.bUseMarkerSyncing)
+				{
+					DeltaTime = MatchLeaderMarkerPhase(GroupContext, GroupLeaderContext, GroupMemberContext);
+				}
+				else
+				{
+					const float PreviousPosition = GroupMemberContext.TimelineState.GetPosition();
+					float CurrentPosition = LeaderProgressRatio * GroupMemberContext.TimelineState.GetDuration();
 
-				GroupSyncTrait.AdvanceToRatio(Context, LeaderProgressRatio);
+					if (GroupContext.LeaderDeltaTime >= 0.0f)
+					{
+						if (CurrentPosition < PreviousPosition)
+						{
+							// We must have looped around but we still want a positive delta time to match our leader
+							CurrentPosition += GroupMemberContext.TimelineState.GetDuration();
+						}
+					}
+					else
+					{
+						if (CurrentPosition > PreviousPosition)
+						{
+							// We must have looped around but we still want a negative delta time to match our leader
+							CurrentPosition -= GroupMemberContext.TimelineState.GetDuration();
+						}
+					}
+
+					DeltaTime = CurrentPosition - PreviousPosition;
+
+#if UE_LOG_ENABLE_MARKER_DEBUG
+					UE_LOG(LogAnimMarkerSync, Warning,
+						TEXT("[%p] Follower Ratio(%0.2f)"),
+						GroupMember.TraitPtr.GetNodeInstance(),
+						LeaderProgressRatio);
+#endif
+				}
+
+				GroupSyncTrait.AdvanceBy(Context, DeltaTime);
+			}
+		}
+	}
+
+	FSyncGroupGraphInstanceComponent::FSyncGroupGraphInstanceComponent(FAnimNextGraphInstance& InOwnerInstance)
+		: FGraphInstanceComponent(InOwnerInstance)
+	{
+	}
+
+	FSyncGroupGraphInstanceComponent::~FSyncGroupGraphInstanceComponent() = default;
+
+	void FSyncGroupGraphInstanceComponent::RegisterWithGroup(const FSyncGroupParameters& GroupParameters, const FWeakTraitPtr& TraitPtr, const FTraitUpdateState& TraitState)
+	{
+		const int32 GroupIndex = SyncGroupMap.FindOrAdd(GroupParameters.GroupName, SyncGroups.Num());
+		Private::FSyncGroupState& GroupState = GroupIndex == SyncGroups.Num() ? SyncGroups.AddDefaulted_GetRef() : SyncGroups[GroupIndex];
+		GroupState.GroupName = GroupParameters.GroupName;
+		GroupState.Members.Add(Private::FSyncGroupMember{ TraitState, TraitPtr, GroupParameters });
+		MaxNumMembers = FMath::Max<uint32>(MaxNumMembers, GroupState.Members.Num());
+	}
+
+	void FSyncGroupGraphInstanceComponent::PreUpdate(FExecutionContext& Context)
+	{
+		bool bRebuildGroupIndexMap = false;
+
+		const int32 NumPreviousSyncGroups = SyncGroups.Num();
+		for (int32 SyncGroupIndex = NumPreviousSyncGroups - 1; SyncGroupIndex >= 0; --SyncGroupIndex)
+		{
+			Private::FSyncGroupState& GroupState = SyncGroups[SyncGroupIndex];
+
+			const bool bWasActive = GroupState.Members.IsEmpty();
+			if (!GroupState.bWasActive && bWasActive)
+			{
+				// We have been inactive for two consecutive updates, time to purge this group
+				SyncGroups.RemoveAtSwap(SyncGroupIndex, EAllowShrinking::No);
+				bRebuildGroupIndexMap = true;
 			}
+			else
+			{
+				GroupState.bWasActive = bWasActive;
+				GroupState.Members.Reset();
+			}
+		}
+
+		if (bRebuildGroupIndexMap)
+		{
+			// We purged one or more inactive groups, rebuild our map
+			SyncGroupMap.Reset();
+
+			const int32 NumSyncGroups = SyncGroups.Num();
+			for (int32 SyncGroupIndex = 0; SyncGroupIndex < NumSyncGroups; ++SyncGroupIndex)
+			{
+				const Private::FSyncGroupState& GroupState = SyncGroups[SyncGroupIndex];
+				SyncGroupMap.Add(GroupState.GroupName, SyncGroupIndex);
+			}
+		}
+
+		MaxNumMembers = 0;
+	}
+
+	void FSyncGroupGraphInstanceComponent::PostUpdate(FExecutionContext& Context)
+	{
+		using namespace Private;
+
+		TArray<FSyncGroupMemberContext, TMemStackAllocator<>> MemberContexts;
+		MemberContexts.AddDefaulted(MaxNumMembers);
+
+		FSyncGroupContext GroupContext;
+
+		// Now that we have discovered all groups and their memberships, we can perform synchronization
+		for (FSyncGroupState& GroupState : SyncGroups)
+		{
+			// Initialize our context for this new group
+			InitGroup(GroupState, MemberContexts, GroupContext);
+
+			// Find our leader for this group
+			FindLeaderIndex(GroupContext);
+
+			// If the leader has sync markers, collect and filter them
+			// We'll retain only markers common to all members
+			CollectValidMarkers(Context, GroupContext);
+
+			// Advance the leader position as it determines the position of its followers
+			AdvanceLeader(Context, GroupContext);
+
+			// Advance the follower positions based on their leader
+			AdvanceFollowers(Context, GroupContext);
+
+			// Retain the leader progress as we might need it if a new leader joins the group during the next update
+			GroupState.PreviousLeaderTimelineState = GroupContext.LeaderTimelineState;
+			GroupState.PreviousLeaderSyncPosition = GroupContext.SyncPositionEnd;
 		}
 	}
 }
diff --git a/Engine/Plugins/Experimental/AnimNextAnimGraph/Source/AnimNextAnimGraph/Private/Graph/SyncGroup_GraphInstanceComponent.h b/Engine/Plugins/Experimental/AnimNextAnimGraph/Source/AnimNextAnimGraph/Private/Graph/SyncGroup_GraphInstanceComponent.h
--- a/Engine/Plugins/Experimental/AnimNextAnimGraph/Source/AnimNextAnimGraph/Private/Graph/SyncGroup_GraphInstanceComponent.h
+++ b/Engine/Plugins/Experimental/AnimNextAnimGraph/Source/AnimNextAnimGraph/Private/Graph/SyncGroup_GraphInstanceComponent.h
@@ -3,13 +3,19 @@
 #pragma once
 
 #include "CoreMinimal.h"
-#include "Animation/AnimationAsset.h"
 #include "TraitCore/TraitPtr.h"
 #include "TraitInterfaces/IUpdate.h"
 #include "Graph/GraphInstanceComponent.h"
 
 namespace UE::AnimNext
 {
+	namespace Private
+	{
+		struct FSyncGroupState;
+	}
+
+	struct FSyncGroupParameters;
+
 	/**
 	 * FSyncGroupGraphInstanceComponent
 	 *
@@ -20,26 +26,22 @@ namespace UE::AnimNext
 		DECLARE_ANIM_GRAPH_INSTANCE_COMPONENT(FSyncGroupGraphInstanceComponent)
 
 		explicit FSyncGroupGraphInstanceComponent(FAnimNextGraphInstance& InOwnerInstance);
+		~FSyncGroupGraphInstanceComponent();
 
-		void RegisterWithGroup(FName GroupName, EAnimGroupRole::Type GroupRole, const FWeakTraitPtr& TraitPtr, const FTraitUpdateState& TraitState);
+		void RegisterWithGroup(const FSyncGroupParameters& GroupParameters, const FWeakTraitPtr& TraitPtr, const FTraitUpdateState& TraitState);
 
 		// FGraphInstanceComponent impl
 		virtual void PreUpdate(FExecutionContext& Context) override;
 		virtual void PostUpdate(FExecutionContext& Context) override;
 
 	private:
-		struct FSyncGroupMember
-		{
-			FTraitUpdateState					TraitState;
-			FWeakTraitPtr						TraitPtr;
-			TEnumAsByte<EAnimGroupRole::Type>	GroupRole;
-		};
-
-		struct FSyncGroupState
-		{
-			TArray<FSyncGroupMember> Members;
-		};
-
-		TMap<FName, FSyncGroupState> SyncGroupMap;
+		// A map of sync group name -> group index
+		TMap<FName, int32> SyncGroupMap;
+
+		// A list of groups and their data
+		TArray<Private::FSyncGroupState> SyncGroups;
+
+		// The number of members in the largest group
+		uint32 MaxNumMembers = 0;
 	};
 }
diff --git a/Engine/Plugins/Experimental/AnimNextAnimGraph/Source/AnimNextAnimGraph/Private/PlayAnim/PlayAnimEvents.cpp b/Engine/Plugins/Experimental/AnimNextAnimGraph/Source/AnimNextAnimGraph/Private/PlayAnim/PlayAnimEvents.cpp
--- a/Engine/Plugins/Experimental/AnimNextAnimGraph/Source/AnimNextAnimGraph/Private/PlayAnim/PlayAnimEvents.cpp
+++ b/Engine/Plugins/Experimental/AnimNextAnimGraph/Source/AnimNextAnimGraph/Private/PlayAnim/PlayAnimEvents.cpp
@@ -24,6 +24,6 @@ namespace UE::AnimNext
 
 	void FPlayAnim_TimelineUpdateEvent::Execute() const
 	{
-		Request->OnTimelineUpdate(TimelineProgress);
+		Request->OnTimelineUpdate(TimelineState);
 	}
 }
diff --git a/Engine/Plugins/Experimental/AnimNextAnimGraph/Source/AnimNextAnimGraph/Private/PlayAnim/PlayAnimEvents.h b/Engine/Plugins/Experimental/AnimNextAnimGraph/Source/AnimNextAnimGraph/Private/PlayAnim/PlayAnimEvents.h
--- a/Engine/Plugins/Experimental/AnimNextAnimGraph/Source/AnimNextAnimGraph/Private/PlayAnim/PlayAnimEvents.h
+++ b/Engine/Plugins/Experimental/AnimNextAnimGraph/Source/AnimNextAnimGraph/Private/PlayAnim/PlayAnimEvents.h
@@ -78,7 +78,7 @@ namespace UE::AnimNext
 		// The request to update
 		FPlayAnimRequestPtr Request;
 
-		// The current request timeline progress
-		FTimelineProgress TimelineProgress;
+		// The current request timeline state
+		FTimelineState TimelineState;
 	};
 }
diff --git a/Engine/Plugins/Experimental/AnimNextAnimGraph/Source/AnimNextAnimGraph/Private/PlayAnim/PlayAnimRequest.cpp b/Engine/Plugins/Experimental/AnimNextAnimGraph/Source/AnimNextAnimGraph/Private/PlayAnim/PlayAnimRequest.cpp
--- a/Engine/Plugins/Experimental/AnimNextAnimGraph/Source/AnimNextAnimGraph/Private/PlayAnim/PlayAnimRequest.cpp
+++ b/Engine/Plugins/Experimental/AnimNextAnimGraph/Source/AnimNextAnimGraph/Private/PlayAnim/PlayAnimRequest.cpp
@@ -77,10 +77,10 @@ namespace UE::AnimNext
 		return Status;
 	}
 
-	FTimelineProgress FPlayAnimRequest::GetTimelineProgress() const
+	const FTimelineState& FPlayAnimRequest::GetTimelineState() const
 	{
 		check(IsInGameThread());
-		return TimelineProgress;
+		return TimelineState;
 	}
 
 	bool FPlayAnimRequest::HasExpired() const
@@ -151,10 +151,10 @@ namespace UE::AnimNext
 		}
 	}
 
-	void FPlayAnimRequest::OnTimelineUpdate(FTimelineProgress NewTimelineProgress)
+	void FPlayAnimRequest::OnTimelineUpdate(const FTimelineState& NewTimelineState)
 	{
 		check(IsInGameThread());
-		TimelineProgress = NewTimelineProgress;
+		TimelineState = NewTimelineState;
 	}
 
 	void FPlayAnimRequest::AddReferencedObjects(FReferenceCollector& Collector)
diff --git a/Engine/Plugins/Experimental/AnimNextAnimGraph/Source/AnimNextAnimGraph/Private/PlayAnim/PlayAnimSlotTrait.cpp b/Engine/Plugins/Experimental/AnimNextAnimGraph/Source/AnimNextAnimGraph/Private/PlayAnim/PlayAnimSlotTrait.cpp
--- a/Engine/Plugins/Experimental/AnimNextAnimGraph/Source/AnimNextAnimGraph/Private/PlayAnim/PlayAnimSlotTrait.cpp
+++ b/Engine/Plugins/Experimental/AnimNextAnimGraph/Source/AnimNextAnimGraph/Private/PlayAnim/PlayAnimSlotTrait.cpp
@@ -247,17 +247,33 @@ namespace UE::AnimNext
 			TTraitBinding<ITimeline> ChildTimelineTrait;
 			ensure(ChildStack.GetInterface(ChildTimelineTrait));
 
-			const FTimelineProgress ChildProgress = ChildTimelineTrait.SimulateAdvanceBy(Context, TraitState.GetDeltaTime());
+			const FTimelineState ChildState = ChildTimelineTrait.GetState(Context);
 
 			if (InstanceData->CurrentlyActiveRequestIndex == RequestIndex)
 			{
-				CurrentRequestTimeLeft = ChildProgress.GetTimeLeft();
+				const float ChildCurrentPosition = ChildState.GetPosition();
+				const float ChildNextPosition = ChildCurrentPosition + (TraitState.GetDeltaTime() * ChildState.GetPlayRate());
+
+				// Compute how much time is left before the timeline ends (can be negative if we overshoot)
+				const float TimeToEnd =
+					ChildNextPosition >= ChildCurrentPosition ?	// Is moving forward?
+					ChildState.GetDuration() - ChildNextPosition :
+					ChildNextPosition;
+
+				CurrentRequestTimeLeft = TimeToEnd;
 			}
 
 			{
 				auto TimelineUpdateEvent = MakeTraitEvent<FPlayAnim_TimelineUpdateEvent>();
 				TimelineUpdateEvent->Request = SlotRequest.Request;
-				TimelineUpdateEvent->TimelineProgress = ChildProgress;
+
+				// We don't have too many options here:
+				//    - We can have one frame delay (as we do now)
+				//    - We could use the speculative estimate (from above) as our new state, but this may not be fully accurate (e.g. ignores sync groups)
+				//    - We could query the timeline during PostUpdate, but this would ignore sync groups
+				//    - We could add a new graph instance component and hook PostUpdate, but then we have an ordering issue with the sync group component
+				//    - To be fully accurate, the timeline would need to broadcast when it changes, and so we would need to register callbacks on it and manage them
+				TimelineUpdateEvent->TimelineState = ChildState;
 
 				Context.RaiseOutputTraitEvent(TimelineUpdateEvent);
 			}
diff --git a/Engine/Plugins/Experimental/AnimNextAnimGraph/Source/AnimNextAnimGraph/Private/TraitInterfaces/IGroupSynchronization.cpp b/Engine/Plugins/Experimental/AnimNextAnimGraph/Source/AnimNextAnimGraph/Private/TraitInterfaces/IGroupSynchronization.cpp
--- a/Engine/Plugins/Experimental/AnimNextAnimGraph/Source/AnimNextAnimGraph/Private/TraitInterfaces/IGroupSynchronization.cpp
+++ b/Engine/Plugins/Experimental/AnimNextAnimGraph/Source/AnimNextAnimGraph/Private/TraitInterfaces/IGroupSynchronization.cpp
@@ -21,45 +21,23 @@ namespace UE::AnimNext
 	}
 #endif // WITH_EDITOR
 
-	FName IGroupSynchronization::GetGroupName(FExecutionContext& Context, const TTraitBinding<IGroupSynchronization>& Binding) const
+	FSyncGroupParameters IGroupSynchronization::GetGroupParameters(FExecutionContext& Context, const TTraitBinding<IGroupSynchronization>& Binding) const
 	{
 		TTraitBinding<IGroupSynchronization> SuperBinding;
 		if (Binding.GetStackInterfaceSuper(SuperBinding))
 		{
-			return SuperBinding.GetGroupName(Context);
+			return SuperBinding.GetGroupParameters(Context);
 		}
 
-		return NAME_None;
+		return FSyncGroupParameters();
 	}
 
-	EAnimGroupRole::Type IGroupSynchronization::GetGroupRole(FExecutionContext& Context, const TTraitBinding<IGroupSynchronization>& Binding) const
+	void IGroupSynchronization::AdvanceBy(FExecutionContext& Context, const TTraitBinding<IGroupSynchronization>& Binding, float DeltaTime) const
 	{
 		TTraitBinding<IGroupSynchronization> SuperBinding;
 		if (Binding.GetStackInterfaceSuper(SuperBinding))
 		{
-			return SuperBinding.GetGroupRole(Context);
-		}
-
-		return EAnimGroupRole::CanBeLeader;
-	}
-
-	FTimelineProgress IGroupSynchronization::AdvanceBy(FExecutionContext& Context, const TTraitBinding<IGroupSynchronization>& Binding, float DeltaTime) const
-	{
-		TTraitBinding<IGroupSynchronization> SuperBinding;
-		if (Binding.GetStackInterfaceSuper(SuperBinding))
-		{
-			return SuperBinding.AdvanceBy(Context, DeltaTime);
-		}
-
-		return FTimelineProgress();
-	}
-
-	void IGroupSynchronization::AdvanceToRatio(FExecutionContext& Context, const TTraitBinding<IGroupSynchronization>& Binding, float ProgressRatio) const
-	{
-		TTraitBinding<IGroupSynchronization> SuperBinding;
-		if (Binding.GetStackInterfaceSuper(SuperBinding))
-		{
-			SuperBinding.AdvanceToRatio(Context, ProgressRatio);
+			SuperBinding.AdvanceBy(Context, DeltaTime);
 		}
 	}
 }
diff --git a/Engine/Plugins/Experimental/AnimNextAnimGraph/Source/AnimNextAnimGraph/Private/TraitInterfaces/ITimeline.cpp b/Engine/Plugins/Experimental/AnimNextAnimGraph/Source/AnimNextAnimGraph/Private/TraitInterfaces/ITimeline.cpp
--- a/Engine/Plugins/Experimental/AnimNextAnimGraph/Source/AnimNextAnimGraph/Private/TraitInterfaces/ITimeline.cpp
+++ b/Engine/Plugins/Experimental/AnimNextAnimGraph/Source/AnimNextAnimGraph/Private/TraitInterfaces/ITimeline.cpp
@@ -21,56 +21,32 @@ namespace UE::AnimNext
 	}
 #endif // WITH_EDITOR
 
-	float ITimeline::GetPlayRate(FExecutionContext& Context, const TTraitBinding<ITimeline>& Binding) const
+	void ITimeline::GetSyncMarkers(FExecutionContext& Context, const TTraitBinding<ITimeline>& Binding, FTimelineSyncMarkerArray& OutSyncMarkers) const
 	{
 		TTraitBinding<ITimeline> SuperBinding;
 		if (Binding.GetStackInterfaceSuper(SuperBinding))
 		{
-			return SuperBinding.GetPlayRate(Context);
+			return SuperBinding.GetSyncMarkers(Context, OutSyncMarkers);
 		}
-
-		return 1.0f;
-	}
-
-	FTimelineProgress ITimeline::GetProgress(FExecutionContext& Context, const TTraitBinding<ITimeline>& Binding) const
-	{
-		TTraitBinding<ITimeline> SuperBinding;
-		if (Binding.GetStackInterfaceSuper(SuperBinding))
-		{
-			return SuperBinding.GetProgress(Context);
-		}
-
-		return FTimelineProgress();
-	}
-
-	FTimelineProgress ITimeline::SimulateAdvanceBy(FExecutionContext& Context, const TTraitBinding<ITimeline>& Binding, float DeltaTime) const
-	{
-		TTraitBinding<ITimeline> SuperBinding;
-		if (Binding.GetStackInterfaceSuper(SuperBinding))
-		{
-			return SuperBinding.SimulateAdvanceBy(Context, DeltaTime);
-		}
-
-		return FTimelineProgress();
 	}
 
-	FTimelineProgress ITimeline::AdvanceBy(FExecutionContext& Context, const TTraitBinding<ITimeline>& Binding, float DeltaTime) const
+	FTimelineState ITimeline::GetState(FExecutionContext& Context, const TTraitBinding<ITimeline>& Binding) const
 	{
 		TTraitBinding<ITimeline> SuperBinding;
 		if (Binding.GetStackInterfaceSuper(SuperBinding))
 		{
-			return SuperBinding.AdvanceBy(Context, DeltaTime);
+			return SuperBinding.GetState(Context);
 		}
 
-		return FTimelineProgress();
+		return FTimelineState();
 	}
 
-	void ITimeline::AdvanceToRatio(FExecutionContext& Context, const TTraitBinding<ITimeline>& Binding, float ProgressRatio) const
+	void ITimeline::AdvanceBy(FExecutionContext& Context, const TTraitBinding<ITimeline>& Binding, float DeltaTime) const
 	{
 		TTraitBinding<ITimeline> SuperBinding;
 		if (Binding.GetStackInterfaceSuper(SuperBinding))
 		{
-			SuperBinding.AdvanceToRatio(Context, ProgressRatio);
+			SuperBinding.AdvanceBy(Context, DeltaTime);
 		}
 	}
 }
diff --git a/Engine/Plugins/Experimental/AnimNextAnimGraph/Source/AnimNextAnimGraph/Private/Traits/BlendStackTrait.cpp b/Engine/Plugins/Experimental/AnimNextAnimGraph/Source/AnimNextAnimGraph/Private/Traits/BlendStackTrait.cpp
--- a/Engine/Plugins/Experimental/AnimNextAnimGraph/Source/AnimNextAnimGraph/Private/Traits/BlendStackTrait.cpp
+++ b/Engine/Plugins/Experimental/AnimNextAnimGraph/Source/AnimNextAnimGraph/Private/Traits/BlendStackTrait.cpp
@@ -219,34 +219,13 @@ namespace UE::AnimNext
 		}
 	}
 
-	float FBlendStackCoreTrait::GetPlayRate(FExecutionContext& Context, const TTraitBinding<ITimeline>& Binding) const
+	FTimelineState FBlendStackCoreTrait::GetState(FExecutionContext& Context, const TTraitBinding<ITimeline>& Binding) const
 	{
-		FInstanceData* InstanceData = Binding.GetInstanceData<FInstanceData>();
-		const int32 CurrentlyActiveGraphIndex = InstanceData->CurrentlyActiveGraphIndex;
-		if (CurrentlyActiveGraphIndex != INDEX_NONE)
-		{
-			FBlendStackCoreTrait::FGraphState& ActiveGraph = InstanceData->ChildGraphs[CurrentlyActiveGraphIndex];
-
-			FTraitStackBinding ChildTraitStack;
-			ensure(Context.GetStack(ActiveGraph.ChildPtr, ChildTraitStack));
-
-			UE::AnimNext::TTraitBinding<UE::AnimNext::ITimeline> Timeline;
-			if (ChildTraitStack.GetInterface<UE::AnimNext::ITimeline>(Timeline))
-			{
-				return Timeline.GetPlayRate(Context);
-			}
-		}
-
-		return 1.0f;
-	}
-
-	FTimelineProgress FBlendStackCoreTrait::GetProgress(FExecutionContext& Context, const TTraitBinding<ITimeline>& Binding) const
-	{
-		FInstanceData* InstanceData = Binding.GetInstanceData<FInstanceData>();
+		const FInstanceData* InstanceData = Binding.GetInstanceData<FInstanceData>();
 		const int32 CurrentlyActiveGraphIndex = InstanceData->CurrentlyActiveGraphIndex;
 		if (CurrentlyActiveGraphIndex != INDEX_NONE)
 		{
-			FBlendStackCoreTrait::FGraphState& ActiveGraph = InstanceData->ChildGraphs[CurrentlyActiveGraphIndex];
+			const FBlendStackCoreTrait::FGraphState& ActiveGraph = InstanceData->ChildGraphs[CurrentlyActiveGraphIndex];
 
 			FTraitStackBinding ChildTraitStack;
 			ensure(Context.GetStack(ActiveGraph.ChildPtr, ChildTraitStack));
@@ -254,11 +233,11 @@ namespace UE::AnimNext
 			UE::AnimNext::TTraitBinding<UE::AnimNext::ITimeline> Timeline;
 			if (ChildTraitStack.GetInterface<UE::AnimNext::ITimeline>(Timeline))
 			{
-				return Timeline.GetProgress(Context);
+				return Timeline.GetState(Context);
 			}
 		}
 
-		return FTimelineProgress();
+		return FTimelineState();
 	}
 
 	void FBlendStackCoreTrait::AddReferencedObjects(const FExecutionContext& Context, const TTraitBinding<IGarbageCollection>& Binding, FReferenceCollector& Collector) const
diff --git a/Engine/Plugins/Experimental/AnimNextAnimGraph/Source/AnimNextAnimGraph/Private/Traits/SequencePlayer.cpp b/Engine/Plugins/Experimental/AnimNextAnimGraph/Source/AnimNextAnimGraph/Private/Traits/SequencePlayer.cpp
--- a/Engine/Plugins/Experimental/AnimNextAnimGraph/Source/AnimNextAnimGraph/Private/Traits/SequencePlayer.cpp
+++ b/Engine/Plugins/Experimental/AnimNextAnimGraph/Source/AnimNextAnimGraph/Private/Traits/SequencePlayer.cpp
@@ -52,85 +52,59 @@ namespace UE::AnimNext
 		Context.AppendTask(Task);
 	}
 
-	float FSequencePlayerTrait::GetPlayRate(FExecutionContext& Context, const TTraitBinding<ITimeline>& Binding) const
-	{
-		const FSharedData* SharedData = Binding.GetSharedData<FSharedData>();
-		return SharedData->GetPlayRate(Binding);
-	}
-
-	FTimelineProgress FSequencePlayerTrait::GetProgress(FExecutionContext& Context, const TTraitBinding<ITimeline>& Binding) const
+	void FSequencePlayerTrait::GetSyncMarkers(FExecutionContext& Context, const TTraitBinding<ITimeline>& Binding, FTimelineSyncMarkerArray& OutSyncMarkers) const
 	{
 		const FInstanceData* InstanceData = Binding.GetInstanceData<FInstanceData>();
 		if (const UAnimSequence* AnimSeq = InstanceData->AnimSequence.Get())
 		{
-			return FTimelineProgress(InstanceData->InternalTimeAccumulator, AnimSeq->GetPlayLength());
-		}
+			OutSyncMarkers.Reserve(AnimSeq->AuthoredSyncMarkers.Num());
 
-		return FTimelineProgress();
+			for (const FAnimSyncMarker& SyncMarker : AnimSeq->AuthoredSyncMarkers)
+			{
+				OutSyncMarkers.Add(FTimelineSyncMarker(SyncMarker.MarkerName, SyncMarker.Time));
+			}
+		}
 	}
 
-	FTimelineProgress FSequencePlayerTrait::SimulateAdvanceBy(FExecutionContext& Context, const TTraitBinding<ITimeline>& Binding, float DeltaTime) const
+	FTimelineState FSequencePlayerTrait::GetState(FExecutionContext& Context, const TTraitBinding<ITimeline>& Binding) const
 	{
 		const FInstanceData* InstanceData = Binding.GetInstanceData<FInstanceData>();
-
-		if (UAnimSequence* AnimSeq = InstanceData->AnimSequence.Get())
+		if (const UAnimSequence* AnimSeq = InstanceData->AnimSequence.Get())
 		{
 			const FSharedData* SharedData = Binding.GetSharedData<FSharedData>();
 
-			TTraitBinding<ITimeline> TimelineTrait;
-			Binding.GetStackInterface(TimelineTrait);
-
-			const float PlayRate = TimelineTrait.GetPlayRate(Context);
-			const bool bIsLooping = SharedData->GetbLoop(Binding);
-			const float SequenceLength = AnimSeq->GetPlayLength();
-
-			float Position = InstanceData->InternalTimeAccumulator;
-			FAnimationRuntime::AdvanceTime(bIsLooping, DeltaTime * PlayRate, Position, SequenceLength);
-
-			return FTimelineProgress(Position, SequenceLength);
+			return FTimelineState(
+				InstanceData->InternalTimeAccumulator,
+				AnimSeq->GetPlayLength(),
+				SharedData->GetPlayRate(Binding),
+				SharedData->GetbLoop(Binding));
 		}
 
-		return FTimelineProgress();
+		return FTimelineState();
 	}
 
-	FTimelineProgress FSequencePlayerTrait::AdvanceBy(FExecutionContext& Context, const TTraitBinding<ITimeline>& Binding, float DeltaTime) const
+	void FSequencePlayerTrait::AdvanceBy(FExecutionContext& Context, const TTraitBinding<ITimeline>& Binding, float DeltaTime) const
 	{
 		FInstanceData* InstanceData = Binding.GetInstanceData<FInstanceData>();
 
-		if (UAnimSequence* AnimSeq = InstanceData->AnimSequence.Get())
+		// We only advance if we have a valid anim sequence
+		if (const UAnimSequence* AnimSeq = InstanceData->AnimSequence.Get())
 		{
-			const FSharedData* SharedData = Binding.GetSharedData<FSharedData>();
-
 			TTraitBinding<ITimeline> TimelineTrait;
 			Binding.GetStackInterface(TimelineTrait);
 
-			const float PlayRate = TimelineTrait.GetPlayRate(Context);
-			const bool bIsLooping = SharedData->GetbLoop(Binding);
-			const float SequenceLength = AnimSeq->GetPlayLength();
+			const FTimelineState State = TimelineTrait.GetState(Context);
+
+			const float PlayRate = State.GetPlayRate();
+			const bool bIsLooping = State.IsLooping();
+			const float SequenceLength = State.GetDuration();
 
 			const float PreviousTime = InstanceData->InternalTimeAccumulator;
 			const float MoveDelta = DeltaTime * PlayRate;
-			InstanceData->DeltaTimeRecord.Set(PreviousTime, MoveDelta);
-			FAnimationRuntime::AdvanceTime(bIsLooping, MoveDelta, InstanceData->InternalTimeAccumulator, SequenceLength);
-
-			return FTimelineProgress(InstanceData->InternalTimeAccumulator, SequenceLength);
-		}
-
-		return FTimelineProgress();
-	}
-
-	void FSequencePlayerTrait::AdvanceToRatio(FExecutionContext& Context, const TTraitBinding<ITimeline>& Binding, float ProgressRatio) const
-	{
-		FInstanceData* InstanceData = Binding.GetInstanceData<FInstanceData>();
 
-		if (UAnimSequence* AnimSeq = InstanceData->AnimSequence.Get())
-		{
-			const float SequenceLength = AnimSeq->GetPlayLength();
-
-			const float PreviousTime = InstanceData->InternalTimeAccumulator;
-			InstanceData->InternalTimeAccumulator = FMath::Clamp(ProgressRatio, 0.0f, 1.0f) * SequenceLength;
-			const float MoveDelta = InstanceData->InternalTimeAccumulator - PreviousTime;
 			InstanceData->DeltaTimeRecord.Set(PreviousTime, MoveDelta);
+
+			FAnimationRuntime::AdvanceTime(bIsLooping, MoveDelta, InstanceData->InternalTimeAccumulator, SequenceLength);
 		}
 	}
 
diff --git a/Engine/Plugins/Experimental/AnimNextAnimGraph/Source/AnimNextAnimGraph/Private/Traits/SequencePlayer.h b/Engine/Plugins/Experimental/AnimNextAnimGraph/Source/AnimNextAnimGraph/Private/Traits/SequencePlayer.h
--- a/Engine/Plugins/Experimental/AnimNextAnimGraph/Source/AnimNextAnimGraph/Private/Traits/SequencePlayer.h
+++ b/Engine/Plugins/Experimental/AnimNextAnimGraph/Source/AnimNextAnimGraph/Private/Traits/SequencePlayer.h
@@ -74,11 +74,9 @@ namespace UE::AnimNext
 		virtual void PreEvaluate(FEvaluateTraversalContext& Context, const TTraitBinding<IEvaluate>& Binding) const override;
 
 		// ITimeline impl
-		virtual float GetPlayRate(FExecutionContext& Context, const TTraitBinding<ITimeline>& Binding) const override;
-		virtual FTimelineProgress GetProgress(FExecutionContext& Context, const TTraitBinding<ITimeline>& Binding) const override;
-		virtual FTimelineProgress SimulateAdvanceBy(FExecutionContext& Context, const TTraitBinding<ITimeline>& Binding, float DeltaTime) const override;
-		virtual FTimelineProgress AdvanceBy(FExecutionContext& Context, const TTraitBinding<ITimeline>& Binding, float DeltaTime) const override;
-		virtual void AdvanceToRatio(FExecutionContext& Context, const TTraitBinding<ITimeline>& Binding, float ProgressRatio) const override;
+		virtual void GetSyncMarkers(FExecutionContext& Context, const TTraitBinding<ITimeline>& Binding, FTimelineSyncMarkerArray& OutSyncMarkers) const override;
+		virtual FTimelineState GetState(FExecutionContext& Context, const TTraitBinding<ITimeline>& Binding) const override;
+		virtual void AdvanceBy(FExecutionContext& Context, const TTraitBinding<ITimeline>& Binding, float DeltaTime) const override;
 
 		// IUpdate impl
 		virtual void OnBecomeRelevant(FUpdateTraversalContext& Context, const TTraitBinding<IUpdate>& Binding, const FTraitUpdateState& TraitState) const override;
diff --git a/Engine/Plugins/Experimental/AnimNextAnimGraph/Source/AnimNextAnimGraph/Private/Traits/SynchronizeUsingGroups.cpp b/Engine/Plugins/Experimental/AnimNextAnimGraph/Source/AnimNextAnimGraph/Private/Traits/SynchronizeUsingGroups.cpp
--- a/Engine/Plugins/Experimental/AnimNextAnimGraph/Source/AnimNextAnimGraph/Private/Traits/SynchronizeUsingGroups.cpp
+++ b/Engine/Plugins/Experimental/AnimNextAnimGraph/Source/AnimNextAnimGraph/Private/Traits/SynchronizeUsingGroups.cpp
@@ -16,6 +16,15 @@ namespace UE::AnimNext
 	GENERATE_ANIM_TRAIT_IMPLEMENTATION(FSynchronizeUsingGroupsTrait, TRAIT_INTERFACE_ENUMERATOR, NULL_ANIM_TRAIT_INTERFACE_ENUMERATOR, NULL_ANIM_TRAIT_EVENT_ENUMERATOR)
 	#undef TRAIT_INTERFACE_ENUMERATOR
 
+	void FSynchronizeUsingGroupsTrait::OnBecomeRelevant(FUpdateTraversalContext& Context, const TTraitBinding<IUpdate>& Binding, const FTraitUpdateState& TraitState) const
+	{
+		IUpdate::OnBecomeRelevant(Context, Binding, TraitState);
+
+		FInstanceData* InstanceData = Binding.GetInstanceData<FInstanceData>();
+
+		InstanceData->bHasReachedFullWeight = false;
+	}
+
 	void FSynchronizeUsingGroupsTrait::PreUpdate(FUpdateTraversalContext& Context, const TTraitBinding<IUpdate>& Binding, const FTraitUpdateState& TraitState) const
 	{
 		FInstanceData* InstanceData = Binding.GetInstanceData<FInstanceData>();
@@ -23,12 +32,13 @@ namespace UE::AnimNext
 		TTraitBinding<IGroupSynchronization> GroupSyncTrait;
 		Binding.GetStackInterface(GroupSyncTrait);
 
-		const FName GroupName = GroupSyncTrait.GetGroupName(Context);
-		const bool bHasGroupName = !GroupName.IsNone();
+		const FSyncGroupParameters GroupParameters = GroupSyncTrait.GetGroupParameters(Context);
+		const bool bHasGroupName = !GroupParameters.GroupName.IsNone();
 
 		// If we have a group name, we are active
 		// Freeze the timeline, our sync group will control it
 		InstanceData->bFreezeTimeline = bHasGroupName;
+		InstanceData->bHasReachedFullWeight |= FAnimWeight::IsFullWeight(TraitState.GetTotalWeight());
 
 		// Forward the PreUpdate call, if the timeline attempts to update, we'll do nothing if we are frozen
 		IUpdate::PreUpdate(Context, Binding, TraitState);
@@ -39,43 +49,36 @@ namespace UE::AnimNext
 			return;
 		}
 
-		const EAnimGroupRole::Type GroupRole = GroupSyncTrait.GetGroupRole(Context);
-
 		// Append this trait to our group, we'll need to synchronize it
 		FSyncGroupGraphInstanceComponent& Component = Context.GetComponent<FSyncGroupGraphInstanceComponent>();
-		Component.RegisterWithGroup(GroupName, GroupRole, Binding.GetTraitPtr(), TraitState);
-	}
-
-	FName FSynchronizeUsingGroupsTrait::GetGroupName(FExecutionContext& Context, const TTraitBinding<IGroupSynchronization>& Binding) const
-	{
-		const FSharedData* SharedData = Binding.GetSharedData<FSharedData>();
-		return SharedData->GroupName;
+		Component.RegisterWithGroup(GroupParameters, Binding.GetTraitPtr(), TraitState);
 	}
 
-	EAnimGroupRole::Type FSynchronizeUsingGroupsTrait::GetGroupRole(FExecutionContext& Context, const TTraitBinding<IGroupSynchronization>& Binding) const
+	FSyncGroupParameters FSynchronizeUsingGroupsTrait::GetGroupParameters(FExecutionContext& Context, const TTraitBinding<IGroupSynchronization>& Binding) const
 	{
 		const FSharedData* SharedData = Binding.GetSharedData<FSharedData>();
-		return SharedData->GroupRole;
-	}
-
-	FTimelineProgress FSynchronizeUsingGroupsTrait::AdvanceBy(FExecutionContext& Context, const TTraitBinding<IGroupSynchronization>& Binding, float DeltaTime) const
-	{
-		FInstanceData* InstanceData = Binding.GetInstanceData<FInstanceData>();
 
-		// When the group advances the timeline, we thaw it to advance
-		InstanceData->bFreezeTimeline = false;
-
-		TTraitBinding<ITimeline> TimelineTrait;
-		Binding.GetStackInterface(TimelineTrait);
+		FSyncGroupParameters Parameters;
+		Parameters.GroupName = SharedData->GroupName;
+		Parameters.GroupRole = SharedData->GroupRole;
+		Parameters.bMatchSyncPoint = SharedData->bMatchSyncPoint;
 
-		const FTimelineProgress Progress = TimelineTrait.AdvanceBy(Context, DeltaTime);
+		// When the role is set to TransitionLeader or TransitionFollower, we will not be part of a sync group until we are full weight
+		// meaning we ignore the sync group while blending in, but not while blending out
+		if (SharedData->GroupRole == EAnimGroupRole::TransitionLeader || SharedData->GroupRole == EAnimGroupRole::TransitionFollower)
+		{
+			const FInstanceData* InstanceData = Binding.GetInstanceData<FInstanceData>();
 
-		InstanceData->bFreezeTimeline = true;
+			if (InstanceData->bHasReachedFullWeight)
+			{
+				Parameters.GroupName = NAME_None;
+			}
+		}
 
-		return Progress;
+		return Parameters;
 	}
 
-	void FSynchronizeUsingGroupsTrait::AdvanceToRatio(FExecutionContext& Context, const TTraitBinding<IGroupSynchronization>& Binding, float ProgressRatio) const
+	void FSynchronizeUsingGroupsTrait::AdvanceBy(FExecutionContext& Context, const TTraitBinding<IGroupSynchronization>& Binding, float DeltaTime) const
 	{
 		FInstanceData* InstanceData = Binding.GetInstanceData<FInstanceData>();
 
@@ -85,41 +88,19 @@ namespace UE::AnimNext
 		TTraitBinding<ITimeline> TimelineTrait;
 		Binding.GetStackInterface(TimelineTrait);
 
-		TimelineTrait.AdvanceToRatio(Context, ProgressRatio);
+		TimelineTrait.AdvanceBy(Context, DeltaTime);
 
 		InstanceData->bFreezeTimeline = true;
 	}
 
-	FTimelineProgress FSynchronizeUsingGroupsTrait::SimulateAdvanceBy(FExecutionContext& Context, const TTraitBinding<ITimeline>& Binding, float DeltaTime) const
-	{
-		const FInstanceData* InstanceData = Binding.GetInstanceData<FInstanceData>();
-		if (InstanceData->bFreezeTimeline)
-		{
-			return ITimeline::GetProgress(Context, Binding);	// If the timeline is frozen, we don't advance, return the previous value
-		}
-
-		return ITimeline::SimulateAdvanceBy(Context, Binding, DeltaTime);
-	}
-
-	FTimelineProgress FSynchronizeUsingGroupsTrait::AdvanceBy(FExecutionContext& Context, const TTraitBinding<ITimeline>& Binding, float DeltaTime) const
-	{
-		const FInstanceData* InstanceData = Binding.GetInstanceData<FInstanceData>();
-		if (InstanceData->bFreezeTimeline)
-		{
-			return ITimeline::GetProgress(Context, Binding);	// If the timeline is frozen, we don't advance, return the previous value
-		}
-
-		return ITimeline::AdvanceBy(Context, Binding, DeltaTime);
-	}
-
-	void FSynchronizeUsingGroupsTrait::AdvanceToRatio(FExecutionContext& Context, const TTraitBinding<ITimeline>& Binding, float ProgressRatio) const
+	void FSynchronizeUsingGroupsTrait::AdvanceBy(FExecutionContext& Context, const TTraitBinding<ITimeline>& Binding, float DeltaTime) const
 	{
 		const FInstanceData* InstanceData = Binding.GetInstanceData<FInstanceData>();
 		if (InstanceData->bFreezeTimeline)
 		{
-			return;			// If the timeline is frozen, we don't advance
+			return;	// If the timeline is frozen, we don't advance
 		}
 
-		ITimeline::AdvanceToRatio(Context, Binding, ProgressRatio);
+		ITimeline::AdvanceBy(Context, Binding, DeltaTime);
 	}
 }
diff --git a/Engine/Plugins/Experimental/AnimNextAnimGraph/Source/AnimNextAnimGraph/Private/Traits/SynchronizeUsingGroups.h b/Engine/Plugins/Experimental/AnimNextAnimGraph/Source/AnimNextAnimGraph/Private/Traits/SynchronizeUsingGroups.h
--- a/Engine/Plugins/Experimental/AnimNextAnimGraph/Source/AnimNextAnimGraph/Private/Traits/SynchronizeUsingGroups.h
+++ b/Engine/Plugins/Experimental/AnimNextAnimGraph/Source/AnimNextAnimGraph/Private/Traits/SynchronizeUsingGroups.h
@@ -24,6 +24,11 @@ struct FAnimNextSynchronizeUsingGroupsTraitSharedData : public FAnimNextTraitSha
 	// The role this player can assume within the group
 	UPROPERTY(EditAnywhere, Category = "Default", meta = (Inline))
 	TEnumAsByte<EAnimGroupRole::Type> GroupRole = EAnimGroupRole::CanBeLeader;
+
+	// Whether or not to match the group sync point when joining as leader
+	// When disabled, the leader will snap the position of the group to its current position
+	UPROPERTY(EditAnywhere, Category = "Default", meta = (Inline))
+	bool bMatchSyncPoint = true;
 };
 
 namespace UE::AnimNext
@@ -42,20 +47,18 @@ namespace UE::AnimNext
 		struct FInstanceData : FTraitInstanceData
 		{
 			bool bFreezeTimeline = false;
+			bool bHasReachedFullWeight = false;
 		};
 
 		// IUpdate impl
+		virtual void OnBecomeRelevant(FUpdateTraversalContext& Context, const TTraitBinding<IUpdate>& Binding, const FTraitUpdateState& TraitState) const override;
 		virtual void PreUpdate(FUpdateTraversalContext& Context, const TTraitBinding<IUpdate>& Binding, const FTraitUpdateState& TraitState) const override;
 
 		// IGroupSynchronization impl
-		virtual FName GetGroupName(FExecutionContext& Context, const TTraitBinding<IGroupSynchronization>& Binding) const override;
-		virtual EAnimGroupRole::Type GetGroupRole(FExecutionContext& Context, const TTraitBinding<IGroupSynchronization>& Binding) const override;
-		virtual FTimelineProgress AdvanceBy(FExecutionContext& Context, const TTraitBinding<IGroupSynchronization>& Binding, float DeltaTime) const override;
-		virtual void AdvanceToRatio(FExecutionContext& Context, const TTraitBinding<IGroupSynchronization>& Binding, float ProgressRatio) const override;
+		virtual FSyncGroupParameters GetGroupParameters(FExecutionContext& Context, const TTraitBinding<IGroupSynchronization>& Binding) const override;
+		virtual void AdvanceBy(FExecutionContext& Context, const TTraitBinding<IGroupSynchronization>& Binding, float DeltaTime) const override;
 
 		// ITimeline impl
-		virtual FTimelineProgress SimulateAdvanceBy(FExecutionContext& Context, const TTraitBinding<ITimeline>& Binding, float DeltaTime) const override;
-		virtual FTimelineProgress AdvanceBy(FExecutionContext& Context, const TTraitBinding<ITimeline>& Binding, float DeltaTime) const override;
-		virtual void AdvanceToRatio(FExecutionContext& Context, const TTraitBinding<ITimeline>& Binding, float ProgressRatio) const override;
+		virtual void AdvanceBy(FExecutionContext& Context, const TTraitBinding<ITimeline>& Binding, float DeltaTime) const override;
 	};
 }
diff --git a/Engine/Plugins/Experimental/AnimNextAnimGraph/Source/AnimNextAnimGraph/Public/PlayAnim/PlayAnimRequest.h b/Engine/Plugins/Experimental/AnimNextAnimGraph/Source/AnimNextAnimGraph/Public/PlayAnim/PlayAnimRequest.h
--- a/Engine/Plugins/Experimental/AnimNextAnimGraph/Source/AnimNextAnimGraph/Public/PlayAnim/PlayAnimRequest.h
+++ b/Engine/Plugins/Experimental/AnimNextAnimGraph/Source/AnimNextAnimGraph/Public/PlayAnim/PlayAnimRequest.h
@@ -144,8 +144,8 @@ namespace UE::AnimNext
 		// Returns the request status
 		EPlayAnimStatus GetStatus() const;
 
-		// Returns the current timeline progress
-		FTimelineProgress GetTimelineProgress() const;
+		// Returns the current timeline state
+		const FTimelineState& GetTimelineState() const;
 
 		// Returns whether or not this request has expired
 		bool HasExpired() const;
@@ -167,7 +167,7 @@ namespace UE::AnimNext
 
 	private:
 		void OnStatusUpdate(EPlayAnimStatus NewStatus);
-		void OnTimelineUpdate(FTimelineProgress NewTimelineProgress);
+		void OnTimelineUpdate(const FTimelineState& NewTimelineState);
 
 		// The request arguments
 		FPlayAnimRequestArgs RequestArgs;
@@ -181,8 +181,8 @@ namespace UE::AnimNext
 		// The current request status
 		EPlayAnimStatus Status = EPlayAnimStatus::None;
 
-		// The current timeline progress
-		FTimelineProgress TimelineProgress;
+		// The current timeline state
+		FTimelineState TimelineState;
 
 		friend FPlayAnim_PlayEvent;
 		friend FPlayAnim_StatusUpdateEvent;
diff --git a/Engine/Plugins/Experimental/AnimNextAnimGraph/Source/AnimNextAnimGraph/Public/TraitInterfaces/IGroupSynchronization.h b/Engine/Plugins/Experimental/AnimNextAnimGraph/Source/AnimNextAnimGraph/Public/TraitInterfaces/IGroupSynchronization.h
--- a/Engine/Plugins/Experimental/AnimNextAnimGraph/Source/AnimNextAnimGraph/Public/TraitInterfaces/IGroupSynchronization.h
+++ b/Engine/Plugins/Experimental/AnimNextAnimGraph/Source/AnimNextAnimGraph/Public/TraitInterfaces/IGroupSynchronization.h
@@ -8,8 +8,24 @@
 #include "TraitCore/TraitBinding.h"
 #include "TraitInterfaces/ITimeline.h"
 
+class FMarkerTickContext;
+
 namespace UE::AnimNext
 {
+	// Encapsulates the parameters that control sync group behavior
+	struct FSyncGroupParameters
+	{
+		// The group name to belong to (if None, syncing is disabled)
+		FName GroupName;
+
+		// The role within the group
+		TEnumAsByte<EAnimGroupRole::Type> GroupRole = EAnimGroupRole::CanBeLeader;
+
+		// Whether or not to match the group sync point when joining as leader
+		// When disabled, the leader will snap the position of the group to its current position
+		bool bMatchSyncPoint = true;
+	};
+
 	/**
 	 * IGroupSynchronization
 	 *
@@ -19,20 +35,12 @@ namespace UE::AnimNext
 	{
 		DECLARE_ANIM_TRAIT_INTERFACE(IGroupSynchronization, 0xf607d0fd)
 
-		// Returns the group name used for synchronization
-		virtual FName GetGroupName(FExecutionContext& Context, const TTraitBinding<IGroupSynchronization>& Binding) const;
-
-		// Returns the group role used for synchronization
-		virtual EAnimGroupRole::Type GetGroupRole(FExecutionContext& Context, const TTraitBinding<IGroupSynchronization>& Binding) const;
+		// Returns the group parameters used for synchronization
+		virtual FSyncGroupParameters GetGroupParameters(FExecutionContext& Context, const TTraitBinding<IGroupSynchronization>& Binding) const;
 
-		// Called by the sync group graph instance component once a group has been synchronized to advance time on the leader
-		// Returns the progress ratio of playback: 0.0 = start of animation, 1.0 = end of animation
-		virtual FTimelineProgress AdvanceBy(FExecutionContext& Context, const TTraitBinding<IGroupSynchronization>& Binding, float DeltaTime) const;
+		// Called by the sync group graph instance component to advance time
+		virtual void AdvanceBy(FExecutionContext& Context, const TTraitBinding<IGroupSynchronization>& Binding, float DeltaTime) const;
 
-		// Called by the sync group graph instance component once a group has been synchronized to advance time on each follower
-		// Progress ratio must be between [0.0, 1.0]
-		virtual void AdvanceToRatio(FExecutionContext& Context, const TTraitBinding<IGroupSynchronization>& Binding, float ProgressRatio) const;
-	
 #if WITH_EDITOR
 		virtual const FText& GetDisplayName() const override;
 		virtual const FText& GetDisplayShortName() const override;
@@ -45,28 +53,16 @@ namespace UE::AnimNext
 	template<>
 	struct TTraitBinding<IGroupSynchronization> : FTraitBinding
 	{
-		// @see IGroupSynchronization::GetGroupName
-		FName GetGroupName(FExecutionContext& Context) const
+		// @see IGroupSynchronization::GetGroupParameters
+		FSyncGroupParameters GetGroupParameters(FExecutionContext& Context) const
 		{
-			return GetInterface()->GetGroupName(Context, *this);
-		}
-
-		// @see IGroupSynchronization::GetGroupRole
-		EAnimGroupRole::Type GetGroupRole(FExecutionContext& Context) const
-		{
-			return GetInterface()->GetGroupRole(Context, *this);
+			return GetInterface()->GetGroupParameters(Context, *this);
 		}
 
 		// @see IGroupSynchronization::AdvanceBy
-		FTimelineProgress AdvanceBy(FExecutionContext& Context, float DeltaTime) const
-		{
-			return GetInterface()->AdvanceBy(Context, *this, DeltaTime);
-		}
-
-		// @see IGroupSynchronization::AdvanceToRatio
-		void AdvanceToRatio(FExecutionContext& Context, float ProgressRatio) const
+		void AdvanceBy(FExecutionContext& Context, float DeltaTime) const
 		{
-			GetInterface()->AdvanceToRatio(Context, *this, ProgressRatio);
+			GetInterface()->AdvanceBy(Context, *this, DeltaTime);
 		}
 
 	protected:
diff --git a/Engine/Plugins/Experimental/AnimNextAnimGraph/Source/AnimNextAnimGraph/Public/TraitInterfaces/ITimeline.h b/Engine/Plugins/Experimental/AnimNextAnimGraph/Source/AnimNextAnimGraph/Public/TraitInterfaces/ITimeline.h
--- a/Engine/Plugins/Experimental/AnimNextAnimGraph/Source/AnimNextAnimGraph/Public/TraitInterfaces/ITimeline.h
+++ b/Engine/Plugins/Experimental/AnimNextAnimGraph/Source/AnimNextAnimGraph/Public/TraitInterfaces/ITimeline.h
@@ -3,37 +3,70 @@
 #pragma once
 
 #include "CoreMinimal.h"
+#include "Misc/MemStack.h"
 #include "TraitCore/ITraitInterface.h"
 #include "TraitCore/TraitBinding.h"
 
 namespace UE::AnimNext
 {
+	struct FTimelineSyncMarker
+	{
+		// Default construct empty
+		FTimelineSyncMarker() = default;
+
+		// Construct with a specific sync marker name and position
+		FTimelineSyncMarker(FName InName, float InPosition)
+			: Name(InName)
+			, Position(InPosition)
+		{
+		}
+
+		// Returns the sync marker name
+		FName GetName() const { return Name; }
+
+		// Returns the sync marker position (in seconds) on the timeline between [0.0, TimelineDuration]
+		float GetPosition() const { return Position; }
+
+	private:
+		// The sync marker name
+		FName Name;
+
+		// The sync marker position
+		float Position = 0.0f;
+	};
+
+	// An array of sync markers
+	// We reserve a small amount inline and spill on the memstack
+	using FTimelineSyncMarkerArray = TArray<FTimelineSyncMarker, TInlineAllocator<8, TMemStackAllocator<>>>;
+
 	/**
-	 * Timeline Progress
-	 * 
-	 * Encapsulates the progress along a timeline.
+	 * Timeline State
+	 *
+	 * Encapsulates the state of a timeline.
 	 */
-	struct FTimelineProgress
+	struct FTimelineState
 	{
-		// Default construct with no progress
-		FTimelineProgress() = default;
+		// Default construct with no state
+		FTimelineState() = default;
 
-		// Construct with a specific position and duration
-		FTimelineProgress(float InPosition, float InDuration)
+		// Construct with a specific state
+		FTimelineState(float InPosition, float InDuration, float InPlayRate, bool bInIsLooping)
 			: Position(InPosition)
 			, Duration(InDuration)
+			, PlayRate(InPlayRate)
+			, bIsLooping(bInIsLooping)
 		{
 		}
 
-		// Resets the timeline progress to its initial state
+		// Resets the timeline state to its default state
 		void Reset()
 		{
-			Position = Duration = 0.0f;
+			Position = 0.0f;
+			Duration = 0.0f;
+			PlayRate = 0.0f;
+			bIsLooping = false;
 		}
 
-		// Returns the timeline duration in seconds
-		float GetDuration() const { return Duration; }
-
 		// Returns the timeline position in seconds
 		float GetPosition() const { return Position; }
 
@@ -41,7 +74,48 @@ namespace UE::AnimNext
 		float GetPositionRatio() const { return Duration != 0.0f ? FMath::Clamp(Position / Duration, 0.0f, 1.0f) : 0.0f; }
 
 		// Returns the time left to play in the timeline
-		float GetTimeLeft() const { return Duration - Position; }
+		float GetTimeLeft() const { return PlayRate >= 0.0f ? (Duration - Position) : Position; }
+
+		// Returns the timeline duration in seconds
+		float GetDuration() const { return Duration; }
+
+		// Returns the play rate of the timeline
+		float GetPlayRate() const { return PlayRate; }
+
+		// Returns whether or not the timeline is looping
+		bool IsLooping() const { return bIsLooping; }
+
+		// Creates a new instance of the timeline state with the supplied position
+		[[nodiscard]] FTimelineState WithPosition(float InPosition) const
+		{
+			FTimelineState Result(*this);
+			Result.Position = InPosition;
+			return Result;
+		}
+
+		// Creates a new instance of the timeline state with the supplied duration
+		[[nodiscard]] FTimelineState WithDuration(float InDuration) const
+		{
+			FTimelineState Result(*this);
+			Result.Duration = InDuration;
+			return Result;
+		}
+
+		// Creates a new instance of the timeline state with the supplied play rate
+		[[nodiscard]] FTimelineState WithPlayRate(float InPlayRate) const
+		{
+			FTimelineState Result(*this);
+			Result.PlayRate = InPlayRate;
+			return Result;
+		}
+
+		// Creates a new instance of the timeline state with the supplied looping flag
+		[[nodiscard]] FTimelineState AsLooping(bool InAsLooping = true) const
+		{
+			FTimelineState Result(*this);
+			Result.bIsLooping = InAsLooping;
+			return Result;
+		}
 
 	private:
 		// Timeline position in seconds
@@ -49,6 +123,12 @@ namespace UE::AnimNext
 
 		// Timeline duration in seconds
 		float Duration = 0.0f;
+
+		// Timeline play rate
+		float PlayRate = 0.0f;
+
+		// Timeline looping status
+		bool bIsLooping = false;
 	};
 
 	/**
@@ -60,23 +140,14 @@ namespace UE::AnimNext
 	{
 		DECLARE_ANIM_TRAIT_INTERFACE(ITimeline, 0x53760727)
 
-		// Returns the play rate of this timeline
-		virtual float GetPlayRate(FExecutionContext& Context, const TTraitBinding<ITimeline>& Binding) const;
+		// Returns the sync markers of this timeline sorted by time
+		virtual void GetSyncMarkers(FExecutionContext& Context, const TTraitBinding<ITimeline>& Binding, FTimelineSyncMarkerArray& OutSyncMarkers) const;
 
-		// Returns the progress of this timeline
-		virtual FTimelineProgress GetProgress(FExecutionContext& Context, const TTraitBinding<ITimeline>& Binding) const;
-
-		// Simulates the advance of time by the provided delta time (positive or negative) on this timeline
-		// Returns the progress of playback
-		virtual FTimelineProgress SimulateAdvanceBy(FExecutionContext& Context, const TTraitBinding<ITimeline>& Binding, float DeltaTime) const;
+		// Returns the state of this timeline
+		virtual FTimelineState GetState(FExecutionContext& Context, const TTraitBinding<ITimeline>& Binding) const;
 
 		// Advances time by the provided delta time (positive or negative) on this timeline
-		// Returns the progress of playback
-		virtual FTimelineProgress AdvanceBy(FExecutionContext& Context, const TTraitBinding<ITimeline>& Binding, float DeltaTime) const;
-
-		// Advances time to the specified progress ratio on this timeline
-		// Progress ratio must be between [0.0, 1.0]
-		virtual void AdvanceToRatio(FExecutionContext& Context, const TTraitBinding<ITimeline>& Binding, float ProgressRatio) const;
+		virtual void AdvanceBy(FExecutionContext& Context, const TTraitBinding<ITimeline>& Binding, float DeltaTime) const;
 
 #if WITH_EDITOR
 		virtual const FText& GetDisplayName() const override;
@@ -90,34 +161,22 @@ namespace UE::AnimNext
 	template<>
 	struct TTraitBinding<ITimeline> : FTraitBinding
 	{
-		// @see ITimeline::GetPlayRate
-		float GetPlayRate(FExecutionContext& Context) const
+		// @see ITimeline::GetSyncMarkers
+		void GetSyncMarkers(FExecutionContext& Context, FTimelineSyncMarkerArray& OutSyncMarkers) const
 		{
-			return GetInterface()->GetPlayRate(Context, *this);
+			return GetInterface()->GetSyncMarkers(Context, *this, OutSyncMarkers);
 		}
 
-		// @see ITimeline::GetProgress
-		FTimelineProgress GetProgress(FExecutionContext& Context) const
+		// @see ITimeline::GetState
+		FTimelineState GetState(FExecutionContext& Context) const
 		{
-			return GetInterface()->GetProgress(Context, *this);
-		}
-
-		// @see ITimeline::SimulateAdvanceBy
-		FTimelineProgress SimulateAdvanceBy(FExecutionContext& Context, float DeltaTime) const
-		{
-			return GetInterface()->SimulateAdvanceBy(Context, *this, DeltaTime);
+			return GetInterface()->GetState(Context, *this);
 		}
 
 		// @see ITimeline::AdvanceBy
-		FTimelineProgress AdvanceBy(FExecutionContext& Context, float DeltaTime) const
-		{
-			return GetInterface()->AdvanceBy(Context, *this, DeltaTime);
-		}
-
-		// @see ITimeline::AdvanceToRatio
-		void AdvanceToRatio(FExecutionContext& Context, float ProgressRatio) const
+		void AdvanceBy(FExecutionContext& Context, float DeltaTime) const
 		{
-			GetInterface()->AdvanceToRatio(Context, *this, ProgressRatio);
+			GetInterface()->AdvanceBy(Context, *this, DeltaTime);
 		}
 
 	protected:
diff --git a/Engine/Plugins/Experimental/AnimNextAnimGraph/Source/AnimNextAnimGraph/Public/TraitInterfaces/IUpdate.h b/Engine/Plugins/Experimental/AnimNextAnimGraph/Source/AnimNextAnimGraph/Public/TraitInterfaces/IUpdate.h
--- a/Engine/Plugins/Experimental/AnimNextAnimGraph/Source/AnimNextAnimGraph/Public/TraitInterfaces/IUpdate.h
+++ b/Engine/Plugins/Experimental/AnimNextAnimGraph/Source/AnimNextAnimGraph/Public/TraitInterfaces/IUpdate.h
@@ -56,7 +56,7 @@ namespace UE::AnimNext
 		bool IsNewlyRelevant() const { return !!bIsNewlyRelevant; }
 
 		// Creates a new instance of the update state with the total weight scaled by the supplied weight
-		FTraitUpdateState WithWeight(float Weight) const
+		[[nodiscard]] FTraitUpdateState WithWeight(float Weight) const
 		{
 			FTraitUpdateState Result(*this);
 			Result.TotalWeight *= Weight;
@@ -64,7 +64,7 @@ namespace UE::AnimNext
 		}
 
 		// Creates a new instance of the update state with the total trajectory weight scaled by the supplied weight
-		FTraitUpdateState WithTrajectoryWeight(float Weight) const
+		[[nodiscard]] FTraitUpdateState WithTrajectoryWeight(float Weight) const
 		{
 			FTraitUpdateState Result(*this);
 			Result.TotalTrajectoryWeight *= Weight;
@@ -72,7 +72,7 @@ namespace UE::AnimNext
 		}
 
 		// Creates a new instance of the update state with the delta time scaled by the supplied scale factor
-		FTraitUpdateState WithTimeScale(float TimeScale) const
+		[[nodiscard]] FTraitUpdateState WithTimeScale(float TimeScale) const
 		{
 			FTraitUpdateState Result(*this);
 			Result.DeltaTime *= TimeScale;
@@ -80,7 +80,7 @@ namespace UE::AnimNext
 		}
 
 		// Creates a new instance of the update state marking it as blending out
-		FTraitUpdateState AsBlendingOut(bool InIsBlendingOut = true) const
+		[[nodiscard]] FTraitUpdateState AsBlendingOut(bool InIsBlendingOut = true) const
 		{
 			FTraitUpdateState Result(*this);
 			Result.bIsBlendingOut = InIsBlendingOut;
@@ -88,7 +88,7 @@ namespace UE::AnimNext
 		}
 
 		// Creates a new instance of the update state marking it as newly relevant
-		FTraitUpdateState AsNewlyRelevant(bool InIsNewlyRelevant = true) const
+		[[nodiscard]] FTraitUpdateState AsNewlyRelevant(bool InIsNewlyRelevant = true) const
 		{
 			FTraitUpdateState Result(*this);
 			Result.bIsNewlyRelevant = InIsNewlyRelevant;
diff --git a/Engine/Plugins/Experimental/AnimNextPoseSearch/Source/AnimNextPoseSearch/Private/MotionMatchingTrait.cpp b/Engine/Plugins/Experimental/AnimNextPoseSearch/Source/AnimNextPoseSearch/Private/MotionMatchingTrait.cpp
--- a/Engine/Plugins/Experimental/AnimNextPoseSearch/Source/AnimNextPoseSearch/Private/MotionMatchingTrait.cpp
+++ b/Engine/Plugins/Experimental/AnimNextPoseSearch/Source/AnimNextPoseSearch/Private/MotionMatchingTrait.cpp
@@ -92,8 +92,8 @@ namespace UE::AnimNext
 			// We adjust the motion matching state asset time to the current player node's asset time. This is done 
 			// because the player node may have ticked more or less time than we expected due to variable dt or the 
 			// dynamic playback rate adjustment and as such the motion matching state does not update by itself
-			const FTimelineProgress TimelineProgress = TimelineTrait.GetProgress(Context);
-			MotionMatchingState.AdjustAssetTime(TimelineProgress.GetPosition());
+			const float TimelinePosition = TimelineTrait.GetState(Context).GetPosition();
+			MotionMatchingState.AdjustAssetTime(TimelinePosition);
 		}
 
 		const TConstArrayView<TObjectPtr<const UPoseSearchDatabase>> DatabasesToSearch = SharedData->GetDatabases(Binding);
diff --git a/Engine/Plugins/Experimental/AnimNextStateTree/Source/AnimNextStateTree/Private/AnimNextStateTreeContext.cpp b/Engine/Plugins/Experimental/AnimNextStateTree/Source/AnimNextStateTree/Private/AnimNextStateTreeContext.cpp
--- a/Engine/Plugins/Experimental/AnimNextStateTree/Source/AnimNextStateTree/Private/AnimNextStateTreeContext.cpp
+++ b/Engine/Plugins/Experimental/AnimNextStateTree/Source/AnimNextStateTree/Private/AnimNextStateTreeContext.cpp
@@ -52,9 +52,7 @@ float FAnimNextStateTreeTraitContext::QueryPlaybackRatio(TNonNullPtr<UObject> In
 		UE::AnimNext::TTraitBinding<UE::AnimNext::ITimeline> Timeline;
 		if (Binding->GetStackInterface<UE::AnimNext::ITimeline>(Timeline))
 		{
-			UE::AnimNext::FTimelineProgress Progress = Timeline.GetProgress(*Context);
-
-			return Progress.GetPositionRatio();
+			return Timeline.GetState(*Context).GetPositionRatio();
 		}
 	}
 
diff --git a/Engine/Plugins/Experimental/AnimNextWarping/Source/AnimNextWarping/Private/SteeringTrait.cpp b/Engine/Plugins/Experimental/AnimNextWarping/Source/AnimNextWarping/Private/SteeringTrait.cpp
--- a/Engine/Plugins/Experimental/AnimNextWarping/Source/AnimNextWarping/Private/SteeringTrait.cpp
+++ b/Engine/Plugins/Experimental/AnimNextWarping/Source/AnimNextWarping/Private/SteeringTrait.cpp
@@ -93,7 +93,7 @@ namespace UE::AnimNext
 				return;
 			}
 
-			InstanceData->CurrentAnimAssetTime = TimelineTrait.GetProgress(Context).GetPosition();
+			InstanceData->CurrentAnimAssetTime = TimelineTrait.GetState(Context).GetPosition();
 		}
 
 		// Update target orientation / root motion
