[TSAN]
- Fix warnings about GPrintLogTimes since we know it is racy in some cases and we're OK with that for now

#rnx

diff --git a/Engine/Plugins/Animation/ACLPlugin/Source/ACLPlugin/Private/ACLPluginModule.cpp b/Engine/Plugins/Animation/ACLPlugin/Source/ACLPlugin/Private/ACLPluginModule.cpp
--- a/Engine/Plugins/Animation/ACLPlugin/Source/ACLPlugin/Private/ACLPluginModule.cpp
+++ b/Engine/Plugins/Animation/ACLPlugin/Source/ACLPlugin/Private/ACLPluginModule.cpp
@@ -113,7 +113,7 @@ static SIZE_T GetCompressedCurveSize(const FCompressedAnimSequence& CompressedDa
 void FACLPlugin::ListCodecs(const TArray<FString>& Args)
 {
 	// Turn off log times to make diffing easier
-	TGuardValue<ELogTimes::Type> DisableLogTimes(GPrintLogTimes, ELogTimes::None);
+	TGuardValue DisableLogTimes(GPrintLogTimes, ELogTimes::None);
 
 	// Make sure to log everything
 	const ELogVerbosity::Type OldVerbosity = LogAnimationCompression.GetVerbosity();
@@ -256,7 +256,7 @@ void FACLPlugin::ListCodecs(const TArray<FString>& Args)
 void FACLPlugin::ListAnimSequences(const TArray<FString>& Args)
 {
 	// Turn off log times to make diffing easier
-	TGuardValue<ELogTimes::Type> DisableLogTimes(GPrintLogTimes, ELogTimes::None);
+	TGuardValue DisableLogTimes(GPrintLogTimes, ELogTimes::None);
 
 	// Make sure to log everything
 	const ELogVerbosity::Type OldVerbosity = LogAnimationCompression.GetVerbosity();
diff --git a/Engine/Plugins/Editor/AssetManagerEditor/Source/AssetManagerEditor/Private/AssetManagerEditorModule.cpp b/Engine/Plugins/Editor/AssetManagerEditor/Source/AssetManagerEditor/Private/AssetManagerEditorModule.cpp
--- a/Engine/Plugins/Editor/AssetManagerEditor/Source/AssetManagerEditor/Private/AssetManagerEditorModule.cpp
+++ b/Engine/Plugins/Editor/AssetManagerEditor/Source/AssetManagerEditor/Private/AssetManagerEditorModule.cpp
@@ -2328,7 +2328,7 @@ void FAssetManagerEditorModule::InitializeRegistrySources(bool bNeedManagementDa
 void FAssetManagerEditorModule::PerformAuditConsoleCommand(const TArray<FString>& Args)
 {
 	// Turn off as it makes diffing hard
-	TGuardValue<ELogTimes::Type> DisableLogTimes(GPrintLogTimes, ELogTimes::None);
+	TGuardValue DisableLogTimes(GPrintLogTimes, ELogTimes::None);
 
 	UAssetManager::Get().UpdateManagementDatabase();
 
diff --git a/Engine/Plugins/Experimental/AnimNextAnimGraph/Source/AnimNextAnimGraph/Private/AnimNextAnimGraphModule.cpp b/Engine/Plugins/Experimental/AnimNextAnimGraph/Source/AnimNextAnimGraph/Private/AnimNextAnimGraphModule.cpp
--- a/Engine/Plugins/Experimental/AnimNextAnimGraph/Source/AnimNextAnimGraph/Private/AnimNextAnimGraphModule.cpp
+++ b/Engine/Plugins/Experimental/AnimNextAnimGraph/Source/AnimNextAnimGraph/Private/AnimNextAnimGraphModule.cpp
@@ -121,7 +121,7 @@ FString FAnimNextAnimGraphModule::GetReferencerName() const
 void FAnimNextAnimGraphModule::ListNodeTemplates(const TArray<FString>& Args)
 {
 	// Turn off log times to make diff-ing easier
-	TGuardValue<ELogTimes::Type> DisableLogTimes(GPrintLogTimes, ELogTimes::None);
+	TGuardValue DisableLogTimes(GPrintLogTimes, ELogTimes::None);
 
 	// Make sure to log everything
 	const ELogVerbosity::Type OldVerbosity = LogAnimation.GetVerbosity();
@@ -176,7 +176,7 @@ void FAnimNextAnimGraphModule::ListNodeTemplates(const TArray<FString>& Args)
 void FAnimNextAnimGraphModule::ListAnimationGraphs(const TArray<FString>& Args)
 {
 	// Turn off log times to make diff-ing easier
-	TGuardValue<ELogTimes::Type> DisableLogTimes(GPrintLogTimes, ELogTimes::None);
+	TGuardValue DisableLogTimes(GPrintLogTimes, ELogTimes::None);
 
 	// Make sure to log everything
 	const ELogVerbosity::Type OldVerbosity = LogAnimation.GetVerbosity();
diff --git a/Engine/Source/Developer/IoStoreUtilities/Private/IoStoreUtilities.cpp b/Engine/Source/Developer/IoStoreUtilities/Private/IoStoreUtilities.cpp
--- a/Engine/Source/Developer/IoStoreUtilities/Private/IoStoreUtilities.cpp
+++ b/Engine/Source/Developer/IoStoreUtilities/Private/IoStoreUtilities.cpp
@@ -7501,7 +7501,7 @@ int32 Describe(
 	}
 	{
 		TRACE_CPUPROFILER_EVENT_SCOPE(GenerateReport);
-		TGuardValue<ELogTimes::Type> GuardPrintLogTimes(GPrintLogTimes, ELogTimes::None);
+		TGuardValue GuardPrintLogTimes(GPrintLogTimes, ELogTimes::None);
 		TGuardValue GuardPrintLogCategory(GPrintLogCategory, false);
 		TGuardValue GuardPrintLogVerbosity(GPrintLogVerbosity, false);
 
@@ -8326,7 +8326,7 @@ bool DiffIoStoreContainer(const TCHAR* CmdLine)
 
 bool LegacyDiffIoStoreContainers(const TCHAR* InContainerFilename1, const TCHAR* InContainerFilename2, bool bInLogUniques1, bool bInLogUniques2, const FKeyChain& InKeyChain1, const FKeyChain* InKeyChain2)
 {
-	TGuardValue<ELogTimes::Type> DisableLogTimes(GPrintLogTimes, ELogTimes::None);
+	TGuardValue DisableLogTimes(GPrintLogTimes, ELogTimes::None);
 	UE_LOG(LogIoStore, Log, TEXT("FileEventType, FileName, Size1, Size2"));
 
 	TUniquePtr<FIoStoreReader> Reader1 = CreateIoStoreReader(InContainerFilename1, InKeyChain1);
diff --git a/Engine/Source/Developer/PakFileUtilities/Private/PakFileUtilities.cpp b/Engine/Source/Developer/PakFileUtilities/Private/PakFileUtilities.cpp
--- a/Engine/Source/Developer/PakFileUtilities/Private/PakFileUtilities.cpp
+++ b/Engine/Source/Developer/PakFileUtilities/Private/PakFileUtilities.cpp
@@ -4268,7 +4268,7 @@ bool DiffFilesInPaks(const FString& InPakFilename1, const FString& InPakFilename
 	int32 NumDifferentContents = 0;
 	int32 NumEqualContents = 0;
 
-	TGuardValue<ELogTimes::Type> DisableLogTimes(GPrintLogTimes, ELogTimes::None);
+	TGuardValue DisableLogTimes(GPrintLogTimes, ELogTimes::None);
 
 	TRefCountPtr<FPakFile> PakFilePtr1 = MakeRefCount<FPakFile>(&FPlatformFileManager::Get().GetPlatformFile(), *InPakFilename1, /*bIsSigned=*/ false);
 	FPakFile& PakFile1 = *PakFilePtr1;
diff --git a/Engine/Source/Editor/KismetCompiler/Private/KismetCompiler.cpp b/Engine/Source/Editor/KismetCompiler/Private/KismetCompiler.cpp
--- a/Engine/Source/Editor/KismetCompiler/Private/KismetCompiler.cpp
+++ b/Engine/Source/Editor/KismetCompiler/Private/KismetCompiler.cpp
@@ -5179,7 +5179,7 @@ void FKismetCompilerContext::CompileFunctions(EInternalCompilerFlags InternalFla
 
 		if (bDisplayBytecode && bIsFullCompile && !IsRunningCommandlet())
 		{
-			TGuardValue<ELogTimes::Type> DisableLogTimes(GPrintLogTimes, ELogTimes::None);
+			TGuardValue DisableLogTimes(GPrintLogTimes, ELogTimes::None);
 
 			FKismetBytecodeDisassembler Disasm(*GLog);
 
diff --git a/Engine/Source/Editor/UnrealEd/Private/Commandlets/PackageUtilities.cpp b/Engine/Source/Editor/UnrealEd/Private/Commandlets/PackageUtilities.cpp
--- a/Engine/Source/Editor/UnrealEd/Private/Commandlets/PackageUtilities.cpp
+++ b/Engine/Source/Editor/UnrealEd/Private/Commandlets/PackageUtilities.cpp
@@ -1691,7 +1691,7 @@ int32 UPkgInfoCommandlet::Main( const FString& Params )
 
 		{
 			// Turn off log categories etc as it makes diffing hard
-			TGuardValue<ELogTimes::Type> GuardPrintLogTimes(GPrintLogTimes, ELogTimes::None);
+			TGuardValue GuardPrintLogTimes(GPrintLogTimes, ELogTimes::None);
 			TGuardValue GuardPrintLogCategory(GPrintLogCategory, false);
 			TGuardValue GuardPrintLogVerbosity(GPrintLogVerbosity, false);
 
diff --git a/Engine/Source/Editor/UnrealEd/Private/Cooker/DiffPackageWriter.cpp b/Engine/Source/Editor/UnrealEd/Private/Cooker/DiffPackageWriter.cpp
--- a/Engine/Source/Editor/UnrealEd/Private/Cooker/DiffPackageWriter.cpp
+++ b/Engine/Source/Editor/UnrealEd/Private/Cooker/DiffPackageWriter.cpp
@@ -390,7 +390,7 @@ void FDiffPackageWriter::ConditionallyDumpObjList()
 			FString ObjListExec = TEXT("OBJ LIST ");
 			ObjListExec += DumpObjListParams;
 
-			TGuardValue<ELogTimes::Type> GuardLogTimes(GPrintLogTimes, ELogTimes::None);
+			TGuardValue GuardLogTimes(GPrintLogTimes, ELogTimes::None);
 			TGuardValue GuardLogCategory(GPrintLogCategory, false);
 			TGuardValue GuardPrintLogVerbosity(GPrintLogVerbosity, false);
 
@@ -415,7 +415,7 @@ void FDiffPackageWriter::ConditionallyDumpObjects()
 				AllObjects.Sort();
 			}
 
-			TGuardValue<ELogTimes::Type> GuardLogTimes(GPrintLogTimes, ELogTimes::None);
+			TGuardValue GuardLogTimes(GPrintLogTimes, ELogTimes::None);
 			TGuardValue GuardLogCategory(GPrintLogCategory, false);
 			TGuardValue GuardPrintLogVerbosity(GPrintLogVerbosity, false);
 
diff --git a/Engine/Source/Editor/UnrealEd/Private/ReferenceInfoUtils.cpp b/Engine/Source/Editor/UnrealEd/Private/ReferenceInfoUtils.cpp
--- a/Engine/Source/Editor/UnrealEd/Private/ReferenceInfoUtils.cpp
+++ b/Engine/Source/Editor/UnrealEd/Private/ReferenceInfoUtils.cpp
@@ -630,7 +630,7 @@ namespace ReferenceInfoUtils
 
 	void GenerateOutput(UWorld* InWorld, int32 Depth, bool bShowDefault, bool bShowScript)
 	{
-		auto PrintLogTimes = GPrintLogTimes;
+		ELogTimes::Type PrintLogTimes = GPrintLogTimes;
 
 		// Create log file
 		
diff --git a/Engine/Source/Runtime/Core/Private/Misc/CoreGlobals.cpp b/Engine/Source/Runtime/Core/Private/Misc/CoreGlobals.cpp
--- a/Engine/Source/Runtime/Core/Private/Misc/CoreGlobals.cpp
+++ b/Engine/Source/Runtime/Core/Private/Misc/CoreGlobals.cpp
@@ -443,7 +443,7 @@ bool					GEmitDrawEvents					= false;
 /** Determines what kind of trace should occur, NAME_None for none.											*/
 FLazyName				GCurrentTraceName;
 /** How to print the time in log output																		*/
-ELogTimes::Type	GPrintLogTimes			= ELogTimes::None;
+TSAN_ATOMIC(ELogTimes::Type) GPrintLogTimes             = ELogTimes::None;
 /** How to print the category in log output.																*/
 TSAN_ATOMIC(bool)		GPrintLogCategory				= true;
 /** How to print the verbosity in log output.																*/
diff --git a/Engine/Source/Runtime/Core/Public/CoreGlobals.h b/Engine/Source/Runtime/Core/Public/CoreGlobals.h
--- a/Engine/Source/Runtime/Core/Public/CoreGlobals.h
+++ b/Engine/Source/Runtime/Core/Public/CoreGlobals.h
@@ -589,7 +589,7 @@ extern CORE_API bool GIsGameThreadIdInitialized;
 extern CORE_API FLazyName GCurrentTraceName;
 
 /** How to print the time in log output. */
-extern CORE_API ELogTimes::Type GPrintLogTimes;
+extern CORE_API TSAN_ATOMIC(ELogTimes::Type) GPrintLogTimes;
 
 /** How to print the category in log output. */
 extern CORE_API TSAN_ATOMIC(bool) GPrintLogCategory;
diff --git a/Engine/Source/Runtime/Engine/Private/Materials/MaterialShader.cpp b/Engine/Source/Runtime/Engine/Private/Materials/MaterialShader.cpp
--- a/Engine/Source/Runtime/Engine/Private/Materials/MaterialShader.cpp
+++ b/Engine/Source/Runtime/Engine/Private/Materials/MaterialShader.cpp
@@ -3754,7 +3754,7 @@ void FMaterialShaderMap::GetEstimatedLWCFuncUsageComplexity(uint32& LWCComplexit
 void FMaterialShaderMap::DumpDebugInfo(FOutputDevice& OutputDevice) const
 {
 	// Turn off as it makes diffing hard
-	TGuardValue<ELogTimes::Type> DisableLogTimes(GPrintLogTimes, ELogTimes::None);
+	TGuardValue DisableLogTimes(GPrintLogTimes, ELogTimes::None);
 
 	OutputDevice.Logf(TEXT("Frequency, Target, VFType, ShaderType, SourceHash, VFSourceHash, OutputHash, IsShaderPipeline"));
 
diff --git a/Engine/Source/Runtime/Engine/Private/UnrealEngine.cpp b/Engine/Source/Runtime/Engine/Private/UnrealEngine.cpp
--- a/Engine/Source/Runtime/Engine/Private/UnrealEngine.cpp
+++ b/Engine/Source/Runtime/Engine/Private/UnrealEngine.cpp
@@ -7797,7 +7797,7 @@ bool UEngine::HandleMemReportDeferredCommand( const TCHAR* Cmd, FOutputDevice& A
 	UE_LOG(LogEngine, Log, TEXT("MemReportDeferred: Executing \"%s\" profile."), *InTypeName);
 
 	// Turn off as it makes diffing hard
-	TGuardValue<ELogTimes::Type> DisableLogTimes(GPrintLogTimes, ELogTimes::None);
+	TGuardValue DisableLogTimes(GPrintLogTimes, ELogTimes::None);
 
 	// Flush rendering and do a GC
 	FlushAsyncLoading();
@@ -8288,7 +8288,7 @@ bool UEngine::HandleDumpParticleCountsCommand( const TCHAR* Cmd, FOutputDevice&
 
 bool UEngine::HandleListLoadedPackagesCommand( const TCHAR* Cmd, FOutputDevice& Ar )
 {
-	TGuardValue<ELogTimes::Type> DisableLogTimes(GPrintLogTimes, ELogTimes::None);
+	TGuardValue DisableLogTimes(GPrintLogTimes, ELogTimes::None);
 
 	struct FPackageInfo
 	{
