AnimNext: Notify support

Added new trait and trait interface to support notify dispatch. Context for legacy notifies (i.e. the skeletal mesh component) is provided via a rig unit that can be run on a module.
Added module component to handle legacy notifies, dispatching on the game thread
Integrated with the find/replace systems. How external pluguins can register with Persona to subscribe to find/replace of notifies.
Minor refactor of public/private variables and functions on FAnimNextModuleInstance

diff --git a/Engine/Plugins/Experimental/AnimNext/Source/AnimNext/Internal/Module/AnimNextModuleContextData.h b/Engine/Plugins/Experimental/AnimNext/Source/AnimNext/Internal/Module/AnimNextModuleContextData.h
--- a/Engine/Plugins/Experimental/AnimNext/Source/AnimNext/Internal/Module/AnimNextModuleContextData.h
+++ b/Engine/Plugins/Experimental/AnimNext/Source/AnimNext/Internal/Module/AnimNextModuleContextData.h
@@ -31,7 +31,7 @@ struct FAnimNextModuleContextData
 	// Get the object that the module instance is bound to, if any
 	UObject* GetObject() const
 	{
-		return ModuleInstance ? ModuleInstance->Object : nullptr;
+		return ModuleInstance ? ModuleInstance->GetObject() : nullptr;
 	}
 
 	FAnimNextModuleInstance& GetModuleInstance() const
diff --git a/Engine/Plugins/Experimental/AnimNext/Source/AnimNext/Internal/Module/AnimNextModuleInstance.h b/Engine/Plugins/Experimental/AnimNext/Source/AnimNext/Internal/Module/AnimNextModuleInstance.h
--- a/Engine/Plugins/Experimental/AnimNext/Source/AnimNext/Internal/Module/AnimNextModuleInstance.h
+++ b/Engine/Plugins/Experimental/AnimNext/Source/AnimNext/Internal/Module/AnimNextModuleInstance.h
@@ -17,10 +17,16 @@
 #include "AnimNextModuleInstance.generated.h"
 
 struct FAnimNextModuleInstance;
+struct FRigUnit_AnimNextRunAnimationGraph;
+class UAnimNextWorldSubsystem;
+struct FRigUnit_AnimNextInitializeEvent;
+struct FRigUnit_CopyModuleProxyVariables;
 
 namespace UE::AnimNext
 {
 	struct FProxyVariablesContext;
+	struct FModuleEventTickFunction;
+	struct FModuleEndTickFunction;
 }
 
 namespace UE::AnimNext::Debug
@@ -70,36 +76,23 @@ struct FAnimNextModuleInstance : public FAnimNextDataInterfaceInstance
 		EAnimNextModuleInitMethod InInitMethod);
 	ANIMNEXT_API ~FAnimNextModuleInstance();
 
-	// Setup the entry
-	void Initialize();
-
-	// Clear data that binds the schedule to a runtime (e.g. tick functions) and any instance data
-	void ResetBindingsAndInstanceData();
-
-	// Used for pooling
-	void Invalidate();
-
 	// Enables/disables the ticking of this entry
 	void Enable(bool bInEnabled);
 
-	// Clears the bTickEvenWhenPaused flags of the entries' tick functions
-	void ClearTickFunctionPauseFlags();
-
 	// Queues an input trait event
 	// Input events will be processed in the next graph update after they are queued
-	void QueueInputTraitEvent(FAnimNextTraitEventPtr Event);
+	ANIMNEXT_API void QueueInputTraitEvent(FAnimNextTraitEventPtr Event);
 
-	// Flip proxy variable buffers then copy any dirty values
-	void CopyProxyVariables();
+	// Queues an output trait event
+	// Output events will be processed at the end of the schedule tick
+	ANIMNEXT_API void QueueOutputTraitEvent(FAnimNextTraitEventPtr Event);
 
-#if WITH_EDITOR
-	// Resets internal state if the module we are bound to is recompiled in editor
-	void OnModuleCompiled();
-#endif
+	// Get the object that this module is bound to
+	ANIMNEXT_API UObject* GetObject() const { return Object; }
 
 	// Get the module that this instance represents
 	ANIMNEXT_API const UAnimNextModule* GetModule() const;
-	
+
 	// Returns a typed module instance component, creating it lazily the first time it is queried
 	template<class ComponentType>
 	ComponentType& GetComponent();
@@ -124,9 +117,6 @@ struct FAnimNextModuleInstance : public FAnimNextDataInterfaceInstance
 	ANIMNEXT_API void ShowDebugDrawing(bool bInShowDebugDrawing);
 #endif
 
-	// Called to disable the initialize event for this module on the game thread when it has executed for the first time
-	void DisableInitializeEvent();
-
 	// Run a simple task on the GT via FFunctionGraphTask::CreateAndDispatchWhenReady
 	static void RunTaskOnGameThread(TUniqueFunction<void(void)>&& InFunction);
 
@@ -135,7 +125,56 @@ struct FAnimNextModuleInstance : public FAnimNextDataInterfaceInstance
 
 	// Find the tick function for the specified event
 	UE::AnimNext::FModuleEventTickFunction* FindTickFunction(FName InEventName);
-	
+
+	// Run the specified RigVM event
+	ANIMNEXT_API void RunRigVMEvent(FName InEventName, float InDeltaTime);
+
+	// Get the world type that this module was instantiated within
+	EWorldType::Type GetWorldType() const { return WorldType; }
+
+	// Get the world type that this module was instantiated within
+	TArrayView<UE::AnimNext::FModuleEventTickFunction> GetTickFunctions() { return TickFunctions; }
+
+private:
+	// Setup the entry
+	void Initialize();
+
+	// Clear data that binds the schedule to a runtime (e.g. tick functions) and any instance data
+	void ResetBindingsAndInstanceData();
+
+	// Used for pooling
+	void Invalidate();
+
+	// Clears the bTickEvenWhenPaused flags of the entries' tick functions
+	void ClearTickFunctionPauseFlags();
+
+	// Flip proxy variable buffers then copy any dirty values
+	void CopyProxyVariables();
+
+#if WITH_EDITOR
+	// Resets internal state if the module we are bound to is recompiled in editor
+	void OnModuleCompiled();
+#endif
+
+	// Called to disable the initialize event for this module on the game thread when it has executed for the first time
+	void DisableInitializeEvent();
+
+	// Give the module components a chance to handle events and call their OnEndExecution extension point 
+	void EndExecution(float InDeltaTime);
+
+	// Handles trait events at the module level.
+	// Called in the module end tick function before any events are dispatched.
+	void RaiseTraitEvents(const UE::AnimNext::FTraitEventList& EventList);
+
+private:
+	friend UE::AnimNext::FModuleEventTickFunction;
+	friend UE::AnimNext::FModuleEndTickFunction;
+	friend FRigUnit_AnimNextRunAnimationGraph;
+	friend UAnimNextWorldSubsystem;
+	friend FRigUnit_AnimNextInitializeEvent;
+	friend UE::AnimNext::FProxyVariablesContext;
+	friend FRigUnit_CopyModuleProxyVariables;
+
 	// Object this entry is bound to
 	UPROPERTY(Transient)
 	TObjectPtr<UObject> Object = nullptr;
@@ -200,7 +239,7 @@ struct FAnimNextModuleInstance : public FAnimNextDataInterfaceInstance
 	EAnimNextModuleInitMethod InitMethod = EAnimNextModuleInitMethod::InitializeAndPauseInEditor;
 
 	// Whether this represents an editor object 
-	bool bIsEditor : 1 = false;
+	EWorldType::Type WorldType = EWorldType::None;
 
 #if WITH_EDITOR
 	// Whether we are currently recreating this instance because of compilation/reinstancing
diff --git a/Engine/Plugins/Experimental/AnimNext/Source/AnimNext/Internal/Module/ModuleInstanceComponent.h b/Engine/Plugins/Experimental/AnimNext/Source/AnimNext/Internal/Module/ModuleInstanceComponent.h
--- a/Engine/Plugins/Experimental/AnimNext/Source/AnimNext/Internal/Module/ModuleInstanceComponent.h
+++ b/Engine/Plugins/Experimental/AnimNext/Source/AnimNext/Internal/Module/ModuleInstanceComponent.h
@@ -6,6 +6,7 @@
 #include "ModuleInstanceComponent.generated.h"
 
 struct FAnimNextModuleInstance;
+struct FAnimNextTraitEvent;
 
 // Declares a module instance component and implements the necessary boilerplate
 #define ANIMNEXT_DECLARE_MODULE_INSTANCE_COMPONENT(ComponentType) \
@@ -36,6 +37,12 @@ struct FModuleInstanceComponent
 	// Returns the owning module instance this component lives on
 	const FAnimNextModuleInstance& GetModuleInstance() const { return *OwnerInstance; }
 
+	// Called at the end of module execution for any output events to be handled, before OnEndExecution
+	virtual void OnTraitEvent(FAnimNextTraitEvent& Event) {}
+
+	// Called at end of module execution
+	virtual void OnEndExecution(float InDeltaTime) {}
+
 private:
 
 	// The owning module instance this component lives on
diff --git a/Engine/Plugins/Experimental/AnimNext/Source/AnimNext/Internal/Module/ModuleTickFunction.h b/Engine/Plugins/Experimental/AnimNext/Source/AnimNext/Internal/Module/ModuleTickFunction.h
--- a/Engine/Plugins/Experimental/AnimNext/Source/AnimNext/Internal/Module/ModuleTickFunction.h
+++ b/Engine/Plugins/Experimental/AnimNext/Source/AnimNext/Internal/Module/ModuleTickFunction.h
@@ -60,7 +60,7 @@ struct FModuleEndTickFunction : public FTickFunction
 	virtual void ExecuteTick(float DeltaTime, ELevelTick TickType, ENamedThreads::Type CurrentThread, const FGraphEventRef& MyCompletionGraphEvent) override;
 	virtual FString DiagnosticMessage() override;
 
-	void Run();
+	void Run(float DeltaTime);
 
 	FAnimNextModuleInstance* ModuleInstance = nullptr;
 };
diff --git a/Engine/Plugins/Experimental/AnimNext/Source/AnimNext/Private/Component/AnimNextComponentWorldSubsystem.cpp b/Engine/Plugins/Experimental/AnimNext/Source/AnimNext/Private/Component/AnimNextComponentWorldSubsystem.cpp
--- a/Engine/Plugins/Experimental/AnimNext/Source/AnimNext/Private/Component/AnimNextComponentWorldSubsystem.cpp
+++ b/Engine/Plugins/Experimental/AnimNext/Source/AnimNext/Private/Component/AnimNextComponentWorldSubsystem.cpp
@@ -61,7 +61,7 @@ void UAnimNextComponentWorldSubsystem::OnModuleCompiled(UAnimNextModule* InModul
 	{
 		if(Instance.GetModule() == InModule)
 		{
-			UAnimNextComponent* AnimNextComponent = CastChecked<UAnimNextComponent>(Instance.Object);
+			UAnimNextComponent* AnimNextComponent = CastChecked<UAnimNextComponent>(Instance.GetObject());
 			AnimNextComponent->OnModuleCompiled();
 		}
 	}
diff --git a/Engine/Plugins/Experimental/AnimNext/Source/AnimNext/Private/Graph/RigUnit_GetActorTransform.cpp b/Engine/Plugins/Experimental/AnimNext/Source/AnimNext/Private/Graph/RigUnit_GetActorTransform.cpp
--- a/Engine/Plugins/Experimental/AnimNext/Source/AnimNext/Private/Graph/RigUnit_GetActorTransform.cpp
+++ b/Engine/Plugins/Experimental/AnimNext/Source/AnimNext/Private/Graph/RigUnit_GetActorTransform.cpp
@@ -13,7 +13,7 @@ FRigUnit_GetActorTransform_Execute()
 
 	// TODO: this is not thread safe at the moment. Once we get the 'weak semantics' CL checked in, we can move commonly-used debug info
 	// into a module component that can be used to access GT data where needed. 
-	UAnimNextComponent* Component = Cast<UAnimNextComponent>(ModuleInstance.Object);
+	UAnimNextComponent* Component = Cast<UAnimNextComponent>(ModuleInstance.GetObject());
 	if(Component == nullptr)
 	{
 		Transform = FTransform::Identity;
diff --git a/Engine/Plugins/Experimental/AnimNext/Source/AnimNext/Private/Module/AnimNextModuleInstance.cpp b/Engine/Plugins/Experimental/AnimNext/Source/AnimNext/Private/Module/AnimNextModuleInstance.cpp
--- a/Engine/Plugins/Experimental/AnimNext/Source/AnimNext/Private/Module/AnimNextModuleInstance.cpp
+++ b/Engine/Plugins/Experimental/AnimNext/Source/AnimNext/Private/Module/AnimNextModuleInstance.cpp
@@ -128,7 +128,10 @@ void FAnimNextModuleInstance::Initialize()
 	check(Handle.IsValid());
 
 	UWorld* World = Object->GetWorld();
-	bIsEditor = World && World->WorldType == EWorldType::Editor;
+	if(World)
+	{
+		WorldType = World->WorldType;
+	}
 
 	// Get all the module events from the VM entry points, sorted by phase
 	URigVM* VM = Module->RigVM;
@@ -292,6 +295,11 @@ void FAnimNextModuleInstance::QueueInputTraitEvent(FAnimNextTraitEventPtr Event)
 	InputEventList.Push(MoveTemp(Event));
 }
 
+void FAnimNextModuleInstance::QueueOutputTraitEvent(FAnimNextTraitEventPtr Event)
+{
+	OutputEventList.Push(MoveTemp(Event));
+}
+
 void FAnimNextModuleInstance::Enable(bool bInEnabled)
 {
 	using namespace UE::AnimNext;
@@ -491,4 +499,70 @@ UE::AnimNext::FModuleEventTickFunction* FAnimNextModuleInstance::FindTickFunctio
 		}
 	}
 	return nullptr;
+}
+
+void FAnimNextModuleInstance::EndExecution(float InDeltaTime)
+{
+	// Give the module a chance to handle events
+	RaiseTraitEvents(OutputEventList);
+
+	// Give each component a chance to finalize execution
+	for(auto It = ComponentMap.CreateIterator(); It; ++It)
+	{
+		FModuleInstanceComponent& Component = It->Value.GetMutable<FModuleInstanceComponent>();
+		Component.OnEndExecution(InDeltaTime);
+	}
+}
+
+void FAnimNextModuleInstance::RaiseTraitEvents(const UE::AnimNext::FTraitEventList& EventList)
+{
+	for(auto It = ComponentMap.CreateIterator(); It; ++It)
+	{
+		FModuleInstanceComponent& Component = It->Value.GetMutable<FModuleInstanceComponent>();
+
+		// Event handlers can raise events and as such the list may change while we iterate
+		// However, if an event is added while we iterate, we will not visit it
+		const int32 NumEvents = EventList.Num();
+		for (int32 EventIndex = 0; EventIndex < NumEvents; ++EventIndex)
+		{
+			const FAnimNextTraitEventPtr Event = EventList[EventIndex];
+			if (Event->IsValid())
+			{
+				Component.OnTraitEvent(*Event);
+			}
+		}
+	}
+}
+
+void FAnimNextModuleInstance::RunRigVMEvent(FName InEventName, float InDeltaTime)
+{
+	URigVM* VM = GetModule()->RigVM;
+	if(VM == nullptr)
+	{
+		return;
+	}
+
+	if(!VM->ContainsEntry(InEventName))
+	{
+		return;
+	}
+
+	check(ExtendedExecuteContext.VMHash == VM->GetVMHash());
+
+	FAnimNextExecuteContext& AnimNextContext = ExtendedExecuteContext.GetPublicDataSafe<FAnimNextExecuteContext>();
+
+	// RigVM setup
+	AnimNextContext.SetDeltaTime(InDeltaTime);
+	AnimNextContext.SetOwningObject(Object);
+
+#if UE_ENABLE_DEBUG_DRAWING
+	AnimNextContext.SetDrawInterface(GetDebugDrawInterface());
+#endif
+
+	// Insert our context data for the scope of execution
+	FAnimNextModuleContextData ContextData(this);
+	UE::AnimNext::FScopedExecuteContextData ContextDataScope(AnimNextContext, ContextData);
+
+	// Run the VM for this event
+	VM->ExecuteVM(ExtendedExecuteContext, InEventName);
 }
\ No newline at end of file
diff --git a/Engine/Plugins/Experimental/AnimNext/Source/AnimNext/Private/Module/ModuleTickFunction.cpp b/Engine/Plugins/Experimental/AnimNext/Source/AnimNext/Private/Module/ModuleTickFunction.cpp
--- a/Engine/Plugins/Experimental/AnimNext/Source/AnimNext/Private/Module/ModuleTickFunction.cpp
+++ b/Engine/Plugins/Experimental/AnimNext/Source/AnimNext/Private/Module/ModuleTickFunction.cpp
@@ -22,13 +22,16 @@ namespace UE::AnimNext
 
 void FModuleEndTickFunction::ExecuteTick(float DeltaTime, ELevelTick TickType, ENamedThreads::Type CurrentThread, const FGraphEventRef& MyCompletionGraphEvent)
 {
-	Run();
+	Run(DeltaTime);
 }
 
-void FModuleEndTickFunction::Run()
+void FModuleEndTickFunction::Run(float DeltaTime)
 {
 	SCOPE_CYCLE_COUNTER(STAT_AnimNext_Module_EndTick);
 
+	// Give module a chance to finish up processing
+	ModuleInstance->EndExecution(DeltaTime);
+
 	// Decrement the remaining lifetime of the input events we processed and queue up any remaining events
 	DecrementLifetimeAndPurgeExpired(ModuleInstance->InputEventList, ModuleInstance->OutputEventList);
 
@@ -66,7 +69,7 @@ void FModuleEndTickFunction::Run()
 	{
 		if(ModuleInstance->InitMethod == EAnimNextModuleInitMethod::InitializeAndPause
 #if WITH_EDITOR
-			|| (ModuleInstance->InitMethod == EAnimNextModuleInitMethod::InitializeAndPauseInEditor && ModuleInstance->bIsEditor)
+			|| (ModuleInstance->InitMethod == EAnimNextModuleInitMethod::InitializeAndPauseInEditor && ModuleInstance->WorldType == EWorldType::Editor)
 #endif
 		)
 		{
@@ -120,28 +123,7 @@ void FModuleEventTickFunction::Run(float InDeltaTime)
 		Function.GetValue()(FModuleTaskContext(*ModuleInstance));
 	}
 
-	if (URigVM* VM = ModuleInstance->GetModule()->RigVM)
-	{
-		FRigVMExtendedExecuteContext& Context = ModuleInstance->GetExtendedExecuteContext();
-		check(Context.VMHash == VM->GetVMHash());
-
-		FAnimNextExecuteContext& AnimNextContext = Context.GetPublicDataSafe<FAnimNextExecuteContext>();
-
-		// RigVM setup
-		AnimNextContext.SetDeltaTime(InDeltaTime);
-		AnimNextContext.SetOwningObject(ModuleInstance->Object);
-
-#if UE_ENABLE_DEBUG_DRAWING
-		AnimNextContext.SetDrawInterface(ModuleInstance->GetDebugDrawInterface());
-#endif
-
-		// Insert our context data for the scope of execution
-		FAnimNextModuleContextData ContextData(ModuleInstance);
-		FScopedExecuteContextData ContextDataScope(AnimNextContext, ContextData);
-
-		// Run the VM for this event
-		VM->ExecuteVM(Context, EventName);
-	}
+	ModuleInstance->RunRigVMEvent(EventName, InDeltaTime);
 
 	while (!PostExecuteTasks.IsEmpty())
 	{
@@ -190,7 +172,7 @@ void FModuleEventTickFunction::InitializeAndRunModule(FAnimNextModuleInstance& I
 	{
 		TickFunction->Run(0.0f);
 	}
-	InModuleInstance.EndTickFunction.Run();
+	InModuleInstance.EndTickFunction.Run(0.0f);
 }
 
 #endif
diff --git a/Engine/Plugins/Experimental/AnimNext/Source/AnimNext/Private/Module/RigUnit_AnimNextModuleEventDependencies.cpp b/Engine/Plugins/Experimental/AnimNext/Source/AnimNext/Private/Module/RigUnit_AnimNextModuleEventDependencies.cpp
--- a/Engine/Plugins/Experimental/AnimNext/Source/AnimNext/Private/Module/RigUnit_AnimNextModuleEventDependencies.cpp
+++ b/Engine/Plugins/Experimental/AnimNext/Source/AnimNext/Private/Module/RigUnit_AnimNextModuleEventDependencies.cpp
@@ -38,7 +38,7 @@ static void ApplyDependenciesHelper(const FAnimNextExecuteContext& InContext, co
 		return;
 	}
 
-	TWeakObjectPtr<UObject> WeakObject = ModuleInstance.Object;
+	TWeakObjectPtr<UObject> WeakObject = ModuleInstance.GetObject();
 	TArray<FTickFunction*> TickFunctions;
 	TickFunctions.Reserve(Dependencies.Num());
 	bool bTickFunctionFound = false; 
diff --git a/Engine/Plugins/Experimental/AnimNext/Source/AnimNext/Private/Module/RigUnit_AnimNextModuleEvents.cpp b/Engine/Plugins/Experimental/AnimNext/Source/AnimNext/Private/Module/RigUnit_AnimNextModuleEvents.cpp
--- a/Engine/Plugins/Experimental/AnimNext/Source/AnimNext/Private/Module/RigUnit_AnimNextModuleEvents.cpp
+++ b/Engine/Plugins/Experimental/AnimNext/Source/AnimNext/Private/Module/RigUnit_AnimNextModuleEvents.cpp
@@ -28,10 +28,11 @@ UE::AnimNext::FModuleEventBindingFunction FRigUnit_AnimNextExecuteBindings::GetB
 		InTickFunction.bRunOnAnyThread = false; // TODO: This will need to vary depending on whether the event is deemed thread-safe 
 
 		// Add prerequisite on this for all subsequent tick functions (assumes tick functions are sorted by phase)
-		for(int32 TickFunctionIndex = InContext.EventIndex + 1; TickFunctionIndex < InContext.ModuleInstance.TickFunctions.Num(); ++TickFunctionIndex)
+		TArrayView<UE::AnimNext::FModuleEventTickFunction> TickFunctions = InContext.ModuleInstance.GetTickFunctions();
+		for(int32 TickFunctionIndex = InContext.EventIndex + 1; TickFunctionIndex <TickFunctions.Num(); ++TickFunctionIndex)
 		{
-			FTickFunction& TickFunctions = InContext.ModuleInstance.TickFunctions[TickFunctionIndex];
-			TickFunctions.AddPrerequisite(InContext.ModuleInstance.Object, InContext.ModuleInstance.TickFunctions[InContext.EventIndex]);
+			FTickFunction& TickFunction = TickFunctions[TickFunctionIndex];
+			TickFunction.AddPrerequisite(InContext.ModuleInstance.GetObject(), TickFunctions[InContext.EventIndex]);
 		}
 
 		const bool bHighPriority = UE::AnimNext::CVars::CVarHighPriorityAnimNextExecuteBindingsTick.GetValueOnGameThread();
@@ -63,10 +64,11 @@ UE::AnimNext::FModuleEventBindingFunction FRigUnit_AnimNextInitializeEvent::GetB
 		InTickFunction.TickGroup = ETickingGroup::TG_PrePhysics;
 
 		// Add prerequisite on this for all subsequent tick functions (assumes tick functions are sorted by phase)
-		for(int32 TickFunctionIndex = InContext.EventIndex + 1; TickFunctionIndex < InContext.ModuleInstance.TickFunctions.Num(); ++TickFunctionIndex)
+		TArrayView<UE::AnimNext::FModuleEventTickFunction> TickFunctions = InContext.ModuleInstance.GetTickFunctions();
+		for(int32 TickFunctionIndex = InContext.EventIndex + 1; TickFunctionIndex < TickFunctions.Num(); ++TickFunctionIndex)
 		{
-			FTickFunction& TickFunction = InContext.ModuleInstance.TickFunctions[TickFunctionIndex];
-			TickFunction.AddPrerequisite(InContext.ModuleInstance.Object, InContext.ModuleInstance.TickFunctions[InContext.EventIndex]);
+			FTickFunction& TickFunction = TickFunctions[TickFunctionIndex];
+			TickFunction.AddPrerequisite(InContext.ModuleInstance.GetObject(), TickFunctions[InContext.EventIndex]);
 		}
 	};
 }
diff --git a/Engine/Plugins/Experimental/AnimNext/Source/AnimNextEditor/Internal/Common/AnimNextAssetItemDetails.cpp b/Engine/Plugins/Experimental/AnimNext/Source/AnimNextEditor/Internal/Common/AnimNextAssetItemDetails.cpp
--- a/Engine/Plugins/Experimental/AnimNext/Source/AnimNextEditor/Internal/Common/AnimNextAssetItemDetails.cpp
+++ b/Engine/Plugins/Experimental/AnimNext/Source/AnimNextEditor/Internal/Common/AnimNextAssetItemDetails.cpp
@@ -10,17 +10,22 @@
 #include "RigVMModel/RigVMGraph.h"
 #include "WorkspaceItemMenuContext.h"
 #include "IWorkspaceEditor.h"
+#include "PersonaModule.h"
+#include "PersonaUtils.h"
 #include "ScopedTransaction.h"
 #include "RigVMModel/RigVMClient.h"
 #include "ToolMenus.h"
 #include "AssetRegistry/AssetRegistryModule.h"
 #include "AssetRegistry/IAssetRegistry.h"
 #include "Entries/AnimNextDataInterfaceEntry.h"
+#include "Framework/Application/SlateApplication.h"
 #include "Framework/Commands/GenericCommands.h"
 #include "Module/AnimNextModule_EditorData.h"
 #include "Module/RigUnit_AnimNextModuleEvents.h"
+#include "RigVMFunctions/Execution/RigVMFunction_UserDefinedEvent.h"
 #include "UObject/UObjectIterator.h"
 #include "Variables/SAddVariablesDialog.h"
+#include "Widgets/Input/STextEntryPopup.h"
 
 #define LOCTEXT_NAMESPACE "FAnimNextGraphItemDetails"
 
@@ -106,43 +111,116 @@ void FAnimNextAssetItemDetails::RegisterToolMenuExtensions()
 					FText::Format(AddEntryTooltipFormat, SubEntryClass->GetDisplayNameText()),
 					FNewToolMenuDelegate::CreateLambda([SubEntryClass, WorkspaceItemContext, EditorData, Asset](UToolMenu* InToolMenu)
 					{
-						FToolMenuSection& Section = InToolMenu->AddSection(SubEntryClass->GetFName());
+						FToolMenuSection& EventsSection = InToolMenu->AddSection(SubEntryClass->GetFName());
 
 						static FTextFormat AddEventTooltipFormat(LOCTEXT("AddEventGraphTooltipFormat", "Adds a {0} event graph to this asset"));
-						
-						for(TObjectIterator<UScriptStruct> It; It; ++It)
+
+						auto AddEventMenuEntry = [&EventsSection, SubEntryClass, WorkspaceItemContext, EditorData, Asset](FName InEventName, const FText& InDisplayName, UScriptStruct* InStruct)
 						{
-							UScriptStruct* Struct = *It;
-							if(!Struct->IsChildOf(FRigUnit_AnimNextModuleEventBase::StaticStruct()) || Struct == FRigUnit_AnimNextModuleEventBase::StaticStruct())
-							{
-								continue;
-							}
-							if(Struct->HasMetaData(FRigVMStruct::HiddenMetaName) || Struct->HasMetaData(FRigVMStruct::AbstractMetaName))
-							{
-								continue;
-							}
-							
-							TInstancedStruct<FRigUnit_AnimNextModuleEventBase> StructInstance;
-							StructInstance.InitializeAsScriptStruct(Struct);
-							FName EventName = StructInstance.Get<FRigUnit_AnimNextModuleEventBase>().GetEventName();
-
-							Section.AddMenuEntry(
-								StructInstance.Get<FRigUnit_AnimNextModuleEventBase>().GetEventName(),
-								FText::FromName(EventName),
-								FText::Format(AddEventTooltipFormat, FText::FromName(EventName)),
+							EventsSection.AddMenuEntry(
+								InEventName,
+								FText::FromName(InEventName),
+								FText::Format(AddEventTooltipFormat, InDisplayName),
 								FSlateIconFinder::FindIconForClass(SubEntryClass, "ClassIcon.Object"),
 								FUIAction(
-									FExecuteAction::CreateWeakLambda(WorkspaceItemContext, [EventName, SubEntryClass, EditorData, Struct]()
+									FExecuteAction::CreateWeakLambda(WorkspaceItemContext, [InEventName, SubEntryClass, EditorData, InStruct]()
 									{
 										FScopedTransaction Transaction(LOCTEXT("AddEventGraph", "Add Event Graph"));
-										EditorData->AddEventGraph(EventName, Struct);
+										EditorData->AddEventGraph(InEventName, InStruct);
 									}),
-									FCanExecuteAction::CreateWeakLambda(WorkspaceItemContext, [Asset, EventName]()
+									FCanExecuteAction::CreateWeakLambda(WorkspaceItemContext, [Asset, InEventName]()
 									{
-										return !Asset->GetVM()->ContainsEntry(EventName);
+										return !Asset->GetVM()->ContainsEntry(InEventName);
 									})
 								));
+						};
+
+						if(EditorData->IsA<UAnimNextModule_EditorData>())
+						{
+							for(TObjectIterator<UScriptStruct> It; It; ++It)
+							{
+								UScriptStruct* Struct = *It;
+								if(!Struct->IsChildOf(FRigUnit_AnimNextModuleEventBase::StaticStruct()) || Struct == FRigUnit_AnimNextModuleEventBase::StaticStruct())
+								{
+									continue;
+								}
+								if(Struct->HasMetaData(FRigVMStruct::HiddenMetaName) || Struct->HasMetaData(FRigVMStruct::AbstractMetaName))
+								{
+									continue;
+								}
+								
+								TInstancedStruct<FRigUnit_AnimNextModuleEventBase> StructInstance;
+								StructInstance.InitializeAsScriptStruct(Struct);
+								FName EventName = StructInstance.Get<FRigUnit_AnimNextModuleEventBase>().GetEventName();
+								AddEventMenuEntry(EventName, FText::FromName(EventName), Struct);
+							}
 						}
+
+						AddEventMenuEntry(TEXT("CustomEvent"), LOCTEXT("CustomEventDisplayNameInline", "custom"), FRigVMFunction_UserDefinedEvent::StaticStruct());
+
+						EventsSection.AddSubMenu(
+							TEXT("NotifyEvents"),
+							LOCTEXT("NotifyEventsLabel", "Notify Event"),
+							LOCTEXT("NotifyEventsTooltip", "Add an event graph to handle a notify event"),
+							FNewToolMenuDelegate::CreateLambda([SubEntryClass, WorkspaceItemContext, EditorData, Asset](UToolMenu* InToolMenu)
+							{
+								FToolMenuSection& NotifySection = InToolMenu->AddSection(TEXT("NotifyEvents"));
+
+								auto CreateNewEventGraph = [Asset, EditorData](FName InEventGraphName)
+								{
+									if(!Asset->GetVM()->ContainsEntry(InEventGraphName))
+									{
+										FScopedTransaction Transaction(LOCTEXT("AddEventGraph", "Add Event Graph"));
+										EditorData->AddEventGraph(InEventGraphName, FRigVMFunction_UserDefinedEvent::StaticStruct());
+									}
+								};
+
+								NotifySection.AddEntry(FToolMenuEntry::InitMenuEntry(
+									TEXT("NewNotifyEvent"),
+									LOCTEXT("AddNewNotifyEventLabel", "Add New Notify Event..."),
+									LOCTEXT("AddNewNotifyEventTooltip", "Add a new notify event as a custom event"),
+									FSlateIcon(),
+									FExecuteAction::CreateLambda([CreateNewEventGraph]()
+									{ 
+										// Show dialog to enter new track name
+										TSharedRef<STextEntryPopup> TextEntry =
+											SNew(STextEntryPopup)
+											.Label(LOCTEXT("NewNotifyLabel", "Notify Name"))
+											.OnTextCommitted_Lambda([CreateNewEventGraph](const FText& InText, ETextCommit::Type InCommitType)
+											{
+												FSlateApplication::Get().DismissAllMenus();
+												FName NotifyName = *InText.ToString();
+												CreateNewEventGraph(NotifyName);
+											});
+
+										// Show dialog to enter new event name
+										FSlateApplication::Get().PushMenu(
+											FSlateApplication::Get().GetInteractiveTopLevelWindows()[0],
+											FWidgetPath(),
+											TextEntry,
+											FSlateApplication::Get().GetCursorPos(),
+											FPopupTransitionEffect( FPopupTransitionEffect::TypeInPopup));
+									}),
+									EUserInterfaceActionType::Button
+								));
+
+								FPersonaModule& PersonaModule = FModuleManager::LoadModuleChecked<FPersonaModule>("Persona");
+								TSharedRef<SWidget> PickerWidget = PersonaModule.CreateSkeletonNotifyPicker(FOnNotifyPicked::CreateLambda([CreateNewEventGraph](FName InNotifyName)
+								{
+									FSlateApplication::Get().DismissAllMenus();
+									CreateNewEventGraph(InNotifyName);
+								}));
+								NotifySection.AddEntry(FToolMenuEntry::InitWidget(
+									TEXT("NotifyEventPicker"),
+									SNew(SBox)
+									.WidthOverride(300.0f)
+									.HeightOverride(400.0f)
+									[
+										PickerWidget
+									],
+									FText::GetEmpty(),
+									true, false, true));
+							}));
 					}),
 					false,
 					FSlateIconFinder::FindIconForClass(SubEntryClass, "ClassIcon.Object"));
diff --git a/Engine/Plugins/Experimental/AnimNext/Source/AnimNextEditor/Private/AnimNextEditorModule.cpp b/Engine/Plugins/Experimental/AnimNext/Source/AnimNextEditor/Private/AnimNextEditorModule.cpp
--- a/Engine/Plugins/Experimental/AnimNext/Source/AnimNextEditor/Private/AnimNextEditorModule.cpp
+++ b/Engine/Plugins/Experimental/AnimNext/Source/AnimNextEditor/Private/AnimNextEditorModule.cpp
@@ -54,6 +54,7 @@
 #include "ExternalPackageHelper.h"
 #include "Common/GraphEditorSchemaActions.h"
 #include "EditorUtils.h"
+#include "PersonaModule.h"
 
 #define LOCTEXT_NAMESPACE "AnimNextEditorModule"
 
@@ -184,6 +185,12 @@ void FAnimNextEditorModule::StartupModule()
 			UAnimNextModule::StaticClass()->GetClassPathName(),
 			UAnimNextDataInterface::StaticClass()->GetClassPathName()
 		});
+
+	FPersonaModule& PersonaModule = FModuleManager::LoadModuleChecked<FPersonaModule>("Persona");
+	FPersonaModule::FNotifyHostAssetParameters NotifyHostParameters;
+	NotifyHostParameters.OnRemoveNotify = FPersonaModule::FNotifyHostAssetParameters::FOnRemoveNotify::CreateStatic(&UAnimNextRigVMAssetEditorData::HandleRemoveNotify);
+	NotifyHostParameters.OnReplaceNotify = FPersonaModule::FNotifyHostAssetParameters::FOnReplaceNotify::CreateStatic(&UAnimNextRigVMAssetEditorData::HandleReplaceNotify);
+	PersonaModule.RegisterNotifyHostAsset(FTopLevelAssetPath(TEXT("/Script/AnimNext.AnimNextRigVMAsset")), NotifyHostParameters);
 }
 
 void FAnimNextEditorModule::ShutdownModule()
@@ -232,6 +239,12 @@ void FAnimNextEditorModule::ShutdownModule()
 	UnregisterLocatorFragmentEditorType("AnimNextObjectCast");
 	UnregisterLocatorFragmentEditorType("AnimNextComponent");
 	UnregisterLocatorFragmentEditorType("AnimNextActor");
+
+	if(FModuleManager::Get().IsModuleLoaded("Persona"))
+	{
+		FPersonaModule& PersonaModule = FModuleManager::GetModuleChecked<FPersonaModule>("Persona");
+		PersonaModule.UnregisterNotifyHostAsset(FTopLevelAssetPath(TEXT("/Script/AnimNext.AnimNextRigVMAsset")));
+	}
 }
 
 void FAnimNextEditorModule::RegisterLocatorFragmentEditorType(FName InLocatorFragmentEditorName)
diff --git a/Engine/Plugins/Experimental/AnimNext/Source/AnimNextEditor/Private/Common/AnimNextGraphItemDetails.cpp b/Engine/Plugins/Experimental/AnimNext/Source/AnimNextEditor/Private/Common/AnimNextGraphItemDetails.cpp
--- a/Engine/Plugins/Experimental/AnimNext/Source/AnimNextEditor/Private/Common/AnimNextGraphItemDetails.cpp
+++ b/Engine/Plugins/Experimental/AnimNext/Source/AnimNextEditor/Private/Common/AnimNextGraphItemDetails.cpp
@@ -133,11 +133,13 @@ bool FAnimNextGraphItemDetails::ValidateName(const FWorkspaceOutlinerItemExport&
 		const FAnimNextGraphOutlinerData& GraphData = Export.GetData().Get<FAnimNextGraphOutlinerData>();
 		UAnimNextRigVMAssetEditorData* EditorData = GraphData.Entry->GetTypedOuter<UAnimNextRigVMAssetEditorData>();
 		UAnimNextRigVMAssetEntry* ExistingEntry = EditorData->FindEntry(FName(*InName.ToString()));
-		if(ExistingEntry)
+		if(ExistingEntry && GraphData.Entry != ExistingEntry)
 		{
 			OutErrorMessage = LOCTEXT("NameAlreadyExistsError", "Name already exists in this module");
 			return false;
 		}
+
+		return true;
 	}
 
 	OutErrorMessage = LOCTEXT("UnsupportedTypeRenameError", "Element type is not supported for rename");
diff --git a/Engine/Plugins/Experimental/AnimNext/Source/AnimNextUncookedOnly/Internal/AnimNextRigVMAssetEditorData.h b/Engine/Plugins/Experimental/AnimNext/Source/AnimNextUncookedOnly/Internal/AnimNextRigVMAssetEditorData.h
--- a/Engine/Plugins/Experimental/AnimNext/Source/AnimNextUncookedOnly/Internal/AnimNextRigVMAssetEditorData.h
+++ b/Engine/Plugins/Experimental/AnimNext/Source/AnimNextUncookedOnly/Internal/AnimNextRigVMAssetEditorData.h
@@ -446,6 +446,12 @@ protected:
 	// Remove any programmatic graphs generated during compilation and consign them to the transient package
 	void RemoveProgrammaticGraphs(TArrayView<URigVMGraph*> InGraphs);
 
+	// Handle removing a notify
+	static void HandleRemoveNotify(UObject* InAsset, const FString& InFindString, bool bFindWholeWord, ESearchCase::Type InSearchCase);
+
+	// Handle replacing a notify
+	static void HandleReplaceNotify(UObject* InAsset, const FString& InFindString, const FString& InReplaceString, bool bFindWholeWord, ESearchCase::Type InSearchCase);
+
 	/** All entries in this asset - not saved, discovered at load time and also contains InternalEntries */
 	UPROPERTY(transient)
 	TArray<TObjectPtr<UAnimNextRigVMAssetEntry>> Entries;
diff --git a/Engine/Plugins/Experimental/AnimNext/Source/AnimNextUncookedOnly/Internal/Module/AnimNextEventGraphSchema.h b/Engine/Plugins/Experimental/AnimNext/Source/AnimNextUncookedOnly/Internal/Module/AnimNextEventGraphSchema.h
new file mode 100644
--- /dev/null
+++ b/Engine/Plugins/Experimental/AnimNext/Source/AnimNextUncookedOnly/Internal/Module/AnimNextEventGraphSchema.h
@@ -0,0 +1,17 @@
+﻿// Copyright Epic Games, Inc. All Rights Reserved.
+
+#pragma once
+
+#include "CoreMinimal.h"
+#include "AnimNextRigVMAssetSchema.h"
+#include "AnimNextEventGraphSchema.generated.h"
+
+UCLASS(MinimalAPI)
+class UAnimNextEventGraphSchema : public UAnimNextRigVMAssetSchema
+{
+	GENERATED_BODY()
+
+	// URigVMSchema interface
+	virtual bool SupportsUnitFunction(URigVMController* InController, const FRigVMFunction* InUnitFunction) const override;
+};
+
diff --git a/Engine/Plugins/Experimental/AnimNext/Source/AnimNextUncookedOnly/Internal/UncookedOnlyUtils.h b/Engine/Plugins/Experimental/AnimNext/Source/AnimNextUncookedOnly/Internal/UncookedOnlyUtils.h
--- a/Engine/Plugins/Experimental/AnimNext/Source/AnimNextUncookedOnly/Internal/UncookedOnlyUtils.h
+++ b/Engine/Plugins/Experimental/AnimNext/Source/AnimNextUncookedOnly/Internal/UncookedOnlyUtils.h
@@ -197,7 +197,7 @@ struct ANIMNEXTUNCOOKEDONLY_API FUtils
 	static FRigVMTemplateArgumentType GetRigVMArgTypeFromParamType(const FAnimNextParamType& InParamType);
 
 	/** Set up a simple event graph */
-	static void SetupEventGraph(URigVMController* InController, UScriptStruct* InEventStruct);
+	static void SetupEventGraph(URigVMController* InController, UScriptStruct* InEventStruct, FName InEventName = NAME_None);
 
 	// Gets the variables that are exported to the asset registry for an asset
 	static bool GetExportedVariablesForAsset(const FAssetData& InAsset, FAnimNextAssetRegistryExports& OutExports);
diff --git a/Engine/Plugins/Experimental/AnimNext/Source/AnimNextUncookedOnly/Private/AnimNextRigVMAssetEditorData.cpp b/Engine/Plugins/Experimental/AnimNext/Source/AnimNextUncookedOnly/Private/AnimNextRigVMAssetEditorData.cpp
--- a/Engine/Plugins/Experimental/AnimNext/Source/AnimNextUncookedOnly/Private/AnimNextRigVMAssetEditorData.cpp
+++ b/Engine/Plugins/Experimental/AnimNext/Source/AnimNextUncookedOnly/Private/AnimNextRigVMAssetEditorData.cpp
@@ -14,6 +14,7 @@
 #include "IWorkspaceEditorModule.h"
 #include "ObjectTools.h"
 #include "UncookedOnlyUtils.h"
+#include "Animation/Skeleton.h"
 #include "DataInterface/AnimNextDataInterface_EditorData.h"
 #include "Entries/AnimNextDataInterfaceEntry.h"
 #include "Entries/AnimNextEventGraphEntry.h"
@@ -28,6 +29,7 @@
 #include "RigVMModel/Nodes/RigVMCollapseNode.h"
 #include "UObject/AssetRegistryTagsContext.h"
 #include "DataInterface/AnimNextDataInterface_EditorData.h"
+#include "RigVMFunctions/Execution/RigVMFunction_UserDefinedEvent.h"
 #include "UObject/SavePackage.h"
 
 #if WITH_EDITOR
@@ -278,6 +280,16 @@ void UAnimNextRigVMAssetEditorData::GetAssetRegistryTags(FAssetRegistryTagsConte
 		HeadersProperty->ExportText_Direct(TagValue, &(FunctionExports.Headers), &(FunctionExports.Headers), nullptr, PPF_None, nullptr);
 		Context.AddTag(FAssetRegistryTag(UE::AnimNext::AnimNextPublicGraphFunctionsExportsRegistryTag, TagValue, FAssetRegistryTag::TT_Hidden));
 	}
+
+	{
+		// Export user defined events as notifies
+		FString NotifyList = USkeleton::AnimNotifyTagDelimiter;
+		for(FName EventName : RigVMClient.GetEntryNames(FRigVMFunction_UserDefinedEvent::StaticStruct()))
+		{
+			NotifyList += FString::Printf(TEXT("%s%s"), *EventName.ToString(), *USkeleton::AnimNotifyTagDelimiter);
+		}
+		Context.AddTag(FAssetRegistryTag(USkeleton::AnimNotifyTag, NotifyList, FAssetRegistryTag::TT_Hidden));
+	}
 }
 
 bool UAnimNextRigVMAssetEditorData::Rename(const TCHAR* NewName, UObject* NewOuter, ERenameFlags Flags)
@@ -855,6 +867,87 @@ void UAnimNextRigVMAssetEditorData::RemoveProgrammaticGraphs(TArrayView<URigVMGr
 	}
 }
 
+void UAnimNextRigVMAssetEditorData::HandleRemoveNotify(UObject* InAsset, const FString& InFindString, bool bFindWholeWord, ESearchCase::Type InSearchCase)
+{
+	UAnimNextRigVMAsset* Asset = Cast<UAnimNextRigVMAsset>(InAsset);
+	if(Asset == nullptr)
+	{
+		return;
+	}
+
+	UAnimNextRigVMAssetEditorData* EditorData = UE::AnimNext::UncookedOnly::FUtils::GetEditorData<UAnimNextRigVMAssetEditorData>(Asset);
+	if(EditorData == nullptr)
+	{
+		return;
+	}
+
+	URigVMController* Controller = EditorData->GetController();
+	Controller->OpenUndoBracket(LOCTEXT("RemoveNotifyEvents", "Remove Notify Events").ToString());
+
+	for(TObjectPtr<URigVMGraph> Model : EditorData->RigVMClient.GetModels())
+	{
+		for(URigVMNode* Node : Model->GetNodes())
+		{
+			if(URigVMUnitNode* UnitNode = Cast<URigVMUnitNode>(Node))
+			{
+				if(UnitNode->GetScriptStruct()->IsChildOf(FRigVMFunction_UserDefinedEvent::StaticStruct()))
+				{
+					URigVMPin* Pin = UnitNode->FindPin(GET_MEMBER_NAME_STRING_CHECKED(FRigVMFunction_UserDefinedEvent, EventName));
+					FString EventNameString = Pin->GetDefaultValue();
+					if( (bFindWholeWord && EventNameString.Equals(InFindString, InSearchCase)) ||
+						(!bFindWholeWord && EventNameString.Contains(InFindString, InSearchCase)))
+					{
+						Controller->RemoveNode(Node);
+					}
+				}
+			}
+		}
+	}
+
+	Controller->CloseUndoBracket();
+}
+
+void UAnimNextRigVMAssetEditorData::HandleReplaceNotify(UObject* InAsset, const FString& InFindString, const FString& InReplaceString, bool bFindWholeWord, ESearchCase::Type InSearchCase)
+{
+	UAnimNextRigVMAsset* Asset = Cast<UAnimNextRigVMAsset>(InAsset);
+	if(Asset == nullptr)
+	{
+		return;
+	}
+
+	UAnimNextRigVMAssetEditorData* EditorData = UE::AnimNext::UncookedOnly::FUtils::GetEditorData<UAnimNextRigVMAssetEditorData>(Asset);
+	if(EditorData == nullptr)
+	{
+		return;
+	}
+
+	URigVMController* Controller = EditorData->GetController();
+	Controller->OpenUndoBracket(LOCTEXT("ReplaceNotifyEvents", "Replace Notify Events").ToString());
+
+	for(TObjectPtr<URigVMGraph> Model : EditorData->RigVMClient.GetModels())
+	{
+		for(URigVMNode* Node : Model->GetNodes())
+		{
+			if(URigVMUnitNode* UnitNode = Cast<URigVMUnitNode>(Node))
+			{
+				if(UnitNode->GetScriptStruct()->IsChildOf(FRigVMFunction_UserDefinedEvent::StaticStruct()))
+				{
+					URigVMPin* Pin = UnitNode->FindPin(GET_MEMBER_NAME_STRING_CHECKED(FRigVMFunction_UserDefinedEvent, EventName));
+					FString EventNameString = Pin->GetDefaultValue();
+					if( (bFindWholeWord && EventNameString.Equals(InFindString, InSearchCase)) ||
+						(!bFindWholeWord && EventNameString.Contains(InFindString, InSearchCase)))
+					{
+						const FString NewName = EventNameString.Replace(*InFindString, *InReplaceString, InSearchCase);
+						Controller->SetPinDefaultValue(Pin, NewName, true, true, false);
+					}
+				}
+			}
+		}
+	}
+
+	Controller->CloseUndoBracket();
+}
+
 void UAnimNextRigVMAssetEditorData::RecompileVMIfRequired()
 {
 	if (bVMRecompilationRequired)
@@ -1613,7 +1706,7 @@ UAnimNextEventGraphEntry* UAnimNextRigVMAssetEditorData::AddEventGraph(FName InN
 			RefreshExternalModels();
 			RigVMClient.AddModel(NewRigVMGraphModel, true);
 			URigVMController* Controller = RigVMClient.GetController(NewRigVMGraphModel);
-			UE::AnimNext::UncookedOnly::FUtils::SetupEventGraph(Controller, InEventStruct);
+			UE::AnimNext::UncookedOnly::FUtils::SetupEventGraph(Controller, InEventStruct, InName);
 		}
 	}
 
diff --git a/Engine/Plugins/Experimental/AnimNext/Source/AnimNextUncookedOnly/Private/Module/AnimNextEventGraphSchema.h b/Engine/Plugins/Experimental/AnimNext/Source/AnimNextUncookedOnly/Private/Module/AnimNextEventGraphSchema.h
deleted file mode 100644
diff --git a/Engine/Plugins/Experimental/AnimNext/Source/AnimNextUncookedOnly/Private/Module/AnimNextModule_EventGraphSchema.cpp b/Engine/Plugins/Experimental/AnimNext/Source/AnimNextUncookedOnly/Private/Module/AnimNextModule_EventGraphSchema.cpp
--- a/Engine/Plugins/Experimental/AnimNext/Source/AnimNextUncookedOnly/Private/Module/AnimNextModule_EventGraphSchema.cpp
+++ b/Engine/Plugins/Experimental/AnimNext/Source/AnimNextUncookedOnly/Private/Module/AnimNextModule_EventGraphSchema.cpp
@@ -1,6 +1,6 @@
 ﻿// Copyright Epic Games, Inc. All Rights Reserved.
 
-#include "AnimNextEventGraphSchema.h"
+#include "Module/AnimNextEventGraphSchema.h"
 #include "AnimNextExecuteContext.h"
 
 bool UAnimNextEventGraphSchema::SupportsUnitFunction(URigVMController* InController, const FRigVMFunction* InUnitFunction) const
diff --git a/Engine/Plugins/Experimental/AnimNext/Source/AnimNextUncookedOnly/Private/UncookedOnlyUtils.cpp b/Engine/Plugins/Experimental/AnimNext/Source/AnimNextUncookedOnly/Private/UncookedOnlyUtils.cpp
--- a/Engine/Plugins/Experimental/AnimNext/Source/AnimNextUncookedOnly/Private/UncookedOnlyUtils.cpp
+++ b/Engine/Plugins/Experimental/AnimNext/Source/AnimNextUncookedOnly/Private/UncookedOnlyUtils.cpp
@@ -27,6 +27,7 @@
 #include "DataInterface/AnimNextDataInterface.h"
 #include "DataInterface/AnimNextDataInterface_EditorData.h"
 #include "Logging/MessageLog.h"
+#include "RigVMFunctions/Execution/RigVMFunction_UserDefinedEvent.h"
 #include "UObject/AssetRegistryTagsContext.h"
 
 #define LOCTEXT_NAMESPACE "AnimNextUncookedOnlyUtils"
@@ -666,13 +667,19 @@ FRigVMTemplateArgumentType FUtils::GetRigVMArgTypeFromParamType(const FAnimNextP
 	return ArgType;
 }
 
-void FUtils::SetupEventGraph(URigVMController* InController, UScriptStruct* InEventStruct)
+void FUtils::SetupEventGraph(URigVMController* InController, UScriptStruct* InEventStruct, FName InEventName)
 {
 	// Clear the graph
 	InController->RemoveNodes(InController->GetGraph()->GetNodes());
 
 	// Add entry point
-	InController->AddUnitNode(InEventStruct, FRigVMStruct::ExecuteName, FVector2D(-200.0f, 0.0f), FString(), false);
+	URigVMUnitNode* EventNode = InController->AddUnitNode(InEventStruct, FRigVMStruct::ExecuteName, FVector2D(-200.0f, 0.0f), FString(), false);
+
+	// Setup custom event name if required
+	if(EventNode && InEventStruct == FRigVMFunction_UserDefinedEvent::StaticStruct() && InEventName != NAME_None && !InController->GetAllEventNames().Contains(InEventName))
+	{
+		InController->SetPinDefaultValue(EventNode->FindPin(GET_MEMBER_NAME_STRING_CHECKED(FRigVMFunction_UserDefinedEvent, EventName)), InEventName.ToString(), true, true, false);
+	}
 }
 
 FAnimNextParamType FUtils::GetParameterTypeFromName(FName InName)
diff --git a/Engine/Plugins/Experimental/AnimNextAnimGraph/Source/AnimNextAnimGraph/Private/Graph/RigUnit_AnimNextSetNotifyContext.cpp b/Engine/Plugins/Experimental/AnimNextAnimGraph/Source/AnimNextAnimGraph/Private/Graph/RigUnit_AnimNextSetNotifyContext.cpp
new file mode 100644
--- /dev/null
+++ b/Engine/Plugins/Experimental/AnimNextAnimGraph/Source/AnimNextAnimGraph/Private/Graph/RigUnit_AnimNextSetNotifyContext.cpp
@@ -0,0 +1,11 @@
+﻿// Copyright Epic Games, Inc. All Rights Reserved.
+
+#include "RigUnit_AnimNextSetNotifyContext.h"
+#include "Traits/NotifyDispatcher.h"
+
+FRigUnit_AnimNextSetNotifyContext_Execute()
+{
+	FAnimNextModuleInstance& ModuleInstance = ExecuteContext.GetContextData<FAnimNextModuleContextData>().GetModuleInstance();
+	FModuleInstanceComponent_NotifyDispatcher& NotifyDispatcher = ModuleInstance.GetComponent<FModuleInstanceComponent_NotifyDispatcher>();
+	NotifyDispatcher.SkeletalMeshComponent = SkeletalMeshComponent;
+}
\ No newline at end of file
diff --git a/Engine/Plugins/Experimental/AnimNextAnimGraph/Source/AnimNextAnimGraph/Private/Graph/RigUnit_AnimNextSetNotifyContext.h b/Engine/Plugins/Experimental/AnimNextAnimGraph/Source/AnimNextAnimGraph/Private/Graph/RigUnit_AnimNextSetNotifyContext.h
new file mode 100644
--- /dev/null
+++ b/Engine/Plugins/Experimental/AnimNextAnimGraph/Source/AnimNextAnimGraph/Private/Graph/RigUnit_AnimNextSetNotifyContext.h
@@ -0,0 +1,24 @@
+﻿// Copyright Epic Games, Inc. All Rights Reserved.
+
+#pragma once
+
+#include "Graph/RigUnit_AnimNextBase.h"
+#include "AnimNextExecuteContext.h"
+#include "RigUnit_AnimNextSetNotifyContext.generated.h"
+
+/** Sets up the context in which notifies are called */
+USTRUCT(meta=(DisplayName="Set Notify Context", Category="Notifies", NodeColor="0, 1, 1", Keywords="Event,Output"))
+struct ANIMNEXTANIMGRAPH_API FRigUnit_AnimNextSetNotifyContext : public FRigUnit_AnimNextBase
+{
+	GENERATED_BODY()
+
+	RIGVM_METHOD()
+	void Execute();
+
+	// The mesh component to supply to any notify dispatches
+	UPROPERTY(meta = (Input))
+	TObjectPtr<USkeletalMeshComponent> SkeletalMeshComponent;
+
+	UPROPERTY(meta = (Input, Output))
+	FAnimNextExecuteContext ExecuteContext;
+};
diff --git a/Engine/Plugins/Experimental/AnimNextAnimGraph/Source/AnimNextAnimGraph/Private/TraitCore/ExecutionContext.cpp b/Engine/Plugins/Experimental/AnimNextAnimGraph/Source/AnimNextAnimGraph/Private/TraitCore/ExecutionContext.cpp
--- a/Engine/Plugins/Experimental/AnimNextAnimGraph/Source/AnimNextAnimGraph/Private/TraitCore/ExecutionContext.cpp
+++ b/Engine/Plugins/Experimental/AnimNextAnimGraph/Source/AnimNextAnimGraph/Private/TraitCore/ExecutionContext.cpp
@@ -600,7 +600,7 @@ namespace UE::AnimNext
 		const FAnimNextModuleInstance* ModuleInstance = GetRootGraphInstance().GetModuleInstance();
 		if (ModuleInstance)
 		{
-			return ModuleInstance->Object;
+			return ModuleInstance->GetObject();
 		}
 
 		return nullptr;
diff --git a/Engine/Plugins/Experimental/AnimNextAnimGraph/Source/AnimNextAnimGraph/Private/TraitInterfaces/INotifySource.cpp b/Engine/Plugins/Experimental/AnimNextAnimGraph/Source/AnimNextAnimGraph/Private/TraitInterfaces/INotifySource.cpp
new file mode 100644
--- /dev/null
+++ b/Engine/Plugins/Experimental/AnimNextAnimGraph/Source/AnimNextAnimGraph/Private/TraitInterfaces/INotifySource.cpp
@@ -0,0 +1,33 @@
+﻿// Copyright Epic Games, Inc. All Rights Reserved.
+
+#include "TraitInterfaces/INotifySource.h"
+#include "TraitCore/ExecutionContext.h"
+
+namespace UE::AnimNext
+{
+
+AUTO_REGISTER_ANIM_TRAIT_INTERFACE(INotifySource)
+
+#if WITH_EDITOR
+const FText& INotifySource::GetDisplayName() const
+{
+	static FText InterfaceName = NSLOCTEXT("TraitInterfaces", "TraitInterface_INotifySource_Name", "Notify Source");
+	return InterfaceName;
+}
+const FText& INotifySource::GetDisplayShortName() const
+{
+	static FText InterfaceShortName = NSLOCTEXT("TraitInterfaces", "TraitInterface_INotifySource_ShortName", "NOT");
+	return InterfaceShortName;
+}
+#endif // WITH_EDITOR
+
+void INotifySource::GetNotifies(FExecutionContext& Context, const TTraitBinding<INotifySource>& Binding, float StartPosition, float Duration, bool bLooping, TArray<FAnimNotifyEventReference>& OutNotifies) const
+{
+	TTraitBinding<INotifySource> SuperBinding;
+	if (Binding.GetStackInterfaceSuper(SuperBinding))
+	{
+		return SuperBinding.GetNotifies(Context, StartPosition, Duration, bLooping, OutNotifies);
+	}
+}
+
+}
diff --git a/Engine/Plugins/Experimental/AnimNextAnimGraph/Source/AnimNextAnimGraph/Private/TraitInterfaces/ITimeline.cpp b/Engine/Plugins/Experimental/AnimNextAnimGraph/Source/AnimNextAnimGraph/Private/TraitInterfaces/ITimeline.cpp
--- a/Engine/Plugins/Experimental/AnimNextAnimGraph/Source/AnimNextAnimGraph/Private/TraitInterfaces/ITimeline.cpp
+++ b/Engine/Plugins/Experimental/AnimNextAnimGraph/Source/AnimNextAnimGraph/Private/TraitInterfaces/ITimeline.cpp
@@ -49,4 +49,15 @@ namespace UE::AnimNext
 			SuperBinding.AdvanceBy(Context, DeltaTime);
 		}
 	}
+
+	FTimelineDelta ITimeline::GetDelta(FExecutionContext& Context, const TTraitBinding<ITimeline>& Binding) const
+	{
+		TTraitBinding<ITimeline> SuperBinding;
+		if (Binding.GetStackInterfaceSuper(SuperBinding))
+		{
+			return SuperBinding.GetDelta(Context);
+		}
+
+		return FTimelineDelta();
+	}
 }
diff --git a/Engine/Plugins/Experimental/AnimNextAnimGraph/Source/AnimNextAnimGraph/Private/Traits/NotifyDispatcher.cpp b/Engine/Plugins/Experimental/AnimNextAnimGraph/Source/AnimNextAnimGraph/Private/Traits/NotifyDispatcher.cpp
new file mode 100644
--- /dev/null
+++ b/Engine/Plugins/Experimental/AnimNextAnimGraph/Source/AnimNextAnimGraph/Private/Traits/NotifyDispatcher.cpp
@@ -0,0 +1,264 @@
+﻿// Copyright Epic Games, Inc. All Rights Reserved.
+
+#include "Traits/NotifyDispatcher.h"
+#include "TraitInterfaces/INotifySource.h"
+#include "AnimationRuntime.h"
+#include "Animation/AnimTrace.h"
+#include "Animation/AnimNotifies/AnimNotifyState.h"
+#include "Components/SkeletalMeshComponent.h"
+#include "Graph/AnimNextGraphInstance.h"
+#include "Module/AnimNextModuleInstance.h"
+
+namespace UE::AnimNext
+{
+	AUTO_REGISTER_ANIM_TRAIT(FNotifyDispatcherTrait)
+
+	// Trait implementation boilerplate
+	#define TRAIT_INTERFACE_ENUMERATOR(GeneratorMacro) \
+		GeneratorMacro(ITimeline) \
+
+		// Trait required interfaces implementation boilerplate
+	#define TRAIT_REQUIRED_INTERFACE_ENUMERATOR(GeneratorMacroRequired) \
+		GeneratorMacroRequired(ITimeline) \
+		GeneratorMacroRequired(INotifySource) \
+
+	GENERATE_ANIM_TRAIT_IMPLEMENTATION(FNotifyDispatcherTrait, TRAIT_INTERFACE_ENUMERATOR, TRAIT_REQUIRED_INTERFACE_ENUMERATOR, NULL_ANIM_TRAIT_EVENT_ENUMERATOR)
+	#undef TRAIT_INTERFACE_ENUMERATOR
+
+	void FNotifyDispatcherTrait::AdvanceBy(FExecutionContext& Context, const TTraitBinding<ITimeline>& Binding, float DeltaTime) const
+	{
+		TTraitBinding<ITimeline> TimelineTrait;
+		ensure(Binding.GetStackInterfaceSuper(TimelineTrait));
+
+		TTraitBinding<INotifySource> NotifySourceTrait;
+		ensure(Binding.GetStackInterfaceSuper(NotifySourceTrait));
+
+		// Get current state from the stack, advance time, then get the delta state
+		const FTimelineState PreAdvanceState = TimelineTrait.GetState(Context);
+		TimelineTrait.AdvanceBy(Context, DeltaTime);
+		const FTimelineDelta Delta = TimelineTrait.GetDelta(Context);
+
+		// Query for notifies
+		TArray<FAnimNotifyEventReference> Notifies;
+		NotifySourceTrait.GetNotifies(Context, PreAdvanceState.GetPosition(), Delta.GetDeltaTime(), PreAdvanceState.IsLooping(), Notifies);
+
+		// Dispatch if we found any
+		if(Notifies.Num() > 0)
+		{
+			// Ensure we have a handler component on the module
+			if(FAnimNextModuleInstance* ModuleInstance = Context.GetRootGraphInstance().GetModuleInstance())
+			{
+				ModuleInstance->GetComponent<FModuleInstanceComponent_NotifyDispatcher>();
+			}
+
+			auto NotifyDispatchEvent = MakeTraitEvent<FNotifyDispatchEvent>();
+			NotifyDispatchEvent->Notifies = MoveTemp(Notifies);
+			NotifyDispatchEvent->Weight = 1.0f;
+			Context.RaiseOutputTraitEvent(NotifyDispatchEvent);
+		}
+	}
+
+	template<typename TriggerFunc>
+	static void TriggerEventFunc(const TArray<FAnimNotifyEventReference>& InEvents, TriggerFunc InTriggerFunc)
+	{
+		for(const FAnimNotifyEventReference& EventReference : InEvents)
+		{
+			const FAnimNotifyEvent* AnimNotifyEvent = EventReference.GetNotify();
+			if(AnimNotifyEvent == nullptr)
+			{
+				continue;
+			}
+
+			InTriggerFunc(AnimNotifyEvent, EventReference);
+		}
+	}
+
+	void FNotifyQueueDispatchEvent::Execute() const
+	{
+		USkeletalMeshComponent* SkeletalMeshComponent = WeakSkeletalMeshComponent.Get();
+
+		// Order of dispatch here mimics that of UAnimInstance::TriggerAnimNotifies
+
+		// Triggers...
+		TriggerEventFunc(EventsToNotify, [SkeletalMeshComponent](const FAnimNotifyEvent* InAnimNotifyEvent, const FAnimNotifyEventReference& InEventReference)
+		{
+			check(InAnimNotifyEvent->Notify);
+			InAnimNotifyEvent->Notify->Notify(SkeletalMeshComponent, Cast<UAnimSequenceBase>(InAnimNotifyEvent->Notify->GetOuter()), InEventReference);
+		});
+
+		// Ends...
+		TriggerEventFunc(EventsToEnd, [SkeletalMeshComponent](const FAnimNotifyEvent* InAnimNotifyEvent, const FAnimNotifyEventReference& InEventReference)
+		{
+			check(InAnimNotifyEvent->NotifyStateClass);
+			InAnimNotifyEvent->NotifyStateClass->NotifyEnd(SkeletalMeshComponent, Cast<UAnimSequenceBase>(InAnimNotifyEvent->NotifyStateClass->GetOuter()), InEventReference);
+		});
+
+		// Begins...
+		TriggerEventFunc(EventsToBegin, [SkeletalMeshComponent](const FAnimNotifyEvent* InAnimNotifyEvent, const FAnimNotifyEventReference& InEventReference)
+		{
+			check(InAnimNotifyEvent->NotifyStateClass);
+			InAnimNotifyEvent->NotifyStateClass->NotifyBegin(SkeletalMeshComponent, Cast<UAnimSequenceBase>(InAnimNotifyEvent->NotifyStateClass->GetOuter()), InAnimNotifyEvent->GetDuration(), InEventReference);
+		});
+
+		// Ticks...
+		TriggerEventFunc(EventsToTick, [SkeletalMeshComponent, DeltaSeconds = DeltaSeconds](const FAnimNotifyEvent* InAnimNotifyEvent, const FAnimNotifyEventReference& InEventReference)
+		{
+			check(InAnimNotifyEvent->NotifyStateClass);
+			InAnimNotifyEvent->NotifyStateClass->NotifyTick(SkeletalMeshComponent, Cast<UAnimSequenceBase>(InAnimNotifyEvent->NotifyStateClass->GetOuter()), DeltaSeconds, InEventReference);
+		});
+	}
+}
+
+void FModuleInstanceComponent_NotifyDispatcher::OnTraitEvent(FAnimNextTraitEvent& Event)
+{
+	if(UE::AnimNext::FNotifyDispatchEvent* NotifyDispatchEvent = Event.AsType<UE::AnimNext::FNotifyDispatchEvent>())
+	{
+		NotifyQueue.AddAnimNotifies(NotifyDispatchEvent->Notifies, NotifyDispatchEvent->Weight);
+	}
+}
+
+void FModuleInstanceComponent_NotifyDispatcher::OnEndExecution(float InDeltaTime)
+{
+	using namespace UE::AnimNext;
+	
+	if(NotifyQueue.AnimNotifies.Num() == 0)
+	{
+		return;
+	}
+
+	FAnimNextModuleInstance& ModuleInstance = GetModuleInstance();
+
+	TSharedPtr<FNotifyQueueDispatchEvent> NotifyQueueDispatchEvent = MakeTraitEvent<FNotifyQueueDispatchEvent>();
+	NotifyQueueDispatchEvent->WeakSkeletalMeshComponent = SkeletalMeshComponent;
+	NotifyQueueDispatchEvent->DeltaSeconds = InDeltaTime;
+
+	NotifyQueue.PredictedLODLevel = SkeletalMeshComponent ? SkeletalMeshComponent->GetPredictedLODLevel() : 0;
+
+	// Logic from here largely mimics that of UAnimInstance::TriggerAnimNotifies 
+
+	// Array that will replace the 'ActiveAnimNotifyState' at the end of this function.
+	TArray<FAnimNotifyEvent> NewActiveAnimNotifyState;
+	NewActiveAnimNotifyState.Reserve(NotifyQueue.AnimNotifies.Num());
+
+	TArray<FAnimNotifyEventReference> NewActiveAnimNotifyEventReference;
+	NewActiveAnimNotifyEventReference.Reserve(NotifyQueue.AnimNotifies.Num());
+
+	// AnimNotifyState freshly added that need their 'NotifyBegin' event called.
+	TArray<const FAnimNotifyEvent*> NotifyStateBeginEvent;
+	TArray<const FAnimNotifyEventReference*> NotifyStateBeginEventReference;
+
+	for (int32 Index= 0 ; Index < NotifyQueue.AnimNotifies.Num(); Index++)
+	{
+		if(const FAnimNotifyEvent* AnimNotifyEvent = NotifyQueue.AnimNotifies[Index].GetNotify())
+		{
+			// AnimNotifyState
+			if (AnimNotifyEvent->NotifyStateClass)
+			{
+				int32 ExistingItemIndex = INDEX_NONE;
+
+				if (ActiveAnimNotifyState.Find(*AnimNotifyEvent, ExistingItemIndex))
+				{
+					check(ActiveAnimNotifyState.Num() == ActiveAnimNotifyEventReference.Num());
+					ActiveAnimNotifyState.RemoveAtSwap(ExistingItemIndex, EAllowShrinking::No);
+					ActiveAnimNotifyEventReference.RemoveAtSwap(ExistingItemIndex, EAllowShrinking::No);
+				}
+				else
+				{
+					NotifyStateBeginEvent.Add(AnimNotifyEvent);
+					NotifyStateBeginEventReference.Add(&NotifyQueue.AnimNotifies[Index]);
+				}
+
+				NewActiveAnimNotifyState.Add(*AnimNotifyEvent);
+				FAnimNotifyEventReference& EventRef = NewActiveAnimNotifyEventReference.Add_GetRef(NotifyQueue.AnimNotifies[Index]);
+				EventRef.SetNotify(&NewActiveAnimNotifyState.Top());
+				continue;
+			}
+
+			// Trigger non 'state' AnimNotifies
+			TriggerSingleAnimNotify(InDeltaTime, *NotifyQueueDispatchEvent.Get(), NotifyQueue.AnimNotifies[Index]);
+		}
+	}
+
+	// Send end notification to AnimNotifyState not active anymore.
+	for (int32 Index = 0; Index < ActiveAnimNotifyState.Num(); ++Index)
+	{
+		const FAnimNotifyEvent& AnimNotifyEvent = ActiveAnimNotifyState[Index];
+		const FAnimNotifyEventReference& EventReference = ActiveAnimNotifyEventReference[Index];
+		if (AnimNotifyEvent.NotifyStateClass)
+		{
+#if WITH_EDITOR
+			// Prevent firing notifies in animation editors if requested
+			if(ModuleInstance.GetWorldType() != EWorldType::EditorPreview || AnimNotifyEvent.NotifyStateClass->ShouldFireInEditor())
+#endif
+			{
+				NotifyQueueDispatchEvent->EventsToEnd.Add(EventReference);
+			}
+		}
+	}
+
+	check(NotifyStateBeginEventReference.Num() == NotifyStateBeginEvent.Num());
+	for (int32 Index = 0; Index < NotifyStateBeginEvent.Num(); Index++)
+	{
+		const FAnimNotifyEvent* AnimNotifyEvent = NotifyStateBeginEvent[Index];
+		const FAnimNotifyEventReference* EventReference = NotifyStateBeginEventReference[Index];
+		if(AnimNotifyEvent->NotifyStateClass)
+		{
+#if WITH_EDITOR
+			// Prevent firing notifies in animation editors if requested 
+			if(ModuleInstance.GetWorldType() != EWorldType::EditorPreview || AnimNotifyEvent->NotifyStateClass->ShouldFireInEditor())
+#endif
+			{
+				NotifyQueueDispatchEvent->EventsToBegin.Add(*EventReference);
+			}
+		}
+	}
+
+	// Switch our arrays.
+	ActiveAnimNotifyState = MoveTemp(NewActiveAnimNotifyState);
+	ActiveAnimNotifyEventReference = MoveTemp(NewActiveAnimNotifyEventReference);
+
+	// Tick currently active AnimNotifyState
+	for (int32 Index = 0; Index < ActiveAnimNotifyState.Num(); Index++)
+	{
+		const FAnimNotifyEvent& AnimNotifyEvent = ActiveAnimNotifyState[Index];
+		const FAnimNotifyEventReference& EventReference = ActiveAnimNotifyEventReference[Index];
+		if(AnimNotifyEvent.NotifyStateClass)
+		{
+#if WITH_EDITOR
+			// Prevent firing notifies in animation editors if requested 
+			if(ModuleInstance.GetWorldType() != EWorldType::EditorPreview || AnimNotifyEvent.NotifyStateClass->ShouldFireInEditor())
+#endif
+			{
+				NotifyQueueDispatchEvent->EventsToTick.Add(EventReference);
+			}
+		}
+	}
+
+	GetModuleInstance().QueueOutputTraitEvent(NotifyQueueDispatchEvent);
+	NotifyQueue.AnimNotifies.Reset();
+}
+
+void FModuleInstanceComponent_NotifyDispatcher::TriggerSingleAnimNotify(float InDeltaTime, UE::AnimNext::FNotifyQueueDispatchEvent& Dispatcher, const FAnimNotifyEventReference& EventReference)
+{
+	// This is for non 'state' anim notifies.
+	const FAnimNotifyEvent* AnimNotifyEvent = EventReference.GetNotify();
+	if (AnimNotifyEvent && (AnimNotifyEvent->NotifyStateClass == NULL))
+	{
+		FAnimNextModuleInstance& ModuleInstance = GetModuleInstance();
+		if (AnimNotifyEvent->Notify != nullptr)
+		{
+#if WITH_EDITOR
+			// Prevent firing notifies in animation editors if requested 
+			if(ModuleInstance.GetWorldType() != EWorldType::EditorPreview || AnimNotifyEvent->Notify->ShouldFireInEditor())
+#endif
+			{
+				Dispatcher.EventsToNotify.Add(EventReference);
+			}
+		}
+		else if (AnimNotifyEvent->NotifyName != NAME_None)
+		{
+			// Named notifies can be handled by a custom module event on our worker thread, no need to queue
+			ModuleInstance.RunRigVMEvent(AnimNotifyEvent->NotifyName, InDeltaTime);
+		}
+	}
+}
\ No newline at end of file
diff --git a/Engine/Plugins/Experimental/AnimNextAnimGraph/Source/AnimNextAnimGraph/Private/Traits/NotifyDispatcher.h b/Engine/Plugins/Experimental/AnimNextAnimGraph/Source/AnimNextAnimGraph/Private/Traits/NotifyDispatcher.h
new file mode 100644
--- /dev/null
+++ b/Engine/Plugins/Experimental/AnimNextAnimGraph/Source/AnimNextAnimGraph/Private/Traits/NotifyDispatcher.h
@@ -0,0 +1,125 @@
+﻿// Copyright Epic Games, Inc. All Rights Reserved.
+
+#pragma once
+
+#include "CoreMinimal.h"
+#include "Module/ModuleEvents.h"
+#include "TraitCore/Trait.h"
+#include "TraitInterfaces/ITimeline.h"
+#include "TraitInterfaces/IUpdate.h"
+#include "Animation/AnimNotifyQueue.h"
+#include "Module/ModuleInstanceComponent.h"
+
+#include "NotifyDispatcher.generated.h"
+
+struct FAnimNotifyEventReference;
+class USkeletalMeshComponent;
+
+namespace UE::AnimNext
+{
+	struct FNotifyQueueDispatchEvent;
+}
+
+USTRUCT(meta = (DisplayName = "Notify Dispatcher"))
+struct FAnimNextNotifyDispatcherTraitSharedData : public FAnimNextTraitSharedData
+{
+	GENERATED_BODY()
+
+	// Latent pin support boilerplate
+	#define TRAIT_LATENT_PROPERTIES_ENUMERATOR(GeneratorMacro)
+
+	GENERATE_TRAIT_LATENT_PROPERTIES(FAnimNextNotifyDispatcherTraitSharedData, TRAIT_LATENT_PROPERTIES_ENUMERATOR)
+	#undef TRAIT_LATENT_PROPERTIES_ENUMERATOR
+};
+
+USTRUCT()
+struct FModuleInstanceComponent_NotifyDispatcher : public FModuleInstanceComponent
+{
+	GENERATED_BODY()
+
+	FModuleInstanceComponent_NotifyDispatcher() = default;
+
+	FModuleInstanceComponent_NotifyDispatcher(FAnimNextModuleInstance& InOwnerInstance)
+		: FModuleInstanceComponent(InOwnerInstance)
+	{}
+
+	// FModuleInstanceComponent interface
+	virtual void OnTraitEvent(FAnimNextTraitEvent& Event) override;
+	virtual void OnEndExecution(float InDeltaTime) override;
+
+	// Triggers a single anim notify in the dispatcher
+	void TriggerSingleAnimNotify(float InDeltaTime, UE::AnimNext::FNotifyQueueDispatchEvent& Dispatcher, const FAnimNotifyEventReference& EventReference);
+
+	// Notify queue to dispatch
+	UPROPERTY()
+	FAnimNotifyQueue NotifyQueue;
+
+	UPROPERTY()
+	TArray<FAnimNotifyEvent> ActiveAnimNotifyState;
+
+	UPROPERTY()
+	TArray<FAnimNotifyEventReference> ActiveAnimNotifyEventReference;
+
+	// Skeletal mesh component to 'fake' dispatch from
+	UPROPERTY()
+	TObjectPtr<USkeletalMeshComponent> SkeletalMeshComponent;
+};
+
+namespace UE::AnimNext
+{
+	/**
+	 * FNotifyDispatcherTrait
+	 * 
+	 * A trait that dispatches notifies according to a timeline advancing
+	 */
+	struct FNotifyDispatcherTrait : FAdditiveTrait, ITimeline
+	{
+		DECLARE_ANIM_TRAIT(FNotifyDispatcherTrait, 0x33a2435b, FAdditiveTrait)
+
+		using FSharedData = FAnimNextNotifyDispatcherTraitSharedData;
+
+		struct FInstanceData : FTrait::FInstanceData
+		{
+		};
+
+		// ITimeline impl
+		virtual void AdvanceBy(FExecutionContext& Context, const TTraitBinding<ITimeline>& Binding, float DeltaTime) const override;
+	};
+
+	// Event that dispatches to the module
+	struct FNotifyDispatchEvent : public FAnimNextTraitEvent
+	{
+		DECLARE_ANIM_TRAIT_EVENT(FNotifyDispatchEvent, FAnimNextTraitEvent)
+
+		// Notifies to be dispatched
+		TArray<FAnimNotifyEventReference> Notifies;
+
+		// Weight at dispatch time
+		float Weight = 1.0f;
+	};
+
+	// Event that dispatches from the module to gameplay
+	struct FNotifyQueueDispatchEvent : public FAnimNextModule_ActionEvent
+	{
+		DECLARE_ANIM_TRAIT_EVENT(FNotifyQueueDispatchEvent, FAnimNextModule_ActionEvent)
+
+		// FAnimNextModule_ActionEvent interface
+		virtual bool IsThreadSafe() const override { return bIsThreadSafe; }
+		virtual void Execute() const override;
+
+		// The various events to dispatch
+		TArray<FAnimNotifyEventReference> EventsToNotify;
+		TArray<FAnimNotifyEventReference> EventsToEnd;
+		TArray<FAnimNotifyEventReference> EventsToBegin;
+		TArray<FAnimNotifyEventReference> EventsToTick;
+
+		// Skeletal mesh component to dispatch as
+		TWeakObjectPtr<USkeletalMeshComponent> WeakSkeletalMeshComponent;
+
+		// Delta time to apply to notifies
+		float DeltaSeconds = 0.0f;
+
+		// Whether this queue thread safe to dispatch
+		bool bIsThreadSafe = false;
+	};
+}
diff --git a/Engine/Plugins/Experimental/AnimNextAnimGraph/Source/AnimNextAnimGraph/Private/Traits/SequencePlayer.cpp b/Engine/Plugins/Experimental/AnimNextAnimGraph/Source/AnimNextAnimGraph/Private/Traits/SequencePlayer.cpp
--- a/Engine/Plugins/Experimental/AnimNextAnimGraph/Source/AnimNextAnimGraph/Private/Traits/SequencePlayer.cpp
+++ b/Engine/Plugins/Experimental/AnimNextAnimGraph/Source/AnimNextAnimGraph/Private/Traits/SequencePlayer.cpp
@@ -18,6 +18,7 @@ namespace UE::AnimNext
 		GeneratorMacro(ITimeline) \
 		GeneratorMacro(IUpdate) \
 		GeneratorMacro(IGarbageCollection) \
+		GeneratorMacro(INotifySource) \
 
 	GENERATE_ANIM_TRAIT_IMPLEMENTATION(FSequencePlayerTrait, TRAIT_INTERFACE_ENUMERATOR, NULL_ANIM_TRAIT_INTERFACE_ENUMERATOR, NULL_ANIM_TRAIT_EVENT_ENUMERATOR)
 	#undef TRAIT_INTERFACE_ENUMERATOR
@@ -98,12 +99,26 @@ namespace UE::AnimNext
 				InstanceData->InternalTimeAccumulator,
 				AnimSeq->GetPlayLength(),
 				SharedData->GetPlayRate(Binding),
+				InstanceData->DeltaTimeRecord.Delta,
 				SharedData->GetbLoop(Binding));
 		}
 
 		return FTimelineState();
 	}
 
+	FTimelineDelta FSequencePlayerTrait::GetDelta(FExecutionContext& Context, const TTraitBinding<ITimeline>& Binding) const
+	{
+		const FInstanceData* InstanceData = Binding.GetInstanceData<FInstanceData>();
+		if (const UAnimSequence* AnimSeq = InstanceData->AnimSequence.Get())
+		{
+			return FTimelineDelta(
+				InstanceData->DeltaTimeRecord.Delta,
+				InstanceData->LastAdvanceType);
+		}
+
+		return FTimelineDelta();
+	}
+
 	void FSequencePlayerTrait::AdvanceBy(FExecutionContext& Context, const TTraitBinding<ITimeline>& Binding, float DeltaTime) const
 	{
 		FInstanceData* InstanceData = Binding.GetInstanceData<FInstanceData>();
@@ -125,7 +140,7 @@ namespace UE::AnimNext
 
 			InstanceData->DeltaTimeRecord.Set(PreviousTime, MoveDelta);
 
-			FAnimationRuntime::AdvanceTime(bIsLooping, MoveDelta, InstanceData->InternalTimeAccumulator, SequenceLength);
+			InstanceData->LastAdvanceType = FAnimationRuntime::AdvanceTime(bIsLooping, MoveDelta, InstanceData->InternalTimeAccumulator, SequenceLength);
 		}
 	}
 
@@ -165,4 +180,17 @@ namespace UE::AnimNext
 
 		Collector.AddReferencedObject(InstanceData->AnimSequence);
 	}
+
+	void FSequencePlayerTrait::GetNotifies(FExecutionContext& Context, const TTraitBinding<INotifySource>& Binding, float StartPosition, float Duration, bool bLooping, TArray<FAnimNotifyEventReference>& OutNotifies) const
+	{
+		const FInstanceData* InstanceData = Binding.GetInstanceData<FInstanceData>();
+		if (const UAnimSequence* AnimSeq = InstanceData->AnimSequence.Get())
+		{
+			FAnimTickRecord TickRecord;
+			TickRecord.bLooping = bLooping;
+			FAnimNotifyContext NotifyContext(TickRecord);
+			AnimSeq->GetAnimNotifies(StartPosition, Duration, NotifyContext);
+			OutNotifies = MoveTemp(NotifyContext.ActiveNotifies);
+		}
+	}
 }
diff --git a/Engine/Plugins/Experimental/AnimNextAnimGraph/Source/AnimNextAnimGraph/Private/Traits/SequencePlayer.h b/Engine/Plugins/Experimental/AnimNextAnimGraph/Source/AnimNextAnimGraph/Private/Traits/SequencePlayer.h
--- a/Engine/Plugins/Experimental/AnimNextAnimGraph/Source/AnimNextAnimGraph/Private/Traits/SequencePlayer.h
+++ b/Engine/Plugins/Experimental/AnimNextAnimGraph/Source/AnimNextAnimGraph/Private/Traits/SequencePlayer.h
@@ -10,6 +10,7 @@
 #include "TraitInterfaces/ITimeline.h"
 #include "TraitInterfaces/IUpdate.h"
 #include "Animation/AnimSequence.h"
+#include "TraitInterfaces/INotifySource.h"
 
 #include "SequencePlayer.generated.h"
 
@@ -32,7 +33,7 @@ struct FAnimNextSequencePlayerTraitSharedData : public FAnimNextTraitSharedData
 
 	/** Whether or not this sequence playback will loop. */
 	UPROPERTY(EditAnywhere, Category = "Default")
-	bool bLoop = 0.0f;
+	bool bLoop = false;
 
 	// Latent pin support boilerplate
 	#define TRAIT_LATENT_PROPERTIES_ENUMERATOR(GeneratorMacro) \
@@ -52,7 +53,7 @@ namespace UE::AnimNext
 	 * 
 	 * A trait that can play an animation sequence.
 	 */
-	struct FSequencePlayerTrait : FBaseTrait, IEvaluate, IAttributeProvider, ITimeline, IUpdate, IGarbageCollection
+	struct FSequencePlayerTrait : FBaseTrait, IEvaluate, IAttributeProvider, ITimeline, IUpdate, IGarbageCollection, INotifySource
 	{
 		DECLARE_ANIM_TRAIT(FSequencePlayerTrait, 0x7a0dc157, FBaseTrait)
 
@@ -65,8 +66,13 @@ namespace UE::AnimNext
 
 			/** Delta time range required for root motion extraction **/
 			FDeltaTimeRecord DeltaTimeRecord;
+
+			// Current time accumulator
 			float InternalTimeAccumulator = 0.0f;
 
+			// The last advance type when FAnimationRuntime::AdvanceTime was called
+			ETypeAdvanceAnim LastAdvanceType = ETAA_Default;
+
 			void Construct(const FExecutionContext& Context, const FTraitBinding& Binding);
 			void Destruct(const FExecutionContext& Context, const FTraitBinding& Binding);
 		};
@@ -81,6 +87,7 @@ namespace UE::AnimNext
 		virtual void GetSyncMarkers(FExecutionContext& Context, const TTraitBinding<ITimeline>& Binding, FTimelineSyncMarkerArray& OutSyncMarkers) const override;
 		virtual FTimelineState GetState(FExecutionContext& Context, const TTraitBinding<ITimeline>& Binding) const override;
 		virtual void AdvanceBy(FExecutionContext& Context, const TTraitBinding<ITimeline>& Binding, float DeltaTime) const override;
+		virtual FTimelineDelta GetDelta(FExecutionContext& Context, const TTraitBinding<ITimeline>& Binding) const override;
 
 		// IUpdate impl
 		virtual void OnBecomeRelevant(FUpdateTraversalContext& Context, const TTraitBinding<IUpdate>& Binding, const FTraitUpdateState& TraitState) const override;
@@ -88,5 +95,8 @@ namespace UE::AnimNext
 
 		// IGarbageCollection impl
 		virtual void AddReferencedObjects(const FExecutionContext& Context, const TTraitBinding<IGarbageCollection>& Binding, FReferenceCollector& Collector) const override;
+
+		// INotifySource impl
+		virtual void GetNotifies(FExecutionContext& Context, const TTraitBinding<INotifySource>& Binding, float StartPosition, float Duration, bool bLooping, TArray<FAnimNotifyEventReference>& OutNotifies) const override;
 	};
 }
diff --git a/Engine/Plugins/Experimental/AnimNextAnimGraph/Source/AnimNextAnimGraph/Public/TraitInterfaces/INotifySource.h b/Engine/Plugins/Experimental/AnimNextAnimGraph/Source/AnimNextAnimGraph/Public/TraitInterfaces/INotifySource.h
new file mode 100644
--- /dev/null
+++ b/Engine/Plugins/Experimental/AnimNextAnimGraph/Source/AnimNextAnimGraph/Public/TraitInterfaces/INotifySource.h
@@ -0,0 +1,46 @@
+﻿// Copyright Epic Games, Inc. All Rights Reserved.
+
+#pragma once
+
+#include "CoreMinimal.h"
+#include "TraitCore/ITraitInterface.h"
+#include "TraitCore/TraitBinding.h"
+
+struct FAnimNotifyEventReference;
+
+namespace UE::AnimNext
+{
+	/**
+	 * INotifySource
+	 *
+	 * This interface exposes a source of animation notify events
+	 */
+	struct ANIMNEXTANIMGRAPH_API INotifySource : ITraitInterface
+	{
+		DECLARE_ANIM_TRAIT_INTERFACE(INotifySource, 0x0be015e0)
+
+		// Returns the events of this timeline in the specified time interval
+		virtual void GetNotifies(FExecutionContext& Context, const TTraitBinding<INotifySource>& Binding, float StartPosition, float Duration, bool bLooping, TArray<FAnimNotifyEventReference>& OutNotifies) const;
+
+#if WITH_EDITOR
+		virtual const FText& GetDisplayName() const override;
+		virtual const FText& GetDisplayShortName() const override;
+#endif // WITH_EDITOR
+	};
+
+	/**
+	 * Specialization for trait binding.
+	 */
+	template<>
+	struct TTraitBinding<INotifySource> : FTraitBinding
+	{
+		// @see INotifySource::GetNotifies
+		void GetNotifies(FExecutionContext& Context, float StartPosition, float Duration, bool bLooping, TArray<FAnimNotifyEventReference>& OutNotifies) const
+		{
+			return GetInterface()->GetNotifies(Context, *this, StartPosition, Duration, bLooping, OutNotifies);
+		}
+
+	protected:
+		const INotifySource* GetInterface() const { return GetInterfaceTyped<INotifySource>(); }
+	};
+}
diff --git a/Engine/Plugins/Experimental/AnimNextAnimGraph/Source/AnimNextAnimGraph/Public/TraitInterfaces/ITimeline.h b/Engine/Plugins/Experimental/AnimNextAnimGraph/Source/AnimNextAnimGraph/Public/TraitInterfaces/ITimeline.h
--- a/Engine/Plugins/Experimental/AnimNextAnimGraph/Source/AnimNextAnimGraph/Public/TraitInterfaces/ITimeline.h
+++ b/Engine/Plugins/Experimental/AnimNextAnimGraph/Source/AnimNextAnimGraph/Public/TraitInterfaces/ITimeline.h
@@ -3,6 +3,7 @@
 #pragma once
 
 #include "CoreMinimal.h"
+#include "Animation/AnimSequenceBase.h"
 #include "Misc/MemStack.h"
 #include "TraitCore/ITraitInterface.h"
 #include "TraitCore/TraitBinding.h"
@@ -50,7 +51,7 @@ namespace UE::AnimNext
 		FTimelineState() = default;
 
 		// Construct with a specific state
-		FTimelineState(float InPosition, float InDuration, float InPlayRate, bool bInIsLooping)
+		FTimelineState(float InPosition, float InDuration, float InPlayRate, float InDelta, bool bInIsLooping)
 			: Position(InPosition)
 			, Duration(InDuration)
 			, PlayRate(InPlayRate)
@@ -131,6 +132,37 @@ namespace UE::AnimNext
 		bool bIsLooping = false;
 	};
 
+	/**
+	 * Timeline Delta
+	 *
+	 * Encapsulates the delta state of a timeline, i.e. the changes that occured during the last AdvanceBy call.
+	 */
+	struct FTimelineDelta
+	{
+		// Default construct with no delta
+		FTimelineDelta() = default;
+
+		// Construct with a specific delta
+		FTimelineDelta(float InDeltaTime, ETypeAdvanceAnim InAdvanceType)
+			: DeltaTime(InDeltaTime)
+			, AdvanceType(InAdvanceType)
+		{
+		}
+
+		// Returns the delta time in seconds when last advanced
+		float GetDeltaTime() const { return DeltaTime; }
+
+		// Returns the result of the last advance
+		ETypeAdvanceAnim GetAdvanceType() const { return AdvanceType; }
+
+	private:
+		// Delta time in seconds when last advanced
+		float DeltaTime = 0.0f;
+
+		// The result of the last advance
+		ETypeAdvanceAnim AdvanceType = ETAA_Default;
+	};
+
 	/**
 	 * ITimeline
 	 *
@@ -149,6 +181,9 @@ namespace UE::AnimNext
 		// Advances time by the provided delta time (positive or negative) on this timeline
 		virtual void AdvanceBy(FExecutionContext& Context, const TTraitBinding<ITimeline>& Binding, float DeltaTime) const;
 
+		// Returns the last change in state of this timeline
+		virtual FTimelineDelta GetDelta(FExecutionContext& Context, const TTraitBinding<ITimeline>& Binding) const;
+
 #if WITH_EDITOR
 		virtual const FText& GetDisplayName() const override;
 		virtual const FText& GetDisplayShortName() const override;
@@ -179,6 +214,12 @@ namespace UE::AnimNext
 			GetInterface()->AdvanceBy(Context, *this, DeltaTime);
 		}
 
+		// @see ITimeline::GetDelta
+		FTimelineDelta GetDelta(FExecutionContext& Context) const
+		{
+			return GetInterface()->GetDelta(Context, *this);
+		}
+
 	protected:
 		const ITimeline* GetInterface() const { return GetInterfaceTyped<ITimeline>(); }
 	};
diff --git a/Engine/Plugins/Experimental/AnimNextAnimGraph/Source/AnimNextAnimGraphEditor/Private/AnimNextAnimGraphEditorModule.cpp b/Engine/Plugins/Experimental/AnimNextAnimGraph/Source/AnimNextAnimGraphEditor/Private/AnimNextAnimGraphEditorModule.cpp
--- a/Engine/Plugins/Experimental/AnimNextAnimGraph/Source/AnimNextAnimGraphEditor/Private/AnimNextAnimGraphEditorModule.cpp
+++ b/Engine/Plugins/Experimental/AnimNextAnimGraph/Source/AnimNextAnimGraphEditor/Private/AnimNextAnimGraphEditorModule.cpp
@@ -33,6 +33,7 @@
 #include "AnimNextTraitStackUnitNode.h"
 #include "RigVMCore/RigVMFunction.h"
 #include "EditorUtils.h"
+#include "Module/AnimNextEventGraphSchema.h"
 
 #define LOCTEXT_NAMESPACE "FAnimNextAnimGraphEditorModule"
 
@@ -182,6 +183,10 @@ class FAnimNextAnimGraphEditorModule : public IModuleInterface
 					}
 				}
 			}
+			else if (RigVMEdGraph->GetModel()->GetSchemaClass() == UAnimNextEventGraphSchema::StaticClass())
+			{
+				InContexMenuBuilder.AddAction(MakeShared<FAnimNextSchemaAction_NotifyEvent>());
+			}
 		}
 
 		// Add trait stack using a custom RigUnit node class
diff --git a/Engine/Plugins/Experimental/AnimNextAnimGraph/Source/AnimNextAnimGraphUncookedOnly/AnimNextAnimGraphUncookedOnly.Build.cs b/Engine/Plugins/Experimental/AnimNextAnimGraph/Source/AnimNextAnimGraphUncookedOnly/AnimNextAnimGraphUncookedOnly.Build.cs
--- a/Engine/Plugins/Experimental/AnimNextAnimGraph/Source/AnimNextAnimGraphUncookedOnly/AnimNextAnimGraphUncookedOnly.Build.cs
+++ b/Engine/Plugins/Experimental/AnimNextAnimGraph/Source/AnimNextAnimGraphUncookedOnly/AnimNextAnimGraphUncookedOnly.Build.cs
@@ -20,6 +20,8 @@ namespace UnrealBuildTool.Rules
 					"AnimGraph",
 					"RigVM",
 					"ControlRig",
+					"SlateCore",
+					"Slate",
 				}
 			);
 
diff --git a/Engine/Plugins/Experimental/AnimNextAnimGraph/Source/AnimNextAnimGraphUncookedOnly/Internal/Graph/AnimGraphEditorSchemaActions.cpp b/Engine/Plugins/Experimental/AnimNextAnimGraph/Source/AnimNextAnimGraphUncookedOnly/Internal/Graph/AnimGraphEditorSchemaActions.cpp
--- a/Engine/Plugins/Experimental/AnimNextAnimGraph/Source/AnimNextAnimGraphUncookedOnly/Internal/Graph/AnimGraphEditorSchemaActions.cpp
+++ b/Engine/Plugins/Experimental/AnimNextAnimGraph/Source/AnimNextAnimGraphUncookedOnly/Internal/Graph/AnimGraphEditorSchemaActions.cpp
@@ -21,6 +21,11 @@
 #include "Exporters/Exporter.h"
 #include "UnrealExporter.h"
 #include "AnimNextUnitNode.h"
+#include "PersonaModule.h"
+#include "ScopedTransaction.h"
+#include "Modules/ModuleManager.h"
+#include "RigVMFunctions/Execution/RigVMFunction_UserDefinedEvent.h"
+#include "Widgets/Input/STextEntryPopup.h"
 
 #define LOCTEXT_NAMESPACE "AnimNextAnimGraphSchemaActions"
 
@@ -102,4 +107,90 @@ UEdGraphNode* FAnimNextSchemaAction_AddManifestNode::PerformAction(UEdGraph* Par
 	return NewNode;
 }
 
+FAnimNextSchemaAction_NotifyEvent::FAnimNextSchemaAction_NotifyEvent()
+	: FAnimNextSchemaAction(LOCTEXT("NotifiesCategory", "Notifies"), LOCTEXT("AddNotifyEventLabel", "Add Notify Event..."), LOCTEXT("AddNotifyEventTooltip", "Add a custom event node to handle a named notify event"))
+{
+}
+
+UEdGraphNode* FAnimNextSchemaAction_NotifyEvent::PerformAction(UEdGraph* ParentGraph, TArray<UEdGraphPin*>& FromPins, const FVector2D Location, bool bSelectNewNode)
+{
+	URigVMEdGraph* EdGraph = Cast<URigVMEdGraph>(ParentGraph);
+	URigVMController* Controller = Cast<URigVMController>(EdGraph->GetController());
+	if(Controller == nullptr)
+	{
+		return nullptr;
+	}
+
+	FMenuBuilder MenuBuilder(true, nullptr);
+	MenuBuilder.BeginSection("AddNotify", LOCTEXT("AddNotifyEventSection", "Add Notify Event"));
+
+	auto CreateNotifyEventWithName = [Controller, Location](FName InNotifyEventName)
+	{
+		if(!Controller->GetAllEventNames().Contains(InNotifyEventName))
+		{
+			Controller->OpenUndoBracket(LOCTEXT("AddNotifyEventTransaction", "Add Notify Event").ToString());
+			URigVMUnitNode* CustomEventNode = Controller->AddUnitNode(FRigVMFunction_UserDefinedEvent::StaticStruct(), FRigVMStruct::ExecuteName, Location);
+			Controller->SetPinDefaultValue(CustomEventNode->FindPin(GET_MEMBER_NAME_STRING_CHECKED(FRigVMFunction_UserDefinedEvent, EventName)), InNotifyEventName.ToString(), true, true, true);
+			Controller->CloseUndoBracket();
+		}
+	};
+
+	MenuBuilder.AddMenuEntry(
+		LOCTEXT("AddNewNotifyEventLabel", "Add New Notify Event..."),
+		LOCTEXT("AddNewNotifyEventTooltip", "Add a new notify event as a custom event"),
+		FSlateIcon(),
+		FUIAction(FExecuteAction::CreateLambda([CreateNotifyEventWithName]()
+		{
+			// Show dialog to enter new track name
+			TSharedRef<STextEntryPopup> TextEntry =
+				SNew(STextEntryPopup)
+				.Label(LOCTEXT("NewNotifyLabel", "Notify Name"))
+				.OnTextCommitted_Lambda([CreateNotifyEventWithName](const FText& InText, ETextCommit::Type InCommitType)
+				{
+					FSlateApplication::Get().DismissAllMenus();
+					FName NotifyName = *InText.ToString();
+					CreateNotifyEventWithName(NotifyName);
+				});
+
+			// Show dialog to enter new event name
+			FSlateApplication::Get().PushMenu(
+				FSlateApplication::Get().GetInteractiveTopLevelWindows()[0],
+				FWidgetPath(),
+				TextEntry,
+				FSlateApplication::Get().GetCursorPos(),
+				FPopupTransitionEffect( FPopupTransitionEffect::TypeInPopup));
+		})),
+		NAME_None,
+		EUserInterfaceActionType::Button);
+	
+	FPersonaModule& PersonaModule = FModuleManager::LoadModuleChecked<FPersonaModule>("Persona");
+	TSharedRef<SWidget> NotifyPickerWidget = PersonaModule.CreateSkeletonNotifyPicker(FOnNotifyPicked::CreateLambda([CreateNotifyEventWithName](FName InNotifyName)
+	{
+		FSlateApplication::Get().DismissAllMenus();
+		CreateNotifyEventWithName(InNotifyName);
+	}));
+
+	MenuBuilder.AddWidget(
+		SNew(SBox)
+		.WidthOverride(300.0f)
+		.HeightOverride(400.0f)
+		[
+			NotifyPickerWidget
+		],
+		FText::GetEmpty(),
+		true, false);
+
+	MenuBuilder.EndSection();
+
+	FSlateApplication::Get().PushMenu(
+		FSlateApplication::Get().GetInteractiveTopLevelWindows()[0],
+		FWidgetPath(),
+		MenuBuilder.MakeWidget(),
+		FSlateApplication::Get().GetCursorPos(),
+		FPopupTransitionEffect(FPopupTransitionEffect::ContextMenu)
+	);
+
+	return nullptr;
+}
+
 #undef LOCTEXT_NAMESPACE
diff --git a/Engine/Plugins/Experimental/AnimNextAnimGraph/Source/AnimNextAnimGraphUncookedOnly/Internal/Graph/AnimGraphEditorSchemaActions.h b/Engine/Plugins/Experimental/AnimNextAnimGraph/Source/AnimNextAnimGraphUncookedOnly/Internal/Graph/AnimGraphEditorSchemaActions.h
--- a/Engine/Plugins/Experimental/AnimNextAnimGraph/Source/AnimNextAnimGraphUncookedOnly/Internal/Graph/AnimGraphEditorSchemaActions.h
+++ b/Engine/Plugins/Experimental/AnimNextAnimGraph/Source/AnimNextAnimGraphUncookedOnly/Internal/Graph/AnimGraphEditorSchemaActions.h
@@ -23,3 +23,15 @@ struct ANIMNEXTANIMGRAPHUNCOOKEDONLY_API FAnimNextSchemaAction_AddManifestNode :
 	TSoftObjectPtr<URigVMGraph> ModelGraph;
 	FString NodeName;
 };
+
+USTRUCT()
+struct ANIMNEXTANIMGRAPHUNCOOKEDONLY_API FAnimNextSchemaAction_NotifyEvent : public FAnimNextSchemaAction
+{
+	GENERATED_BODY()
+
+	FAnimNextSchemaAction_NotifyEvent();
+
+	// FEdGraphSchemaAction Interface
+	virtual UEdGraphNode* PerformAction(UEdGraph* ParentGraph, TArray<UEdGraphPin*>& FromPins, const FVector2D Location, bool bSelectNewNode = true) override;
+	virtual UEdGraphNode* PerformAction(UEdGraph* ParentGraph, UEdGraphPin* FromPin, const FVector2D Location, bool bSelectNewNode = true) override { return nullptr; }
+};
diff --git a/Engine/Plugins/Experimental/AnimNextPoseSearch/Source/AnimNextPoseSearch/Private/MotionMatchingTrait.cpp b/Engine/Plugins/Experimental/AnimNextPoseSearch/Source/AnimNextPoseSearch/Private/MotionMatchingTrait.cpp
--- a/Engine/Plugins/Experimental/AnimNextPoseSearch/Source/AnimNextPoseSearch/Private/MotionMatchingTrait.cpp
+++ b/Engine/Plugins/Experimental/AnimNextPoseSearch/Source/AnimNextPoseSearch/Private/MotionMatchingTrait.cpp
@@ -127,7 +127,7 @@ namespace UE::AnimNext
 			const UObject* AnimContext = nullptr;
 			if (const FAnimNextModuleInstance* ModuleInstance = Context.GetRootGraphInstance().GetModuleInstance())
 			{
-				AnimContext = ModuleInstance->Object;
+				AnimContext = ModuleInstance->GetObject();
 			}
 
 			UPoseSearchLibrary::UpdateMotionMatchingState(
diff --git a/Engine/Plugins/Experimental/AnimNextPoseSearch/Source/AnimNextPoseSearch/Private/RigUnit_DebugDrawTrajectory.cpp b/Engine/Plugins/Experimental/AnimNextPoseSearch/Source/AnimNextPoseSearch/Private/RigUnit_DebugDrawTrajectory.cpp
--- a/Engine/Plugins/Experimental/AnimNextPoseSearch/Source/AnimNextPoseSearch/Private/RigUnit_DebugDrawTrajectory.cpp
+++ b/Engine/Plugins/Experimental/AnimNextPoseSearch/Source/AnimNextPoseSearch/Private/RigUnit_DebugDrawTrajectory.cpp
@@ -18,7 +18,7 @@ FRigUnit_DebugDrawTrajectory_Execute()
 
 	// TODO: this is not thread safe at the moment. Once we get the 'weak semantics' CL checked in, we can move commonly-used debug info
 	// into a module component that can be used to access GT data where needed. 
-	UAnimNextComponent* Component = Cast<UAnimNextComponent>(ModuleInstance.Object);
+	UAnimNextComponent* Component = Cast<UAnimNextComponent>(ModuleInstance.GetObject());
 
 	Trajectory.DebugDrawTrajectory(Component, LogPoseSearch, ELogVerbosity::Display, DebugThickness, DebugOffset);
 #endif
diff --git a/Engine/Plugins/Experimental/MoverAnimNext/Source/MoverAnimNext/Private/Graph/RigUnit_GenerateMoverTrajectory.cpp b/Engine/Plugins/Experimental/MoverAnimNext/Source/MoverAnimNext/Private/Graph/RigUnit_GenerateMoverTrajectory.cpp
--- a/Engine/Plugins/Experimental/MoverAnimNext/Source/MoverAnimNext/Private/Graph/RigUnit_GenerateMoverTrajectory.cpp
+++ b/Engine/Plugins/Experimental/MoverAnimNext/Source/MoverAnimNext/Private/Graph/RigUnit_GenerateMoverTrajectory.cpp
@@ -15,7 +15,7 @@ FRigUnit_GenerateMoverTrajectory_Execute()
 
 	// TODO: this is not thread safe at the moment. Once we get the 'weak semantics' CL checked in, we can move commonly-used debug info
 	// into a module component that can be used to access GT data where needed. 
-	UAnimNextComponent* Component = Cast<UAnimNextComponent>(ModuleInstance.Object);
+	UAnimNextComponent* Component = Cast<UAnimNextComponent>(ModuleInstance.GetObject());
 
 	// Note this logic is based on UPoseSearchTrajectoryLibrary::PoseSearchGeneratePredictorTrajectory
 
diff --git a/Engine/Plugins/Runtime/RigVM/Source/RigVMDeveloper/Private/RigVMModel/RigVMClient.cpp b/Engine/Plugins/Runtime/RigVM/Source/RigVMDeveloper/Private/RigVMModel/RigVMClient.cpp
--- a/Engine/Plugins/Runtime/RigVM/Source/RigVMDeveloper/Private/RigVMModel/RigVMClient.cpp
+++ b/Engine/Plugins/Runtime/RigVM/Source/RigVMDeveloper/Private/RigVMModel/RigVMClient.cpp
@@ -748,13 +748,25 @@ URigVMFunctionLibrary* FRigVMClient::GetOrCreateFunctionLibrary(TSubclassOf<URig
 	return NewFunctionLibrary;
 }
 
-TArray<FName> FRigVMClient::GetEntryNames() const
+TArray<FName> FRigVMClient::GetEntryNames(UScriptStruct* InUnitScriptStructFilter) const
 {
 	TArray<FName> EntryNames;
 	for(const URigVMGraph* Model : GetModels())
 	{
 		for(const URigVMNode* Node : Model->GetNodes())
 		{
+			// Filter out unit nodes that are not of the specified type
+			if(InUnitScriptStructFilter != nullptr)
+			{
+				if(const URigVMUnitNode* UnitNode = Cast<URigVMUnitNode>(Node))
+				{
+					if(!UnitNode->GetScriptStruct()->IsChildOf(InUnitScriptStructFilter))
+					{
+						continue;
+					}
+				}
+			}
+
 			const FName EntryName = Node->GetEventName();
 			if(!EntryName.IsNone())
 			{
diff --git a/Engine/Plugins/Runtime/RigVM/Source/RigVMDeveloper/Public/RigVMModel/RigVMClient.h b/Engine/Plugins/Runtime/RigVM/Source/RigVMDeveloper/Public/RigVMModel/RigVMClient.h
--- a/Engine/Plugins/Runtime/RigVM/Source/RigVMDeveloper/Public/RigVMModel/RigVMClient.h
+++ b/Engine/Plugins/Runtime/RigVM/Source/RigVMDeveloper/Public/RigVMModel/RigVMClient.h
@@ -261,7 +261,7 @@ public:
 	URigVMFunctionLibrary* GetFunctionLibrary() const { return FunctionLibrary; }
 	URigVMFunctionLibrary* GetOrCreateFunctionLibrary(bool bSetupUndoRedo, const FObjectInitializer* ObjectInitializer = nullptr, bool bCreateController = true);
 	URigVMFunctionLibrary* GetOrCreateFunctionLibrary(TSubclassOf<URigVMSchema> InSchemaClass, bool bSetupUndoRedo, const FObjectInitializer* ObjectInitializer = nullptr, bool bCreateController = true);
-	TArray<FName> GetEntryNames() const;
+	TArray<FName> GetEntryNames(UScriptStruct* InUnitScriptStructFilter = nullptr) const;
 	UE_DEPRECATED(5.5, "Please use GetDefaultExecuteContextStruct or get an execute context from a graph/controller schema.")
 	UScriptStruct* GetExecuteContextStruct() const { return GetDefaultExecuteContextStruct(); }
 	UScriptStruct* GetDefaultExecuteContextStruct() const;
diff --git a/Engine/Source/Editor/Persona/Private/AnimAssetFindReplaceNotifies.cpp b/Engine/Source/Editor/Persona/Private/AnimAssetFindReplaceNotifies.cpp
--- a/Engine/Source/Editor/Persona/Private/AnimAssetFindReplaceNotifies.cpp
+++ b/Engine/Source/Editor/Persona/Private/AnimAssetFindReplaceNotifies.cpp
@@ -4,6 +4,7 @@
 
 #include "IEditableSkeleton.h"
 #include "ISkeletonEditorModule.h"
+#include "PersonaModule.h"
 #include "SAnimAssetFindReplace.h"
 #include "Animation/Skeleton.h"
 #include "Animation/AnimationAsset.h"
@@ -65,7 +66,27 @@ FString UAnimAssetFindReplaceNotifies::GetFindResultStringFromAssetData(const FA
 
 TConstArrayView<UClass*> UAnimAssetFindReplaceNotifies::GetSupportedAssetTypes() const
 {
-	static UClass* Types[] = { UAnimSequenceBase::StaticClass(), USkeleton::StaticClass() };
+	FPersonaModule& PersonaModule = FModuleManager::GetModuleChecked<FPersonaModule>("Persona");
+	
+	static TArray<UClass*> Types;
+	static uint32 Version = 0;
+	if(Types.Num() == 0 || Version != PersonaModule.GetNotifyHostAssetVersion())
+	{
+		Types.Reset();
+		Types.Append({ UAnimSequenceBase::StaticClass(), USkeleton::StaticClass() });
+
+		for(const FTopLevelAssetPath& TopLevelAssetPath : PersonaModule.GetAllNotifyHostAssetClassPaths())
+		{
+			FSoftClassPath SoftClassPath(TopLevelAssetPath.ToString());
+			if(UClass* ResolvedClass = SoftClassPath.ResolveClass())
+			{
+				Types.Add(ResolvedClass);
+			}
+		}
+
+		Version = PersonaModule.GetNotifyHostAssetVersion();
+	}
+
 	return Types;
 }
 
@@ -213,6 +234,14 @@ void UAnimAssetFindReplaceNotifies::ReplaceInAsset(const FAssetData& InAssetData
 				}
 			}
 		}
+		else
+		{
+			FPersonaModule& PersonaModule = FModuleManager::GetModuleChecked<FPersonaModule>("Persona");
+			if(const FPersonaModule::FNotifyHostAssetParameters* NotifyHostParameters = PersonaModule.FindNotifyHostAsset(Asset))
+			{
+				NotifyHostParameters->OnReplaceNotify.ExecuteIfBound(Asset, GetFindStringRef(), GetReplaceStringRef(), GetFindWholeWord(), GetSearchCase());
+			}
+		}
 	}
 }
 
@@ -273,6 +302,14 @@ void UAnimAssetFindReplaceNotifies::RemoveInAsset(const FAssetData& InAssetData)
 			}
 		}
 	}
+	else
+	{
+		FPersonaModule& PersonaModule = FModuleManager::GetModuleChecked<FPersonaModule>("Persona");
+		if(const FPersonaModule::FNotifyHostAssetParameters* NotifyHostParameters = PersonaModule.FindNotifyHostAsset(Asset))
+		{
+			NotifyHostParameters->OnRemoveNotify.ExecuteIfBound(Asset, GetFindStringRef(), GetFindWholeWord(), GetSearchCase());
+		}
+	}
 }
 
 void UAnimAssetFindReplaceNotifies::GetAutoCompleteNames(TArrayView<FAssetData> InAssetDatas, TSet<FString>& OutUniqueNames) const
diff --git a/Engine/Source/Editor/Persona/Private/PersonaModule.cpp b/Engine/Source/Editor/Persona/Private/PersonaModule.cpp
--- a/Engine/Source/Editor/Persona/Private/PersonaModule.cpp
+++ b/Engine/Source/Editor/Persona/Private/PersonaModule.cpp
@@ -1132,6 +1132,56 @@ TSharedRef<SWidget> FPersonaModule::CreateCurvePicker(const USkeleton* InSkeleto
 		.IsCurveNameMarkedForExclusion(InIsCurveNameMarkedForExclusion);
 }
 
+TSharedRef<SWidget> FPersonaModule::CreateSkeletonNotifyPicker(FOnNotifyPicked InOnNotifyPicked, TSharedPtr<IEditableSkeleton> InEditableSkeleton, TSharedPtr<FAssetEditorToolkit> InHostingApp, bool bInShowSyncMarkers, bool bInShowNotifies, bool bInShowCompatibleSkeletonAssets, bool bInShowOtherAssets)
+{
+	return SNew(SSkeletonAnimNotifies, InHostingApp)
+		.OnItemSelected(InOnNotifyPicked)
+		.EditableSkeleton(InEditableSkeleton)
+		.ShowSyncMarkers(bInShowSyncMarkers)
+		.ShowNotifies(bInShowNotifies)
+		.ShowCompatibleSkeletonAssets(bInShowCompatibleSkeletonAssets)
+		.ShowOtherAssets(bInShowOtherAssets);
+}
+
+void FPersonaModule::RegisterNotifyHostAsset(const FTopLevelAssetPath& InAssetClassPath, const FNotifyHostAssetParameters& InParameters)
+{
+	NotifyHostParameters.Add(InAssetClassPath, InParameters);
+	NotifyHostAssetVersion++;
+}
+
+void FPersonaModule::UnregisterNotifyHostAsset(const FTopLevelAssetPath& InAssetClassPath)
+{
+	NotifyHostParameters.Remove(InAssetClassPath);
+}
+
+const FPersonaModule::FNotifyHostAssetParameters* FPersonaModule::FindNotifyHostAsset(UObject* InAsset)
+{
+	check(InAsset);
+	UClass* AssetClass = InAsset->GetClass();
+	while(AssetClass)
+	{
+		if(const FNotifyHostAssetParameters* FoundParameters = NotifyHostParameters.Find(AssetClass->GetClassPathName()))
+		{
+			return FoundParameters;
+		}
+		AssetClass = AssetClass->GetSuperClass();
+	}
+
+	return nullptr;
+}
+
+TArray<FTopLevelAssetPath> FPersonaModule::GetAllNotifyHostAssetClassPaths() const
+{
+	TArray<FTopLevelAssetPath> AssetClassPaths;
+	NotifyHostParameters.GenerateKeyArray(AssetClassPaths);
+	return AssetClassPaths;
+}
+
+uint32 FPersonaModule::GetNotifyHostAssetVersion() const
+{
+	return NotifyHostAssetVersion;
+}
+
 void FPersonaModule::AddCommonMenuExtensions(UToolMenu* InToolMenu, const FCommonToolMenuExtensionArgs& InArgs)
 {
 	UPersonaToolMenuContext* PersonaToolMenuContext = InToolMenu->FindContext<UPersonaToolMenuContext>();
diff --git a/Engine/Source/Editor/Persona/Private/SSkeletonAnimNotifies.cpp b/Engine/Source/Editor/Persona/Private/SSkeletonAnimNotifies.cpp
--- a/Engine/Source/Editor/Persona/Private/SSkeletonAnimNotifies.cpp
+++ b/Engine/Source/Editor/Persona/Private/SSkeletonAnimNotifies.cpp
@@ -229,6 +229,7 @@ void SSkeletonAnimNotifies::Construct(const FArguments& InArgs, const TSharedPtr
 			Filter->SetActive(true);
 		}
 
+		if (EditableSkeleton.IsValid())
 		{
 			CurrentFilterFlags |= bShowOtherAssets ? EAnimNotifyFilterFlags::OtherAssets : EAnimNotifyFilterFlags::None;
 
@@ -244,7 +245,12 @@ void SSkeletonAnimNotifies::Construct(const FArguments& InArgs, const TSharedPtr
 			FilterBar->AddFilter(Filter);
 			Filter->SetActive(bShowOtherAssets);
 		}
+		else
+		{
+			CurrentFilterFlags |= EAnimNotifyFilterFlags::OtherAssets;
+		}
 
+		if (EditableSkeleton.IsValid())
 		{
 			CurrentFilterFlags |= bShowCompatibleSkeletonAssets ? EAnimNotifyFilterFlags::CompatibleAssets : EAnimNotifyFilterFlags::None;
 
@@ -260,6 +266,10 @@ void SSkeletonAnimNotifies::Construct(const FArguments& InArgs, const TSharedPtr
 			FilterBar->AddFilter(Filter);
 			Filter->SetActive(bShowCompatibleSkeletonAssets);
 		}
+		else
+		{
+			CurrentFilterFlags |= EAnimNotifyFilterFlags::CompatibleAssets;
+		}
 
 		TSharedRef<SWidget> AddFilterButton = SBasicFilterBar<EAnimNotifyFilterFlags>::MakeAddFilterButton(FilterBar);
 		
@@ -343,21 +353,25 @@ TSharedRef<ITableRow> SSkeletonAnimNotifies::GenerateNotifyRow(TSharedPtr<FDispl
 		TooltipBuilder.AppendLineFormat(LOCTEXT("NotifyTooltip", "Notify '{0}'"), FText::FromName(InInfo->Name));
 	}
 
-	if (EnumHasAnyFlags(InInfo->ItemFlags, EAnimNotifyFilterFlags::CurrentSkeleton))
+	const USkeleton* CurrentSkeleton = EditableSkeleton.IsValid() ? &EditableSkeleton->GetSkeleton() : nullptr;
+	if(CurrentSkeleton)
 	{
-		TooltipBuilder.AppendLine(LOCTEXT("CurrentSkeleton", "Item is on the current skeleton"));
-	}
+		if (EnumHasAnyFlags(InInfo->ItemFlags, EAnimNotifyFilterFlags::CurrentSkeleton))
+		{
+			TooltipBuilder.AppendLine(LOCTEXT("CurrentSkeleton", "Item is on the current skeleton"));
+		}
 
-	if (EnumHasAnyFlags(InInfo->ItemFlags, EAnimNotifyFilterFlags::CompatibleAssets))
-	{
-		TooltipBuilder.AppendLine(LOCTEXT("CompatibleAsset", "Item is on a compatible asset"));
-	}
+		if (EnumHasAnyFlags(InInfo->ItemFlags, EAnimNotifyFilterFlags::CompatibleAssets))
+		{
+			TooltipBuilder.AppendLine(LOCTEXT("CompatibleAsset", "Item is on a compatible asset"));
+		}
 
-	if (EnumHasAnyFlags(InInfo->ItemFlags, EAnimNotifyFilterFlags::OtherAssets))
-	{
-		TooltipBuilder.AppendLine(LOCTEXT("AnotherAsset", "Item is on another asset"));
+		if (EnumHasAnyFlags(InInfo->ItemFlags, EAnimNotifyFilterFlags::OtherAssets))
+		{
+			TooltipBuilder.AppendLine(LOCTEXT("AnotherAsset", "Item is on another asset"));
+		}
 	}
-	
+
 	return
 		SNew( STableRow<TSharedPtr<FDisplayedAnimNotifyInfo>>, OwnerTable )
 		.ToolTipText(TooltipBuilder.ToText())
@@ -379,12 +393,12 @@ TSharedRef<ITableRow> SSkeletonAnimNotifies::GenerateNotifyRow(TSharedPtr<FDispl
 				SAssignNew(InInfo->InlineEditableText, SInlineEditableTextBlock)
 				.Text(FText::FromName(InInfo->Name))
 				.Font(FAppStyle::Get().GetFontStyle("SmallFont"))
-				.ColorAndOpacity(EnumHasAnyFlags(InInfo->ItemFlags, EAnimNotifyFilterFlags::CurrentSkeleton) ? FSlateColor::UseForeground() : FSlateColor::UseSubduedForeground())
+				.ColorAndOpacity((CurrentSkeleton == nullptr || EnumHasAnyFlags(InInfo->ItemFlags, EAnimNotifyFilterFlags::CurrentSkeleton)) ? FSlateColor::UseForeground() : FSlateColor::UseSubduedForeground())
 				.OnVerifyTextChanged(this, &SSkeletonAnimNotifies::OnVerifyNotifyNameCommit, InInfo)
 				.OnTextCommitted(this, &SSkeletonAnimNotifies::OnNotifyNameCommitted, InInfo)
 				.IsSelected(this, &SSkeletonAnimNotifies::IsSelected)
 				.HighlightText_Lambda([this](){ return FilterText; })
-				.IsReadOnly(bIsPicker && EditableSkeleton.IsValid())
+				.IsReadOnly(bIsPicker || !EditableSkeleton.IsValid())
 			]
 		];
 }
@@ -729,6 +743,10 @@ void SSkeletonAnimNotifies::CreateNotifiesList()
 	FARFilter Filter;
 	Filter.bRecursiveClasses = true;
 	Filter.ClassPaths.Append({ UAnimSequenceBase::StaticClass()->GetClassPathName(), USkeleton::StaticClass()->GetClassPathName() } );
+
+	FPersonaModule& PersonaModule = FModuleManager::GetModuleChecked<FPersonaModule>("Persona");
+	Filter.ClassPaths.Append(PersonaModule.GetAllNotifyHostAssetClassPaths());
+	
 	Filter.TagsAndValues.Add(USkeleton::AnimNotifyTag);
 	Filter.TagsAndValues.Add(USkeleton::AnimSyncMarkerTag);
 
diff --git a/Engine/Source/Editor/Persona/Private/SSkeletonAnimNotifies.h b/Engine/Source/Editor/Persona/Private/SSkeletonAnimNotifies.h
--- a/Engine/Source/Editor/Persona/Private/SSkeletonAnimNotifies.h
+++ b/Engine/Source/Editor/Persona/Private/SSkeletonAnimNotifies.h
@@ -24,9 +24,6 @@ class IEditableSkeleton;
 class SToolTip;
 struct FNotificationInfo;
 
-/** Delegate fired when a notify is selected */
-DECLARE_DELEGATE_OneParam(FOnItemSelected, const FName& /*InSelectedNotify*/);
-
 /////////////////////////////////////////////////////
 // FSkeletonAnimNotifiesSummoner
 struct FSkeletonAnimNotifiesSummoner : public FWorkflowTabFactory
@@ -123,7 +120,7 @@ public:
 	SLATE_EVENT(FOnObjectsSelected, OnObjectsSelected)
 
 	/** Delegate fired when a notify is selected */
-	SLATE_EVENT(FOnItemSelected, OnItemSelected)
+	SLATE_EVENT(FOnNotifyPicked, OnItemSelected)
 
 	/** Whether we should use this dialog as a picker or an editor. In picker mode we cant add, remove or rename notifies. */
 	SLATE_ARGUMENT(bool, IsPicker)
@@ -268,7 +265,7 @@ private:
 	FOnObjectsSelected OnObjectsSelected;
 
 	/** Delegate fired when a notify is selected */
-	FOnItemSelected OnItemSelected;
+	FOnNotifyPicked OnItemSelected;
 
 	/** The app we are embedded in */
 	TWeakPtr<class FAssetEditorToolkit> WeakHostingApp;
diff --git a/Engine/Source/Editor/Persona/Public/PersonaDelegates.h b/Engine/Source/Editor/Persona/Public/PersonaDelegates.h
--- a/Engine/Source/Editor/Persona/Public/PersonaDelegates.h
+++ b/Engine/Source/Editor/Persona/Public/PersonaDelegates.h
@@ -75,4 +75,7 @@ DECLARE_DELEGATE_OneParam(FOnSetBlendSpacePreviewPosition, FVector /*InBlendSamp
 DECLARE_DELEGATE_OneParam(FOnCurvePicked, const FName& /*PickedName*/);
 
 // Called by the curve picker to filter out curve names that are displayed to the user
-DECLARE_DELEGATE_RetVal_OneParam(bool, FIsCurveNameMarkedForExclusion, const FName& /*CurveName*/);
\ No newline at end of file
+DECLARE_DELEGATE_RetVal_OneParam(bool, FIsCurveNameMarkedForExclusion, const FName& /*CurveName*/);
+
+// Called by the notify picker to indicate that a notify was chosen by the user
+DECLARE_DELEGATE_OneParam(FOnNotifyPicked, const FName& /*PickedName*/);
\ No newline at end of file
diff --git a/Engine/Source/Editor/Persona/Public/PersonaModule.h b/Engine/Source/Editor/Persona/Public/PersonaModule.h
--- a/Engine/Source/Editor/Persona/Public/PersonaModule.h
+++ b/Engine/Source/Editor/Persona/Public/PersonaModule.h
@@ -459,7 +459,37 @@ public:
 
 	UE_DEPRECATED(5.3, "Please use CreateCurvePicker that takes a const USkeleton*")
 	virtual TSharedRef<SWidget> CreateCurvePicker(TSharedRef<IEditableSkeleton> InEditableSkeleton, FOnCurvePicked InOnCurvePicked, FIsCurveNameMarkedForExclusion InIsCurveNameMarkedForExclusion = FIsCurveNameMarkedForExclusion());
-	
+
+	/** Create a widget that can choose a notify name. @see SSkeletonAnimNotifies. */
+	virtual TSharedRef<SWidget> CreateSkeletonNotifyPicker(FOnNotifyPicked InOnNotifyPicked, TSharedPtr<IEditableSkeleton> InEditableSkeleton = nullptr, TSharedPtr<FAssetEditorToolkit> InHostingApp = nullptr, bool bInShowSyncMarkers = false, bool bInShowNotifies = true, bool bInShowCompatibleSkeletonAssets = true, bool bInShowOtherAssets = true);
+
+	// Struct allowing the anim notify system to interact with an asset type
+	struct FNotifyHostAssetParameters
+	{
+		// Delegate called to replace a notify event with another
+		using FOnReplaceNotify = TDelegate<void(UObject* InAsset, const FString& InFindString, const FString& InReplaceString, bool bFindWholeWord, ESearchCase::Type InSearchCase)>;
+		FOnReplaceNotify OnReplaceNotify;
+
+		// Delegate called to remove a notify event
+		using FOnRemoveNotify = TDelegate<void(UObject* InAsset, const FString& InFindString, bool bFindWholeWord, ESearchCase::Type InSearchCase)>;
+		FOnRemoveNotify OnRemoveNotify;
+	};
+
+	// Register an asset type that can contain named anim notifies
+	virtual void RegisterNotifyHostAsset(const FTopLevelAssetPath& InAssetClassPath, const FNotifyHostAssetParameters& InParameters);
+
+	// Unregister an asset type that can contain named anim notifies
+	virtual void UnregisterNotifyHostAsset(const FTopLevelAssetPath& InAssetClassPath);
+
+	// Find registered notify host parameters by asset class
+	virtual const FNotifyHostAssetParameters* FindNotifyHostAsset(UObject* InAsset);
+
+	// Get all the registered notify host class paths
+	virtual TArray<FTopLevelAssetPath> GetAllNotifyHostAssetClassPaths() const;
+
+	// Get the version information for the notify host asset registry
+	virtual uint32 GetNotifyHostAssetVersion() const;
+
 private:
 	/** When a new anim notify blueprint is created, this will handle post creation work such as adding non-event default nodes */
 	void HandleNewAnimNotifyBlueprintCreated(UBlueprint* InBlueprint);
@@ -526,5 +556,11 @@ private:
 
 	/** Delegate broadcast to register common tabs */
 	FOnRegisterTabs OnRegisterTabsDelegate;
+
+	/** Parameters for notify hosts */
+	TMap<FTopLevelAssetPath, FNotifyHostAssetParameters> NotifyHostParameters;
+
+	/** Version counter for NotifyHostParameters */
+	uint32 NotifyHostAssetVersion = 0;
 };
 
diff --git a/Engine/Source/Runtime/Engine/Public/Animation/AnimNotifyQueue.h b/Engine/Source/Runtime/Engine/Public/Animation/AnimNotifyQueue.h
--- a/Engine/Source/Runtime/Engine/Public/Animation/AnimNotifyQueue.h
+++ b/Engine/Source/Runtime/Engine/Public/Animation/AnimNotifyQueue.h
@@ -176,7 +176,7 @@ struct FAnimNotifyQueue
 	void AddAnimNotify(const FAnimNotifyEvent* Notify, const UObject* NotifySource);
 
 	/** Add anim notifies **/
-	void AddAnimNotifies(bool bSrcIsLeader, const TArray<FAnimNotifyEventReference>& NewNotifies, const float InstanceWeight);
+	ENGINE_API void AddAnimNotifies(bool bSrcIsLeader, const TArray<FAnimNotifyEventReference>& NewNotifies, const float InstanceWeight);
 
 	/** Add anim notifies from montage**/
 	void AddAnimNotifies(bool bSrcIsLeader, const TMap<FName, TArray<FAnimNotifyEventReference>>& NewNotifies, const float InstanceWeight);
