AnimNext: Fix use-after free when processing notify state events

Also fixed macro usage in GENERATE_ANIM_TRAIT_IMPLEMENTATION setup

diff --git a/Engine/Plugins/Experimental/AnimNextAnimGraph/Source/AnimNextAnimGraph/Private/Traits/NotifyDispatcher.cpp b/Engine/Plugins/Experimental/AnimNextAnimGraph/Source/AnimNextAnimGraph/Private/Traits/NotifyDispatcher.cpp
--- a/Engine/Plugins/Experimental/AnimNextAnimGraph/Source/AnimNextAnimGraph/Private/Traits/NotifyDispatcher.cpp
+++ b/Engine/Plugins/Experimental/AnimNextAnimGraph/Source/AnimNextAnimGraph/Private/Traits/NotifyDispatcher.cpp
@@ -23,6 +23,7 @@ namespace UE::AnimNext
 		GeneratorMacroRequired(INotifySource) \
 
 	GENERATE_ANIM_TRAIT_IMPLEMENTATION(FNotifyDispatcherTrait, TRAIT_INTERFACE_ENUMERATOR, TRAIT_REQUIRED_INTERFACE_ENUMERATOR, NULL_ANIM_TRAIT_EVENT_ENUMERATOR)
+	#undef TRAIT_REQUIRED_INTERFACE_ENUMERATOR
 	#undef TRAIT_INTERFACE_ENUMERATOR
 
 	void FNotifyDispatcherTrait::AdvanceBy(FExecutionContext& Context, const TTraitBinding<ITimeline>& Binding, float DeltaTime) const
@@ -137,7 +138,7 @@ void FModuleInstanceComponent_NotifyDispatcher::OnEndExecution(float InDeltaTime
 	// Logic from here largely mimics that of UAnimInstance::TriggerAnimNotifies 
 
 	// Array that will replace the 'ActiveAnimNotifyState' at the end of this function.
-	TArray<FAnimNotifyEvent> NewActiveAnimNotifyState;
+	TArray<const FAnimNotifyEvent*> NewActiveAnimNotifyState;
 	NewActiveAnimNotifyState.Reserve(NotifyQueue.AnimNotifies.Num());
 
 	TArray<FAnimNotifyEventReference> NewActiveAnimNotifyEventReference;
@@ -154,9 +155,12 @@ void FModuleInstanceComponent_NotifyDispatcher::OnEndExecution(float InDeltaTime
 			// AnimNotifyState
 			if (AnimNotifyEvent->NotifyStateClass)
 			{
-				int32 ExistingItemIndex = INDEX_NONE;
-
-				if (ActiveAnimNotifyState.Find(*AnimNotifyEvent, ExistingItemIndex))
+				int32 ExistingItemIndex = ActiveAnimNotifyState.IndexOfByPredicate([AnimNotifyEvent](const FAnimNotifyEvent* InOther)
+				{
+					return *AnimNotifyEvent == *InOther;
+				});
+				
+				if (ExistingItemIndex != INDEX_NONE)
 				{
 					check(ActiveAnimNotifyState.Num() == ActiveAnimNotifyEventReference.Num());
 					ActiveAnimNotifyState.RemoveAtSwap(ExistingItemIndex, EAllowShrinking::No);
@@ -168,9 +172,9 @@ void FModuleInstanceComponent_NotifyDispatcher::OnEndExecution(float InDeltaTime
 					NotifyStateBeginEventReference.Add(&NotifyQueue.AnimNotifies[Index]);
 				}
 
-				NewActiveAnimNotifyState.Add(*AnimNotifyEvent);
+				NewActiveAnimNotifyState.Add(AnimNotifyEvent);
 				FAnimNotifyEventReference& EventRef = NewActiveAnimNotifyEventReference.Add_GetRef(NotifyQueue.AnimNotifies[Index]);
-				EventRef.SetNotify(&NewActiveAnimNotifyState.Top());
+				EventRef.SetNotify(AnimNotifyEvent);
 				continue;
 			}
 
@@ -182,13 +186,13 @@ void FModuleInstanceComponent_NotifyDispatcher::OnEndExecution(float InDeltaTime
 	// Send end notification to AnimNotifyState not active anymore.
 	for (int32 Index = 0; Index < ActiveAnimNotifyState.Num(); ++Index)
 	{
-		const FAnimNotifyEvent& AnimNotifyEvent = ActiveAnimNotifyState[Index];
+		const FAnimNotifyEvent* AnimNotifyEvent = ActiveAnimNotifyState[Index];
 		const FAnimNotifyEventReference& EventReference = ActiveAnimNotifyEventReference[Index];
-		if (AnimNotifyEvent.NotifyStateClass)
+		if (AnimNotifyEvent->NotifyStateClass)
 		{
 #if WITH_EDITOR
 			// Prevent firing notifies in animation editors if requested
-			if(ModuleInstance.GetWorldType() != EWorldType::EditorPreview || AnimNotifyEvent.NotifyStateClass->ShouldFireInEditor())
+			if(ModuleInstance.GetWorldType() != EWorldType::EditorPreview || AnimNotifyEvent->NotifyStateClass->ShouldFireInEditor())
 #endif
 			{
 				NotifyQueueDispatchEvent->EventsToEnd.Add(EventReference);
@@ -220,13 +224,13 @@ void FModuleInstanceComponent_NotifyDispatcher::OnEndExecution(float InDeltaTime
 	// Tick currently active AnimNotifyState
 	for (int32 Index = 0; Index < ActiveAnimNotifyState.Num(); Index++)
 	{
-		const FAnimNotifyEvent& AnimNotifyEvent = ActiveAnimNotifyState[Index];
+		const FAnimNotifyEvent* AnimNotifyEvent = ActiveAnimNotifyState[Index];
 		const FAnimNotifyEventReference& EventReference = ActiveAnimNotifyEventReference[Index];
-		if(AnimNotifyEvent.NotifyStateClass)
+		if(AnimNotifyEvent->NotifyStateClass)
 		{
 #if WITH_EDITOR
 			// Prevent firing notifies in animation editors if requested 
-			if(ModuleInstance.GetWorldType() != EWorldType::EditorPreview || AnimNotifyEvent.NotifyStateClass->ShouldFireInEditor())
+			if(ModuleInstance.GetWorldType() != EWorldType::EditorPreview || AnimNotifyEvent->NotifyStateClass->ShouldFireInEditor())
 #endif
 			{
 				NotifyQueueDispatchEvent->EventsToTick.Add(EventReference);
diff --git a/Engine/Plugins/Experimental/AnimNextAnimGraph/Source/AnimNextAnimGraph/Private/Traits/NotifyDispatcher.h b/Engine/Plugins/Experimental/AnimNextAnimGraph/Source/AnimNextAnimGraph/Private/Traits/NotifyDispatcher.h
--- a/Engine/Plugins/Experimental/AnimNextAnimGraph/Source/AnimNextAnimGraph/Private/Traits/NotifyDispatcher.h
+++ b/Engine/Plugins/Experimental/AnimNextAnimGraph/Source/AnimNextAnimGraph/Private/Traits/NotifyDispatcher.h
@@ -54,8 +54,7 @@ struct FModuleInstanceComponent_NotifyDispatcher : public FModuleInstanceCompone
 	UPROPERTY()
 	FAnimNotifyQueue NotifyQueue;
 
-	UPROPERTY()
-	TArray<FAnimNotifyEvent> ActiveAnimNotifyState;
+	TArray<const FAnimNotifyEvent*> ActiveAnimNotifyState;
 
 	UPROPERTY()
 	TArray<FAnimNotifyEventReference> ActiveAnimNotifyEventReference;
