Relax trait latent property system to be more flexible

diff --git a/Engine/Plugins/Experimental/AnimNextAnimGraph/Source/AnimNextAnimGraph/Private/TraitCore/TraitReader.cpp b/Engine/Plugins/Experimental/AnimNextAnimGraph/Source/AnimNextAnimGraph/Private/TraitCore/TraitReader.cpp
--- a/Engine/Plugins/Experimental/AnimNextAnimGraph/Source/AnimNextAnimGraph/Private/TraitCore/TraitReader.cpp
+++ b/Engine/Plugins/Experimental/AnimNextAnimGraph/Source/AnimNextAnimGraph/Private/TraitCore/TraitReader.cpp
@@ -130,9 +130,11 @@ namespace UE::AnimNext
 			{
 				const FTraitTemplate& TraitTemplate = TraitTemplates[TraitIndex];
 
-				int32 NumLatentProperties = 0;
+				uint32 NumLatentProperties = 0;
 				*this << NumLatentProperties;
-				check(NumLatentProperties == TraitTemplate.GetNumLatentPropreties());
+
+				// We allow traits to reserve more latent property handles than they might need
+				check(NumLatentProperties <= TraitTemplate.GetNumLatentPropreties());
 
 				const uint32 BaseTraitIndex = TraitIndex - TraitTemplate.GetTraitIndex();
 				const FTraitTemplate& BaseTraitTemplate = TraitTemplates[BaseTraitIndex];
@@ -140,6 +142,7 @@ namespace UE::AnimNext
 				// Latent header is looked up using the base trait
 				FLatentPropertiesHeader& LatentHeader = BaseTraitTemplate.GetTraitLatentPropertiesHeader(*NodeDesc);
 				FLatentPropertyHandle* LatentHandles = TraitTemplate.GetTraitLatentPropertyHandles(*NodeDesc);
+				const FAnimNextTraitSharedData* SharedData = TraitTemplate.GetTraitDescription(*NodeDesc);
 
 				// If we are a base trait, reset our header
 				if (TraitTemplate.GetMode() == ETraitMode::Base)
@@ -159,7 +162,7 @@ namespace UE::AnimNext
 				const FTraitRegistryHandle TraitHandle = TraitTemplate.GetRegistryHandle();
 				const FTrait* Trait = TraitRegistry.Find(TraitHandle);
 
-				for (int32 LatentPropertyIndex = 0; LatentPropertyIndex < NumLatentProperties; ++LatentPropertyIndex)
+				for (uint32 LatentPropertyIndex = 0; LatentPropertyIndex < NumLatentProperties; ++LatentPropertyIndex)
 				{
 					FLatentPropertyMetadata Metadata;
 					*this << Metadata;
@@ -170,10 +173,10 @@ namespace UE::AnimNext
 
 					if (Trait != nullptr)
 					{
-						// If this property is valid, setup out binding for it
+						// If this property is valid, setup our binding for it
 						if (Metadata.RigVMIndex != MAX_uint16)
 						{
-							const FTraitLatentPropertyMemoryLayout PropertyMemoryLayout = Trait->GetLatentPropertyMemoryLayout(Metadata.Name, LatentPropertyIndex);
+							const FTraitLatentPropertyMemoryLayout PropertyMemoryLayout = Trait->GetLatentPropertyMemoryLayout(*SharedData, Metadata.Name, LatentPropertyIndex);
 
 							// Align our property
 							LatentPropertyOffset = Align(LatentPropertyOffset, PropertyMemoryLayout.Alignment);
diff --git a/Engine/Plugins/Experimental/AnimNextAnimGraph/Source/AnimNextAnimGraph/Private/TraitCore/TraitWriter.cpp b/Engine/Plugins/Experimental/AnimNextAnimGraph/Source/AnimNextAnimGraph/Private/TraitCore/TraitWriter.cpp
--- a/Engine/Plugins/Experimental/AnimNextAnimGraph/Source/AnimNextAnimGraph/Private/TraitCore/TraitWriter.cpp
+++ b/Engine/Plugins/Experimental/AnimNextAnimGraph/Source/AnimNextAnimGraph/Private/TraitCore/TraitWriter.cpp
@@ -203,7 +203,7 @@ namespace UE::AnimNext
 
 			TArray<FLatentPropertyMetadata> LatentProperties = Trait->GetLatentPropertyHandles(IsFilterEditorOnly(), GetTraitLatentPropertyIndexAt);
 
-			int32 NumLatentProperties = LatentProperties.Num();
+			uint32 NumLatentProperties = LatentProperties.Num();
 			*this << NumLatentProperties;
 
 			for (FLatentPropertyMetadata& Metadata : LatentProperties)
diff --git a/Engine/Plugins/Experimental/AnimNextAnimGraph/Source/AnimNextAnimGraph/Public/TraitCore/Trait.h b/Engine/Plugins/Experimental/AnimNextAnimGraph/Source/AnimNextAnimGraph/Public/TraitCore/Trait.h
--- a/Engine/Plugins/Experimental/AnimNextAnimGraph/Source/AnimNextAnimGraph/Public/TraitCore/Trait.h
+++ b/Engine/Plugins/Experimental/AnimNextAnimGraph/Source/AnimNextAnimGraph/Public/TraitCore/Trait.h
@@ -18,36 +18,47 @@
 
 class FArchive;
 
-// Helper macros
-// In the trait class declaration, this macro declares the Super alias and base functions we override
-// We use use FNV1A32 to hash entire TraitName, including 'F' part (Currently must be manually done). 
-#define DECLARE_ANIM_TRAIT(TraitName, TraitNameHash, SuperTraitName) \
+#define ANIM_NEXT_IMPL_DECLARE_ANIM_TRAIT_BASIC(TraitName, TraitNameHash, SuperTraitName) \
 	using TraitSuper = SuperTraitName; \
 	/* FTrait impl */ \
 	static constexpr UE::AnimNext::FTraitUID TraitUID = UE::AnimNext::FTraitUID(TraitNameHash, TEXT(#TraitName)); \
 	virtual UE::AnimNext::FTraitUID GetTraitUID() const override { return TraitUID; } \
 	virtual FString GetTraitName() const override { return TEXT(#TraitName); } \
+	static_assert(std::is_base_of<FAnimNextTraitSharedData, FSharedData>::value, "Trait shared data must derive from FAnimNextTraitSharedData"); \
+	static_assert(std::is_base_of<UE::AnimNext::FTraitInstanceData, FInstanceData>::value, "Trait instance data must derive from FTraitInstanceData");
+
+#define ANIM_NEXT_IMPL_DECLARE_ANIM_TRAIT_INSTANCING_SUPPORT() \
 	static const UE::AnimNext::FTraitMemoryLayout TraitMemoryDescription; \
 	virtual UE::AnimNext::FTraitMemoryLayout GetTraitMemoryDescription() const override { return TraitMemoryDescription; } \
 	virtual UScriptStruct* GetTraitSharedDataStruct() const override { return FSharedData::StaticStruct(); } \
 	virtual void ConstructTraitInstance(const UE::AnimNext::FExecutionContext& Context, const UE::AnimNext::FTraitBinding& Binding) const override; \
-	virtual void DestructTraitInstance(const UE::AnimNext::FExecutionContext& Context, const UE::AnimNext::FTraitBinding& Binding) const override; \
+	virtual void DestructTraitInstance(const UE::AnimNext::FExecutionContext& Context, const UE::AnimNext::FTraitBinding& Binding) const override;
+
+#define ANIM_NEXT_IMPL_DECLARE_ANIM_TRAIT_INTERFACE_SUPPORT() \
 	virtual const UE::AnimNext::ITraitInterface* GetTraitInterface(UE::AnimNext::FTraitInterfaceUID InterfaceUID) const override; \
-	virtual TConstArrayView<FTraitInterfaceUID> GetTraitInterfaces() const override; \
-	virtual TConstArrayView<FTraitInterfaceUID> GetTraitRequiredInterfaces() const override; \
+	virtual TConstArrayView<UE::AnimNext::FTraitInterfaceUID> GetTraitInterfaces() const override; \
+	virtual TConstArrayView<UE::AnimNext::FTraitInterfaceUID> GetTraitRequiredInterfaces() const override;
+
+#define ANIM_NEXT_IMPL_DECLARE_ANIM_TRAIT_EVENT_SUPPORT() \
 	virtual UE::AnimNext::ETraitStackPropagation OnTraitEvent(UE::AnimNext::FExecutionContext& Context, UE::AnimNext::FTraitBinding& Binding, FAnimNextTraitEvent& Event) const override; \
-	virtual TConstArrayView<FTraitEventUID> GetTraitEvents() const override; \
+	virtual TConstArrayView<UE::AnimNext::FTraitEventUID> GetTraitEvents() const override;
+
+#define ANIM_NEXT_IMPL_DECLARE_ANIM_TRAIT_LATENT_PROPERTY_SUPPORT() \
 	virtual uint32 GetNumLatentTraitProperties() const override { return -FSharedData::GetLatentPropertyIndex(~(size_t)0); } \
-	virtual FTraitLatentPropertyMemoryLayout GetLatentPropertyMemoryLayout(FName PropertyName, uint32 PropertyIndex) const override; \
-	static_assert(std::is_base_of<FAnimNextTraitSharedData, FSharedData>::value, "Trait shared data must derive from FAnimNextTraitSharedData"); \
-	static_assert(std::is_base_of<FTraitInstanceData, FInstanceData>::value, "Trait instance data must derive from FTraitInstanceData");
+	virtual UE::AnimNext::FTraitLatentPropertyMemoryLayout GetLatentPropertyMemoryLayout(const FAnimNextTraitSharedData& SharedData, FName PropertyName, uint32 PropertyIndex) const override;
+
+// Helper macros
+// In the trait class declaration, this macro declares the Super alias and base functions we override
+// We use use FNV1A32 to hash entire TraitName, including 'F' part (Currently must be manually done). 
+#define DECLARE_ANIM_TRAIT(TraitName, TraitNameHash, SuperTraitName) \
+	ANIM_NEXT_IMPL_DECLARE_ANIM_TRAIT_BASIC(TraitName, TraitNameHash, SuperTraitName) \
+	ANIM_NEXT_IMPL_DECLARE_ANIM_TRAIT_INSTANCING_SUPPORT() \
+	ANIM_NEXT_IMPL_DECLARE_ANIM_TRAIT_INTERFACE_SUPPORT() \
+	ANIM_NEXT_IMPL_DECLARE_ANIM_TRAIT_EVENT_SUPPORT() \
+	ANIM_NEXT_IMPL_DECLARE_ANIM_TRAIT_LATENT_PROPERTY_SUPPORT()
 
 #define DECLARE_ABSTRACT_ANIM_TRAIT(TraitName, TraitNameHash, SuperTraitName) \
-	using TraitSuper = SuperTraitName; \
-	/* FTrait impl */ \
-	static constexpr UE::AnimNext::FTraitUID TraitUID = UE::AnimNext::FTraitUID(TraitNameHash, TEXT(#TraitName)); \
-	virtual UE::AnimNext::FTraitUID GetTraitUID() const override { return TraitUID; } \
-	virtual FString GetTraitName() const override { return TEXT(#TraitName); }
+	ANIM_NEXT_IMPL_DECLARE_ANIM_TRAIT_BASIC(TraitName, TraitNameHash, SuperTraitName)
 
 // In the trait cpp, these three macros implement the base functionality
 // 
@@ -63,12 +74,6 @@ class FArchive;
 #define ANIM_NEXT_IMPL_DEFINE_ANIM_TRAIT(TraitName) \
 	const UE::AnimNext::FTraitMemoryLayout TraitName::TraitMemoryDescription = \
 		UE::AnimNext::FTraitMemoryLayout{ sizeof(TraitName), alignof(TraitName), sizeof(TraitName::FSharedData), alignof(TraitName::FSharedData), sizeof(TraitName::FInstanceData), alignof(TraitName::FInstanceData) }; \
-	FTraitLatentPropertyMemoryLayout TraitName::GetLatentPropertyMemoryLayout(FName PropertyName, uint32 PropertyIndex) const \
-	{ \
-		/* Thread safe cache initialization */ \
-		static TArray<FTraitLatentPropertyMemoryLayout> CachedLatentPropertyMemoryLayouts = [this](){ TArray<FTraitLatentPropertyMemoryLayout> Result; Result.SetNum(GetNumLatentTraitProperties()); return Result; }(); \
-		return GetLatentPropertyMemoryLayoutImpl(PropertyName, PropertyIndex, CachedLatentPropertyMemoryLayouts); \
-	} \
 	void TraitName::ConstructTraitInstance(const UE::AnimNext::FExecutionContext& Context, const UE::AnimNext::FTraitBinding& Binding) const \
 	{ \
 		/* Construct the base struct first */ \
@@ -87,6 +92,15 @@ class FArchive;
 		Data->~FInstanceData(); \
 	}
 
+// Implements GetLatentPropertyMemoryLayout()
+#define ANIM_NEXT_IMPL_DEFINE_ANIM_TRAIT_GET_LATENT_PROPERTY_MEMORY_LAYOUT(TraitName) \
+	UE::AnimNext::FTraitLatentPropertyMemoryLayout TraitName::GetLatentPropertyMemoryLayout(const FAnimNextTraitSharedData& SharedData, FName PropertyName, uint32 PropertyIndex) const \
+	{ \
+		/* Thread safe cache initialization */ \
+		static TArray<UE::AnimNext::FTraitLatentPropertyMemoryLayout> CachedLatentPropertyMemoryLayouts = [this](){ TArray<UE::AnimNext::FTraitLatentPropertyMemoryLayout> Result; Result.SetNum(GetNumLatentTraitProperties()); return Result; }(); \
+		return GetLatentPropertyMemoryLayoutImpl(PropertyName, PropertyIndex, CachedLatentPropertyMemoryLayouts); \
+	}
+
 // Helper that handles the GetTraitInterface() details for each interface specified by the generator macro
 #define ANIM_NEXT_IMPL_GET_INTERFACE_IMPL_FOR_INTERFACE(InterfaceName) \
 	if (InInterfaceUID == InterfaceName::InterfaceUID) \
@@ -208,6 +222,7 @@ namespace UE::AnimNext::Private
   */
 #define GENERATE_ANIM_TRAIT_IMPLEMENTATION(TraitName, InterfaceEnumeratorMacro, RequiredInterfaceEnumeratorMacro, EventEnumeratorMacro) \
 	ANIM_NEXT_IMPL_DEFINE_ANIM_TRAIT(TraitName) \
+	ANIM_NEXT_IMPL_DEFINE_ANIM_TRAIT_GET_LATENT_PROPERTY_MEMORY_LAYOUT(TraitName) \
 	ANIM_NEXT_IMPL_DEFINE_ANIM_TRAIT_GET_INTERFACE(TraitName, InterfaceEnumeratorMacro) \
 	ANIM_NEXT_IMPL_DEFINE_ANIM_TRAIT_GET_INTERFACES(TraitName, InterfaceEnumeratorMacro) \
 	ANIM_NEXT_IMPL_DEFINE_ANIM_TRAIT_GET_REQUIRED_INTERFACES(TraitName, RequiredInterfaceEnumeratorMacro) \
@@ -373,7 +388,7 @@ namespace UE::AnimNext
 		virtual uint32 GetNumLatentTraitProperties() const { return 0; }
 
 		// Returns the memory layout of the specified latent property
-		virtual FTraitLatentPropertyMemoryLayout GetLatentPropertyMemoryLayout(FName PropertyName, uint32 PropertyIndex) const { return FTraitLatentPropertyMemoryLayout(); }
+		virtual FTraitLatentPropertyMemoryLayout GetLatentPropertyMemoryLayout(const FAnimNextTraitSharedData& SharedData, FName PropertyName, uint32 PropertyIndex) const { return FTraitLatentPropertyMemoryLayout(); }
 
 		// Called to serialize trait shared data
 		virtual void SerializeTraitSharedData(FArchive& Ar, FAnimNextTraitSharedData& SharedData) const;
