Added some more timings

diff --git a/Engine/Plugins/Experimental/AnimNext/Source/AnimNext/Private/Graph/RigUnit_AnimNextRemapPose.cpp b/Engine/Plugins/Experimental/AnimNext/Source/AnimNext/Private/Graph/RigUnit_AnimNextRemapPose.cpp
--- a/Engine/Plugins/Experimental/AnimNext/Source/AnimNext/Private/Graph/RigUnit_AnimNextRemapPose.cpp
+++ b/Engine/Plugins/Experimental/AnimNext/Source/AnimNext/Private/Graph/RigUnit_AnimNextRemapPose.cpp
@@ -9,9 +9,12 @@
 #include "Engine/SkeletalMesh.h"
 
 #include UE_INLINE_GENERATED_CPP_BY_NAME(RigUnit_AnimNextRemapPose)
+DEFINE_STAT(STAT_AnimNext_RigUnit_RemapPose);
 
 FRigUnit_AnimNextRemapPose_Execute()
 {
+	SCOPE_CYCLE_COUNTER(STAT_AnimNext_RigUnit_RemapPose);
+
 	using namespace UE::AnimNext;
 
 	if(!TargetAnimGraphRefPose.ReferencePose.IsValid())
diff --git a/Engine/Plugins/Experimental/AnimNext/Source/AnimNext/Private/Graph/RigUnit_AnimNextRemapPose.h b/Engine/Plugins/Experimental/AnimNext/Source/AnimNext/Private/Graph/RigUnit_AnimNextRemapPose.h
--- a/Engine/Plugins/Experimental/AnimNext/Source/AnimNext/Private/Graph/RigUnit_AnimNextRemapPose.h
+++ b/Engine/Plugins/Experimental/AnimNext/Source/AnimNext/Private/Graph/RigUnit_AnimNextRemapPose.h
@@ -6,9 +6,12 @@
 #include "Graph/AnimNext_LODPose.h"
 #include "Graph/RigUnit_AnimNextBase.h"
 #include "RemapPoseData.h"
+#include "AnimNextStats.h"
 
 #include "RigUnit_AnimNextRemapPose.generated.h"
 
+DECLARE_CYCLE_STAT_EXTERN(TEXT("AnimNext RigUnit: Remap Pose"), STAT_AnimNext_RigUnit_RemapPose, STATGROUP_AnimNext, ANIMNEXT_API);
+
 /** Remaps an anim graph pose from to another skeletal mesh component. */
 USTRUCT(meta=(DisplayName="Remap Pose", Category="Animation Graph", NodeColor="0, 1, 1", Keywords="Output,Pose,Port"))
 struct FRigUnit_AnimNextRemapPose : public FRigUnit_AnimNextBase
diff --git a/Engine/Plugins/Experimental/AnimNextAnimGraph/Source/AnimNextAnimGraph/Private/EvaluationVM/Tasks/PushPose.cpp b/Engine/Plugins/Experimental/AnimNextAnimGraph/Source/AnimNextAnimGraph/Private/EvaluationVM/Tasks/PushPose.cpp
--- a/Engine/Plugins/Experimental/AnimNextAnimGraph/Source/AnimNextAnimGraph/Private/EvaluationVM/Tasks/PushPose.cpp
+++ b/Engine/Plugins/Experimental/AnimNextAnimGraph/Source/AnimNextAnimGraph/Private/EvaluationVM/Tasks/PushPose.cpp
@@ -3,6 +3,8 @@
 #include "EvaluationVM/Tasks/PushPose.h"
 #include "EvaluationVM/EvaluationVM.h"
 
+DEFINE_STAT(STAT_AnimNext_Task_PushPose);
+
 FAnimNextPushPoseTask::FAnimNextPushPoseTask(const FAnimNextGraphLODPose* InGraphPose)
 	: GraphPose(InGraphPose)
 {
@@ -21,6 +23,8 @@ FAnimNextPushPoseTask& FAnimNextPushPoseTask::operator=(const FAnimNextPushPoseT
 
 void FAnimNextPushPoseTask::Execute(UE::AnimNext::FEvaluationVM& VM) const
 {
+	SCOPE_CYCLE_COUNTER(STAT_AnimNext_Task_PushPose);
+
 	using namespace UE::AnimNext;
 
 	FKeyframeState Keyframe = VM.MakeUninitializedKeyframe(GraphPose->LODPose.IsAdditive());
diff --git a/Engine/Plugins/Experimental/AnimNextAnimGraph/Source/AnimNextAnimGraph/Public/EvaluationVM/Tasks/PushPose.h b/Engine/Plugins/Experimental/AnimNextAnimGraph/Source/AnimNextAnimGraph/Public/EvaluationVM/Tasks/PushPose.h
--- a/Engine/Plugins/Experimental/AnimNextAnimGraph/Source/AnimNextAnimGraph/Public/EvaluationVM/Tasks/PushPose.h
+++ b/Engine/Plugins/Experimental/AnimNextAnimGraph/Source/AnimNextAnimGraph/Public/EvaluationVM/Tasks/PushPose.h
@@ -5,9 +5,12 @@
 #include "CoreMinimal.h"
 #include "EvaluationVM/EvaluationTask.h"
 #include "Graph/AnimNext_LODPose.h"
+#include "AnimNextStats.h"
 
 #include "PushPose.generated.h"
 
+DECLARE_CYCLE_STAT_EXTERN(TEXT("AnimNext Task: PushPose"), STAT_AnimNext_Task_PushPose, STATGROUP_AnimNext, ANIMNEXTANIMGRAPH_API);
+
 USTRUCT()
 struct ANIMNEXTANIMGRAPH_API FAnimNextPushPoseTask : public FAnimNextEvaluationTask
 {
