Avoid reserving excessive size since we have an inline allocator

diff --git a/Engine/Plugins/Experimental/AnimNextAnimGraph/Source/AnimNextAnimGraph/Private/TraitInterfaces/IEvaluate.cpp b/Engine/Plugins/Experimental/AnimNextAnimGraph/Source/AnimNextAnimGraph/Private/TraitInterfaces/IEvaluate.cpp
--- a/Engine/Plugins/Experimental/AnimNextAnimGraph/Source/AnimNextAnimGraph/Private/TraitInterfaces/IEvaluate.cpp
+++ b/Engine/Plugins/Experimental/AnimNextAnimGraph/Source/AnimNextAnimGraph/Private/TraitInterfaces/IEvaluate.cpp
@@ -112,7 +112,6 @@ namespace UE::AnimNext
 		FMemMark Mark(MemStack);
 
 		FChildrenArray Children;
-		Children.Reserve(16);
 
 		FEvaluateTraversalContext TraversalContext(EvaluationProgram);
 		TraversalContext.BindTo(GraphRootPtr);
