Ensure AnimNext trait shared/instance data derives from the same type on its base class

diff --git a/Engine/Plugins/Experimental/AnimNextAnimGraph/Source/AnimNextAnimGraph/Public/TraitCore/Trait.h b/Engine/Plugins/Experimental/AnimNextAnimGraph/Source/AnimNextAnimGraph/Public/TraitCore/Trait.h
--- a/Engine/Plugins/Experimental/AnimNextAnimGraph/Source/AnimNextAnimGraph/Public/TraitCore/Trait.h
+++ b/Engine/Plugins/Experimental/AnimNextAnimGraph/Source/AnimNextAnimGraph/Public/TraitCore/Trait.h
@@ -24,8 +24,8 @@ class FArchive;
 	static constexpr UE::AnimNext::FTraitUID TraitUID = UE::AnimNext::FTraitUID(TraitNameHash, TEXT(#TraitName)); \
 	virtual UE::AnimNext::FTraitUID GetTraitUID() const override { return TraitUID; } \
 	virtual FString GetTraitName() const override { return TEXT(#TraitName); } \
-	static_assert(std::is_base_of<FAnimNextTraitSharedData, FSharedData>::value, "Trait shared data must derive from FAnimNextTraitSharedData"); \
-	static_assert(std::is_base_of<UE::AnimNext::FTraitInstanceData, FInstanceData>::value, "Trait instance data must derive from FTraitInstanceData");
+	static_assert(std::is_base_of<TraitSuper::FSharedData, FSharedData>::value, "Trait shared data must derive from its Super's FSharedData"); \
+	static_assert(std::is_base_of<TraitSuper::FInstanceData, FInstanceData>::value, "Trait instance data must derive from Super's FInstanceData");
 
 #define ANIM_NEXT_IMPL_DECLARE_ANIM_TRAIT_INSTANCING_SUPPORT() \
 	static const UE::AnimNext::FTraitMemoryLayout TraitMemoryDescription; \
@@ -318,8 +318,8 @@ namespace UE::AnimNext
 
 		// Empty shared/instance data types
 		// Derived types must define an alias for these
-		using FSharedData = FAnimNextTraitSharedData;
-		using FInstanceData = FTraitInstanceData;
+		using FSharedData = FAnimNextTraitSharedData;			// Trait shared data must derive from FAnimNextTraitSharedData
+		using FInstanceData = FTraitInstanceData;				// Trait instance data must derive from FTraitInstanceData
 
 		// The globally unique UID for this trait
 		// Derived types will have their own TraitUID member that hides/aliases/shadows this one
