Remove unnecessary padding from trait shared/instance data when they are empty

diff --git a/Engine/Plugins/Experimental/AnimNextAnimGraph/Source/AnimNextAnimGraph/Public/TraitCore/Trait.h b/Engine/Plugins/Experimental/AnimNextAnimGraph/Source/AnimNextAnimGraph/Public/TraitCore/Trait.h
--- a/Engine/Plugins/Experimental/AnimNextAnimGraph/Source/AnimNextAnimGraph/Public/TraitCore/Trait.h
+++ b/Engine/Plugins/Experimental/AnimNextAnimGraph/Source/AnimNextAnimGraph/Public/TraitCore/Trait.h
@@ -73,7 +73,12 @@ class FArchive;
 // Implements various parts of FTrait
 #define ANIM_NEXT_IMPL_DEFINE_ANIM_TRAIT(TraitName) \
 	const UE::AnimNext::FTraitMemoryLayout TraitName::TraitMemoryDescription = \
-		UE::AnimNext::FTraitMemoryLayout{ sizeof(TraitName), alignof(TraitName), sizeof(TraitName::FSharedData), alignof(TraitName::FSharedData), sizeof(TraitName::FInstanceData), alignof(TraitName::FInstanceData) }; \
+		UE::AnimNext::FTraitMemoryLayout{ \
+			sizeof(TraitName), alignof(TraitName), \
+			/* Use the effective size to ensure optimal packing for empty shared/instance data */ \
+			UE::AnimNext::Private::TGetEffectiveSize<TraitName::FSharedData>::Value, alignof(TraitName::FSharedData), \
+			UE::AnimNext::Private::TGetEffectiveSize<TraitName::FInstanceData>::Value, alignof(TraitName::FInstanceData) \
+		}; \
 	void TraitName::ConstructTraitInstance(const UE::AnimNext::FExecutionContext& Context, const UE::AnimNext::FTraitBinding& Binding) const \
 	{ \
 		/* Construct the base struct first */ \
@@ -175,6 +180,21 @@ namespace UE::AnimNext::Private
 	{
 		using EventType = std::remove_const_t<HandlerEventType>;
 	};
+
+	template<class BaseType>
+	struct TPaddedDerivedType : public BaseType { uint8 Dummy; };
+
+	// Helper to find the effective size of a struct
+	// If a struct contains no members, its size as defined by C++ is 1 byte to allow for pointer arithmetic
+	// However, if we don't use pointer arithmetic (e.g. have them in an array), then we can avoid this waste
+	// by computing its effective size. The effective size of a struct is its size when you derive from it.
+	// A struct that has no members ends up with an effective size of zero bytes.
+	template<class StructType>
+	struct TGetEffectiveSize
+	{
+		// The effective size of the specified struct
+		static constexpr uint32 Value = sizeof(TPaddedDerivedType<StructType>) == sizeof(uint8) ? 0 : sizeof(StructType);
+	};
 }
 
 // Helper that handles the OnTraitEvent() details for each event specified by the generator macro
