AnimNext: Injection

diff --git a/Engine/Plugins/Experimental/AnimNextAnimGraph/Source/AnimNextAnimGraph/Internal/PlayAnim/PlayAnimCallbackProxy.h b/Engine/Plugins/Experimental/AnimNextAnimGraph/Source/AnimNextAnimGraph/Internal/PlayAnim/PlayAnimCallbackProxy.h
--- a/Engine/Plugins/Experimental/AnimNextAnimGraph/Source/AnimNextAnimGraph/Internal/PlayAnim/PlayAnimCallbackProxy.h
+++ b/Engine/Plugins/Experimental/AnimNextAnimGraph/Source/AnimNextAnimGraph/Internal/PlayAnim/PlayAnimCallbackProxy.h
@@ -2,10 +2,13 @@
 #pragma once
 
 #include "CoreMinimal.h"
-#include "PlayAnim/PlayAnimRequest.h"
+#include "Injection/InjectionRequest.h"
 
 #include "PlayAnimCallbackProxy.generated.h"
 
+class UAnimSequence;
+class UAnimNextComponent;
+
 DECLARE_DYNAMIC_MULTICAST_DELEGATE(FOnPlayAnimPlayDelegate);
 
 UCLASS(MinimalAPI)
@@ -28,56 +31,49 @@ class UPlayAnimCallbackProxy : public UObject
 	// Called to perform the query internally
 	UFUNCTION(BlueprintCallable, meta = (BlueprintInternalUseOnly = "true"))
 	static ANIMNEXTANIMGRAPH_API UPlayAnimCallbackProxy* CreateProxyObjectForPlayAnim(
-		class UAnimNextComponent* AnimNextComponent,
-		FName SlotName,
-		class UAnimSequence* AnimSequence,
+		UAnimNextComponent* AnimNextComponent,
+		FName SiteName,
+		UAnimSequence* AnimSequence,
 		float PlayRate = 1.0f,
 		float StartPosition = 0.0f,
-		FAnimNextPlayAnimBlendSettings BlendInSettings = FAnimNextPlayAnimBlendSettings(),
-		FAnimNextPlayAnimBlendSettings BlendOutSettings = FAnimNextPlayAnimBlendSettings());
+		FAnimNextInjectionBlendSettings BlendInSettings = FAnimNextInjectionBlendSettings(),
+		FAnimNextInjectionBlendSettings BlendOutSettings = FAnimNextInjectionBlendSettings());
 
-	// Called to perform the query internally
-	UFUNCTION(BlueprintCallable, meta = (BlueprintInternalUseOnly = "true"))
+	UFUNCTION(BlueprintCallable, meta=(BlueprintInternalUseOnly = "true", DeprecatedFunction, DeprecatedMessage="Please use Inject instead."))
 	static ANIMNEXTANIMGRAPH_API UPlayAnimCallbackProxy* CreateProxyObjectForPlayAsset(
-		class UAnimNextComponent* AnimNextComponent,
-		FName SlotName,
-		UObject* Asset,
+		UAnimNextComponent* AnimNextComponent,
+		FName SiteName,
+		UObject* Object,
 		const FInstancedStruct& Payload,
-		FAnimNextPlayAnimBlendSettings BlendInSettings = FAnimNextPlayAnimBlendSettings(),
-		FAnimNextPlayAnimBlendSettings BlendOutSettings = FAnimNextPlayAnimBlendSettings());
-	
+		FAnimNextInjectionBlendSettings BlendInSettings = FAnimNextInjectionBlendSettings(),
+		FAnimNextInjectionBlendSettings BlendOutSettings = FAnimNextInjectionBlendSettings())
+	{
+		return nullptr;
+	}
+
 public:
 	//~ Begin UObject Interface
 	ANIMNEXTANIMGRAPH_API virtual void BeginDestroy() override;
 	//~ End UObject Interface
 
 protected:
-	ANIMNEXTANIMGRAPH_API void OnPlayAnimCompleted(const UE::AnimNext::FPlayAnimRequest& Request);
-	ANIMNEXTANIMGRAPH_API void OnPlayAnimInterrupted(const UE::AnimNext::FPlayAnimRequest& Request);
-	ANIMNEXTANIMGRAPH_API void OnPlayAnimBlendingOut(const UE::AnimNext::FPlayAnimRequest& Request);
+	ANIMNEXTANIMGRAPH_API void OnPlayAnimCompleted(const UE::AnimNext::FInjectionRequest& Request);
+	ANIMNEXTANIMGRAPH_API void OnPlayAnimInterrupted(const UE::AnimNext::FInjectionRequest& Request);
+	ANIMNEXTANIMGRAPH_API void OnPlayAnimBlendingOut(const UE::AnimNext::FInjectionRequest& Request);
 
-	// Attempts to play an object with the specified payload. Returns whether it started or not.
-	ANIMNEXTANIMGRAPH_API bool Play(
-		class UAnimNextComponent* AnimNextComponent,
-		FName SlotName,
-		UObject* Object,
-		FInstancedStruct&& Payload,
-		const UE::AnimNext::FPlayAnimBlendSettings& BlendInSettings,
-		const UE::AnimNext::FPlayAnimBlendSettings& BlendOutSettings);
-	
 	// Attempts to play an animation with the specified settings. Returns whether it started or not.
 	ANIMNEXTANIMGRAPH_API bool Play(
-		class UAnimNextComponent* AnimNextComponent,
-		FName SlotName,
-		class UAnimSequence* AnimSequence,
+		UAnimNextComponent* AnimNextComponent,
+		FName SiteName,
+		UAnimSequence* AnimSequence,
 		float PlayRate,
 		float StartPosition,
-		const UE::AnimNext::FPlayAnimBlendSettings& BlendInSettings,
-		const UE::AnimNext::FPlayAnimBlendSettings& BlendOutSettings);
+		const UE::AnimNext::FInjectionBlendSettings& BlendInSettings,
+		const UE::AnimNext::FInjectionBlendSettings& BlendOutSettings);
 
 private:
 	void Reset();
 
-	UE::AnimNext::FPlayAnimRequestPtr PlayingRequest;
+	UE::AnimNext::FInjectionRequestPtr PlayingRequest;
 	bool bWasInterrupted = false;
 };
diff --git a/Engine/Plugins/Experimental/AnimNextAnimGraph/Source/AnimNextAnimGraph/Internal/Traits/BlendStackTrait.h b/Engine/Plugins/Experimental/AnimNextAnimGraph/Source/AnimNextAnimGraph/Internal/Traits/BlendStackTrait.h
--- a/Engine/Plugins/Experimental/AnimNextAnimGraph/Source/AnimNextAnimGraph/Internal/Traits/BlendStackTrait.h
+++ b/Engine/Plugins/Experimental/AnimNextAnimGraph/Source/AnimNextAnimGraph/Internal/Traits/BlendStackTrait.h
@@ -3,7 +3,6 @@
 #pragma once
 
 #include "Graph/AnimNextAnimationGraph.h"
-#include "Graph/AnimNextGraphInstancePtr.h"
 
 #include "TraitCore/Trait.h"
 #include "TraitInterfaces/IBlendStack.h"
@@ -12,6 +11,7 @@
 #include "TraitInterfaces/IHierarchy.h"
 #include "TraitInterfaces/IUpdate.h"
 #include "TraitInterfaces/IAttributeProvider.h"
+#include "TraitInterfaces/IInertializerBlend.h"
 #include "TraitInterfaces/ITimeline.h"
 #include "TraitInterfaces/ISmoothBlend.h"
 
@@ -86,7 +86,7 @@ namespace UE::AnimNext
 	 * 
 	 * Used as base trait in state machine-like setups where state machine is the additive trait that pushes graphs and blend settings to the blend stack.
 	 */
-	struct FBlendStackCoreTrait : FBaseTrait, IUpdateTraversal, IHierarchy, IDiscreteBlend, ISmoothBlend, IAttributeProvider, ITimeline, IGarbageCollection, IBlendStack
+	struct FBlendStackCoreTrait : FBaseTrait, IUpdate, IUpdateTraversal, IHierarchy, IDiscreteBlend, ISmoothBlend, IInertializerBlend, IAttributeProvider, ITimeline, IGarbageCollection, IBlendStack
 	{
 		DECLARE_ANIM_TRAIT(FBlendStackCoreTrait, 0xda7b7f8d, FBaseTrait)
 
@@ -135,6 +135,9 @@ namespace UE::AnimNext
 		// Internal impl
 		static int32 FindFreeGraphIndexOrAdd(FInstanceData& InstanceData);
 
+		// IUpdate impl
+		virtual void PreUpdate(FUpdateTraversalContext& Context, const TTraitBinding<IUpdate>& Binding, const FTraitUpdateState& TraitState) const override;
+
 		// IUpdateTraversal impl
 		virtual void QueueChildrenForTraversal(FUpdateTraversalContext& Context, const TTraitBinding<IUpdateTraversal>& Binding, const FTraitUpdateState& TraitState, FUpdateTraversalQueue& TraversalQueue) const override;
 
@@ -143,12 +146,15 @@ namespace UE::AnimNext
 		virtual void GetChildren(const FExecutionContext& Context, const TTraitBinding<IHierarchy>& Binding, FChildrenArray& Children) const override;
 
 		// IDiscreteBlend impl
+		virtual float GetBlendWeight(FExecutionContext& Context, const TTraitBinding<IDiscreteBlend>& Binding, int32 ChildIndex) const override;
 		virtual int32 GetBlendDestinationChildIndex(FExecutionContext& Context, const TTraitBinding<IDiscreteBlend>& Binding) const override;
 		virtual void OnBlendTransition(FExecutionContext& Context, const TTraitBinding<IDiscreteBlend>& Binding, int32 OldChildIndex, int32 NewChildIndex) const override;
 		virtual void OnBlendInitiated(FExecutionContext& Context, const TTraitBinding<IDiscreteBlend>& Binding, int32 ChildIndex) const override;
 		virtual void OnBlendTerminated(FExecutionContext& Context, const TTraitBinding<IDiscreteBlend>& Binding, int32 ChildIndex) const override;
 
-		//@TODO: Add inertial blend.
+		// IInertializerBlend impl
+		virtual float GetBlendTime(FExecutionContext& Context, const TTraitBinding<IInertializerBlend>& Binding, int32 ChildIndex) const override;
+
 		// ISmoothBlend impl
 		virtual float GetBlendTime(FExecutionContext& Context, const TTraitBinding<ISmoothBlend>& Binding, int32 ChildIndex) const override;
 		virtual EAlphaBlendOption GetBlendType(FExecutionContext& Context, const TTraitBinding<ISmoothBlend>& Binding, int32 ChildIndex) const override;
@@ -164,8 +170,9 @@ namespace UE::AnimNext
 		virtual void AddReferencedObjects(const FExecutionContext& Context, const TTraitBinding<IGarbageCollection>& Binding, FReferenceCollector& Collector) const override;
 
 		// IBlendStack impl
-		virtual void PushGraph(FExecutionContext& Context, const TTraitBinding<IBlendStack>& Binding, IBlendStack::FGraphRequest&& GraphRequest) const override;
-		virtual void GetActiveGraph(FExecutionContext& Context, const TTraitBinding<IBlendStack>& Binding, IBlendStack::FGraphRequestPtr& OutGraphRequest) const override;
+		virtual int32 PushGraph(FExecutionContext& Context, const TTraitBinding<IBlendStack>& Binding, IBlendStack::FGraphRequest&& GraphRequest) const override;
+		virtual int32 GetActiveGraph(FExecutionContext& Context, const TTraitBinding<IBlendStack>& Binding, IBlendStack::FGraphRequestPtr& OutGraphRequest) const override;
+		virtual IBlendStack::FGraphRequestPtr GetGraph(FExecutionContext& Context, const TTraitBinding<IBlendStack>& Binding, int32 InChildIndex) const override;
 
 #if WITH_EDITOR
 		virtual bool IsHidden() const override { return true; }
@@ -178,7 +185,7 @@ namespace UE::AnimNext
 	 * Standalone blend stack that exposes desired object and blend time as pins.
 	 * A new graph will be pushed every time the newest object does not match the pin (or bForceBlend is true).
 	 */	
-	struct FBlendStackTrait : FBlendStackCoreTrait, IUpdate
+	struct FBlendStackTrait : FBlendStackCoreTrait
 	{
 		DECLARE_ANIM_TRAIT(FBlendStackTrait, 0x46678513, FBlendStackCoreTrait)
 
diff --git a/Engine/Plugins/Experimental/AnimNextAnimGraph/Source/AnimNextAnimGraph/Private/AnimNextAnimGraphModule.cpp b/Engine/Plugins/Experimental/AnimNextAnimGraph/Source/AnimNextAnimGraph/Private/AnimNextAnimGraphModule.cpp
--- a/Engine/Plugins/Experimental/AnimNextAnimGraph/Source/AnimNextAnimGraph/Private/AnimNextAnimGraphModule.cpp
+++ b/Engine/Plugins/Experimental/AnimNextAnimGraph/Source/AnimNextAnimGraph/Private/AnimNextAnimGraphModule.cpp
@@ -13,6 +13,7 @@
 #include "Graph/AnimNextAnimationGraph.h"
 #include "Graph/AnimNext_LODPose.h"
 #include "Graph/RigVMTrait_AnimNextPublicVariables.h"
+#include "Injection/InjectionRequest.h"
 #include "TraitCore/NodeTemplateRegistry.h"
 #include "TraitCore/TraitInterfaceRegistry.h"
 #include "TraitCore/TraitRegistry.h"
@@ -62,7 +63,8 @@ void FAnimNextAnimGraphModule::StartupModule()
 	static UScriptStruct* const AllowedStructTypes[] =
 	{
 		FAnimNextGraphReferencePose::StaticStruct(),
-		FRigVMTrait_AnimNextPublicVariables::StaticStruct()
+		FRigVMTrait_AnimNextPublicVariables::StaticStruct(),
+		FAnimNextInjectionBlendSettings::StaticStruct()
 	};
 
 	RigVMRegistry.RegisterStructTypes(AllowedStructTypes);
diff --git a/Engine/Plugins/Experimental/AnimNextAnimGraph/Source/AnimNextAnimGraph/Private/AnimNextAnimGraphSettings.cpp b/Engine/Plugins/Experimental/AnimNextAnimGraph/Source/AnimNextAnimGraph/Private/AnimNextAnimGraphSettings.cpp
--- a/Engine/Plugins/Experimental/AnimNextAnimGraph/Source/AnimNextAnimGraph/Private/AnimNextAnimGraphSettings.cpp
+++ b/Engine/Plugins/Experimental/AnimNextAnimGraph/Source/AnimNextAnimGraph/Private/AnimNextAnimGraphSettings.cpp
@@ -220,6 +220,16 @@ void UAnimNextAnimGraphSettings::LoadAndRebuildMappings()
 	AnimGraphModule.LoadedGraphs.Reset();
 	AssetGraphMap.Reset();
 
+	const UAnimNextAnimationGraph* DefaultAnimationGraph = DefaultRunGraphHost.LoadSynchronous();
+	if(DefaultAnimationGraph == nullptr)
+	{
+		UE_LOGFMT(LogAnimation, Error, "UAnimNextAnimGraphSettings::LoadAndRebuildMappings: Could not load default host animation graph '{GraphName}'", DefaultRunGraphHost.ToString());
+	}
+	else
+	{
+		AnimGraphModule.LoadedGraphs.Add(DefaultAnimationGraph);
+	}
+
 	for(TEnumerateRef<FAnimNextAssetGraphMapping> Mapping : EnumerateRange(AssetGraphMappings))
 	{
 		const UAnimNextAnimationGraph* AnimationGraph = Mapping->AnimationGraph.LoadSynchronous();
diff --git a/Engine/Plugins/Experimental/AnimNextAnimGraph/Source/AnimNextAnimGraph/Private/Graph/AnimNextAnimationGraph.cpp b/Engine/Plugins/Experimental/AnimNextAnimGraph/Source/AnimNextAnimGraph/Private/Graph/AnimNextAnimationGraph.cpp
--- a/Engine/Plugins/Experimental/AnimNextAnimGraph/Source/AnimNextAnimGraph/Private/Graph/AnimNextAnimationGraph.cpp
+++ b/Engine/Plugins/Experimental/AnimNextAnimGraph/Source/AnimNextAnimGraph/Private/Graph/AnimNextAnimationGraph.cpp
@@ -12,6 +12,7 @@
 #include "Serialization/MemoryReader.h"
 #include "AnimNextAnimGraphStats.h"
 #include "Graph/AnimNextGraphEntryPoint.h"
+#include "Graph/RigUnit_AnimNextGraphEvaluator.h"
 #include "UObject/LinkerLoad.h"
 #include "UObject/ObjectResource.h"
 
@@ -44,7 +45,7 @@ TSharedPtr<FAnimNextGraphInstance> UAnimNextAnimationGraph::AllocateInstance(FAn
 	TSharedPtr<FAnimNextGraphInstance> InstanceImpl = MakeShared<FAnimNextGraphInstance>();
 	InstanceImpl->DataInterface = this;
 	InstanceImpl->ModuleInstance = InModuleInstance;
-	InstanceImpl->ParentGraphInstance = InParentGraphInstance;
+	InstanceImpl->HostInstance = InParentGraphInstance != nullptr ? static_cast<FAnimNextDataInterfaceInstance*>(InParentGraphInstance) : static_cast<FAnimNextDataInterfaceInstance*>(InModuleInstance);
 	InstanceImpl->EntryPoint = EntryPoint;
 
 	// If we have a parent graph, use its root since we share the same root, otherwise if we have no parent, we are the root
diff --git a/Engine/Plugins/Experimental/AnimNextAnimGraph/Source/AnimNextAnimGraph/Private/Graph/AnimNextGraphInstance.cpp b/Engine/Plugins/Experimental/AnimNextAnimGraph/Source/AnimNextAnimGraph/Private/Graph/AnimNextGraphInstance.cpp
--- a/Engine/Plugins/Experimental/AnimNextAnimGraph/Source/AnimNextAnimGraph/Private/Graph/AnimNextGraphInstance.cpp
+++ b/Engine/Plugins/Experimental/AnimNextAnimGraph/Source/AnimNextAnimGraph/Private/Graph/AnimNextGraphInstance.cpp
@@ -15,6 +15,7 @@
 #include "Logging/StructuredLog.h"
 #include "Module/AnimNextModule.h"
 #include "Module/AnimNextModuleInstance.h"
+#include "Module/ModuleGuard.h"
 
 DEFINE_STAT(STAT_AnimNext_Graph_RigVM);
 
@@ -49,7 +50,6 @@ void FAnimNextGraphInstance::Release()
 
 	GraphInstancePtr.Reset();
 	ModuleInstance = nullptr;
-	ParentGraphInstance = nullptr;
 	RootGraphInstance = nullptr;
 	ExtendedExecuteContext.Reset();
 	Components.Empty();
@@ -83,7 +83,14 @@ FAnimNextModuleInstance* FAnimNextGraphInstance::GetModuleInstance() const
 
 FAnimNextGraphInstance* FAnimNextGraphInstance::GetParentGraphInstance() const
 {
-	return ParentGraphInstance;
+	if(IsRoot())
+	{
+		return nullptr;
+	}
+	else
+	{
+		return static_cast<FAnimNextGraphInstance*>(HostInstance);
+	}
 }
 
 FAnimNextGraphInstance* FAnimNextGraphInstance::GetRootGraphInstance() const
@@ -158,16 +165,6 @@ void FAnimNextGraphInstance::Update()
 	bHasUpdatedOnce = true;
 }
 
-FAnimNextDataInterfaceInstance* FAnimNextGraphInstance::GetHost() const
-{
-	if(ParentGraphInstance != nullptr)
-	{
-		return ParentGraphInstance;
-	}
-
-	return ModuleInstance;
-}
-
 void FAnimNextGraphInstance::ExecuteLatentPins(const TConstArrayView<UE::AnimNext::FLatentPropertyHandle>& LatentHandles, void* DestinationBasePtr, bool bIsFrozen)
 {
 	SCOPE_CYCLE_COUNTER(STAT_AnimNext_Graph_RigVM);
@@ -230,65 +227,163 @@ void FAnimNextGraphInstance::Thaw()
 
 void FAnimNextGraphInstance::OnModuleCompiled(UAnimNextModule* InModule)
 {
-	// If we are hosted directly by a module, invalidate and mark our bindings as needing update. They will be lazily re-bound the next time we run.
-	if(ModuleInstance && InModule == ModuleInstance->GetModule() && ParentGraphInstance == nullptr)
+	auto DependentModuleCompiled = [this, InModule]()
 	{
-		UnbindPublicVariables();
+		bool bFound = false;
+		ModuleInstance->ForEachPrerequisite([InModule, &bFound](FAnimNextModuleInstance& InPrerequisite)
+		{
+			if(InPrerequisite.GetModule() == InModule)
+			{
+				bFound = true;
+			}
+		});
+		return bFound;
+	};
+
+	// If we are hosted transitively by the compiled module, or if we could be bound to the compiled module,
+	// invalidate and mark our bindings as needing update. They will be lazily re-bound the next time we run.
+	if(ModuleInstance)
+	{
+		UE::AnimNext::FModuleWriteGuard Guard(ModuleInstance);
+		if(InModule == ModuleInstance->GetModule() || DependentModuleCompiled())
+		{
+			UnbindPublicVariables();
+		}
 	}
 }
 
 #endif
 
-// Helper for binding to both FAnimNextDataInterfaceInstance and IDataInterfaceHost without an abstraction between the two
-template<typename HostType>
-bool FAnimNextGraphInstance::BindToHostHelper(const HostType& InHost, bool bInAutoBind)
+bool FAnimNextGraphInstance::BindToCachedBindings()
 {
 	bool bPublicVariablesBound = false;
 
+	auto FindCachedInterfaceBindings = [this](const UAnimNextDataInterface* InDataInterface) -> const FCachedDataInterfaceBinding*
+	{
+		for(const FCachedDataInterfaceBinding& CachedBinding : CachedBindings)
+		{
+			if(CachedBinding.DataInterface == InDataInterface)
+			{
+				return &CachedBinding;
+			}
+		}
+
+		return nullptr;
+	};
+
 	const UAnimNextAnimationGraph* AnimationGraph = GetAnimationGraph();
 	const UPropertyBag* PropertyBag = Variables.GetPropertyBagStruct();
 	for(const FAnimNextImplementedDataInterface& ImplementedInterface : AnimationGraph->GetImplementedInterfaces())
 	{
-		if(bInAutoBind && !ImplementedInterface.bAutoBindToHost)
+		if(!ImplementedInterface.bAutoBindToHost)
 		{
 			continue;
 		}
 
-		const FAnimNextImplementedDataInterface* HostImplementedInterface = InHost.GetDataInterface()->FindImplementedInterface(ImplementedInterface.DataInterface);
-		if(HostImplementedInterface == nullptr)
+		const FCachedDataInterfaceBinding* CachedBinding = FindCachedInterfaceBindings(ImplementedInterface.DataInterface);
+		if(CachedBinding == nullptr)
 		{
-			// Host does not implement this interface, so skip
+			// Did not cache this binding, so skip
 			continue;
 		}
 
-		if(HostImplementedInterface->NumVariables != ImplementedInterface.NumVariables)
+		int32 VariableIndex = ImplementedInterface.VariableIndex;
+		int32 InterfaceVariableIndex = 0;
+		const int32 EndVariableIndex = ImplementedInterface.VariableIndex + ImplementedInterface.NumVariables;
+		for(; VariableIndex < EndVariableIndex; ++VariableIndex, ++InterfaceVariableIndex)
 		{
-			UE_LOGFMT(LogAnimation, Error, "BindToHost: Mismatched interface variables: '{Name}' ({Count}) vs Host '{Host}' ({HostCount})", AnimationGraph->GetFName(), ImplementedInterface.NumVariables, InHost.GetDataInterfaceName(), HostImplementedInterface->NumVariables);
-			continue;
+			const FPropertyBagPropertyDesc& Desc = PropertyBag->GetPropertyDescs()[VariableIndex];
+			const FCachedDataInterfaceBinding::FVariable& CachedBindingVariable = CachedBinding->CachedBindings[InterfaceVariableIndex];
+			if(CachedBindingVariable.Memory == nullptr)
+			{
+				continue;
+			}
+
+			if( Desc.Name == CachedBindingVariable.VariableName &&
+				InterfaceVariableIndex == CachedBindingVariable.VariableIndex &&
+				Desc.CachedProperty->GetClass() == CachedBindingVariable.Property->GetClass())
+			{
+				ExtendedExecuteContext.ExternalVariableRuntimeData[VariableIndex].Memory = CachedBindingVariable.Memory;
+				bPublicVariablesBound = true;
+			}
+		}
+	}
+
+	return bPublicVariablesBound;
+}
+
+void FAnimNextGraphInstance::UpdateCachedBindingsForHost()
+{
+	if(IsRoot())
+	{
+		if(FAnimNextModuleInstance* ModuleHost = GetModuleInstance())
+		{
+			UpdateCachedBindingsForHost(*ModuleHost);
 		}
+	}
+	else if(FAnimNextGraphInstance* ParentGraphInstance = GetParentGraphInstance())
+	{
+		// Just copy parent cached bindings
+		CachedBindings = ParentGraphInstance->CachedBindings;
+	}
+}
+
+void FAnimNextGraphInstance::UpdateCachedBindingsForHost(const FAnimNextModuleInstance& InHost)
+{
+	UpdateCachedBindingsForHostHelper(InHost);
+}
+
+void FAnimNextGraphInstance::UpdateCachedBindingsForHost(const UE::AnimNext::IDataInterfaceHost& InHost)
+{
+	UpdateCachedBindingsForHostHelper(InHost);
+}
+
+template<typename HostType>
+void FAnimNextGraphInstance::UpdateCachedBindingsForHostHelper(const HostType& InHost)
+{
+	auto FindOrAddCachedInterfaceBindings = [this](const UAnimNextDataInterface* InDataInterface) -> FCachedDataInterfaceBinding&
+	{
+		for(FCachedDataInterfaceBinding& CachedBinding : CachedBindings)
+		{
+			if(CachedBinding.DataInterface == InDataInterface)
+			{
+				return CachedBinding;
+			}
+		}
+
+		FCachedDataInterfaceBinding& NewCachedBinding = CachedBindings.AddDefaulted_GetRef();
+		NewCachedBinding.DataInterface = InDataInterface;
+		NewCachedBinding.CachedBindings.SetNum(InDataInterface->GetPublicVariableDefaults().GetNumPropertiesInBag());
+		return NewCachedBinding;
+	};
+
+	const UAnimNextDataInterface* HostDataInterface = InHost.GetDataInterface();
+	if(HostDataInterface == nullptr)
+	{
+		return;
+	}
+
+	const UPropertyBag* PropertyBag = HostDataInterface->GetPublicVariableDefaults().GetPropertyBagStruct();
+	for(const FAnimNextImplementedDataInterface& ImplementedInterface : HostDataInterface->GetImplementedInterfaces())
+	{
+		FCachedDataInterfaceBinding& CachedBinding = FindOrAddCachedInterfaceBindings(ImplementedInterface.DataInterface);
 
 		int32 VariableIndex = ImplementedInterface.VariableIndex;
-		int32 HostVariableIndex = HostImplementedInterface->VariableIndex;
 		const int32 EndVariableIndex = ImplementedInterface.VariableIndex + ImplementedInterface.NumVariables;
-		for(; VariableIndex < EndVariableIndex; ++VariableIndex, ++HostVariableIndex)
+		for(; VariableIndex < EndVariableIndex; ++VariableIndex)
 		{
 			const FPropertyBagPropertyDesc& Desc = PropertyBag->GetPropertyDescs()[VariableIndex];
-
-			uint8* HostMemory = InHost.GetMemoryForVariable(HostVariableIndex, Desc.Name, Desc.CachedProperty);
+			uint8* HostMemory = InHost.GetMemoryForVariable(VariableIndex, Desc.Name, Desc.CachedProperty);
 			if(HostMemory == nullptr)
 			{
 				continue;
 			}
-
-			ExtendedExecuteContext.ExternalVariableRuntimeData[VariableIndex].Memory = HostMemory;
-			bPublicVariablesBound = true;
+			CachedBinding.CachedBindings[VariableIndex] = FCachedDataInterfaceBinding::FVariable(Desc.Name, VariableIndex, Desc.CachedProperty, HostMemory);
 		}
 	}
-
-	return bPublicVariablesBound;
 }
 
-void FAnimNextGraphInstance::BindPublicVariables(TArrayView<FStructView> InHostStructs)
+void FAnimNextGraphInstance::BindPublicVariables(TArrayView<FStructView> InHostStructs, TConstArrayView<UE::AnimNext::IDataInterfaceHost*> InHosts)
 {
 	using namespace UE::AnimNext;
 
@@ -298,28 +393,33 @@ void FAnimNextGraphInstance::BindPublicVariables(TArrayView<FStructView> InHostS
 		return;
 	}
 
-	// First apply any automatic bindings to this instance's host
-   	if(FAnimNextDataInterfaceInstance* InstanceHost = GetHost())
-   	{
-   		BindToHostHelper(*InstanceHost, true);
-   	}
+	ensure(CachedBindings.Num() == 0);
 
-	const int32 NumStructs = InHostStructs.Num();
-	if(NumStructs > 0)
 	{
+		const int32 NumStructs = InHostStructs.Num();
+		const int32 NumHosts = InHosts.Num();
+
 		// stack allocating uninitialized arrays to contain FDataInterfaceStructAdapter(s)
+		// Note this means we can end up calling FMemory_Alloca(0), which is somewhat implementation specific.
+		// This is OK as we dont every iterate or use the memory in that case
 		TArrayView<FDataInterfaceStructAdapter> Adapters((FDataInterfaceStructAdapter*)FMemory_Alloca(NumStructs * sizeof(FDataInterfaceStructAdapter)), NumStructs);
-		TArrayView<IDataInterfaceHost*> PayloadVariableBindings((IDataInterfaceHost**)FMemory_Alloca(NumStructs * sizeof(IDataInterfaceHost*)), NumStructs);
+		TArrayView<IDataInterfaceHost*> Hosts((IDataInterfaceHost**)FMemory_Alloca((NumStructs + NumHosts) * sizeof(IDataInterfaceHost*)), NumStructs + NumHosts);
 
 		// initializing the FDataInterfaceStructAdapter(s) with their respective HostStruct
 		for (int32 StructIndex = 0; StructIndex < NumStructs; ++StructIndex)
 		{
 			FStructView& StructView = InHostStructs[StructIndex];
-			PayloadVariableBindings[StructIndex] = new (&Adapters[StructIndex]) FDataInterfaceStructAdapter(AnimationGraph, StructView);
+			Hosts[StructIndex] = new (&Adapters[StructIndex]) FDataInterfaceStructAdapter(AnimationGraph, StructView);
+		}
+
+		// Copy incoming hosts to provide a combined set to BindPublicVariables
+		for (int32 HostIndex = 0; HostIndex < NumHosts; ++HostIndex)
+		{
+			Hosts[NumStructs + HostIndex] = InHosts[HostIndex];
 		}
 
 		// Inner binding 
-		BindPublicVariables(PayloadVariableBindings);
+		BindPublicVariables(Hosts);
 
 		// uninitializing the FDataInterfaceStructAdapter(s)
 		for (int32 StructIndex = 0; StructIndex < NumStructs; ++StructIndex)
@@ -354,22 +454,20 @@ void FAnimNextGraphInstance::BindPublicVariables(TConstArrayView<UE::AnimNext::I
 		return;
 	}
 
-	bool bPublicVariablesBound = false;
+	ensure(CachedBindings.Num() == 0);
 
-	// First apply any automatic bindings to this instance's host
-	if(FAnimNextDataInterfaceInstance* InstanceHost = GetHost())
-	{
-		bPublicVariablesBound |= BindToHostHelper(*InstanceHost, true);
-	}
+	// First cache any bindings to this instance's host
+	UpdateCachedBindingsForHost();
 
-	// Next bind to any supplied host interfaces (bInAutoBind = false)
+	// Next cache any supplied host interfaces 
 	for(const IDataInterfaceHost* HostInterface : InHosts)
 	{
 		check(HostInterface);
-		bPublicVariablesBound |= BindToHostHelper(*HostInterface, false);
+		UpdateCachedBindingsForHost(*HostInterface);
 	}
 
-	if(bPublicVariablesBound)
+	// Bind to cached bindings we built above
+	if(BindToCachedBindings())
 	{
 		// Re-initialize memory handles
 		AnimationGraph->RigVM->InitializeInstance(ExtendedExecuteContext, /* bCopyMemory = */false);
@@ -403,6 +501,8 @@ void FAnimNextGraphInstance::UnbindPublicVariables()
 	// Re-initialize memory handles
 	AnimationGraph->RigVM->InitializeInstance(ExtendedExecuteContext, /* bCopyMemory = */false);
 
+	CachedBindings.Reset();
+
 	PublicVariablesState = EPublicVariablesState::Unbound;
 }
 
diff --git a/Engine/Plugins/Experimental/AnimNextAnimGraph/Source/AnimNextAnimGraph/Private/Graph/AnimNextModuleAnimGraphComponent.cpp b/Engine/Plugins/Experimental/AnimNextAnimGraph/Source/AnimNextAnimGraph/Private/Graph/AnimNextModuleAnimGraphComponent.cpp
new file mode 100644
--- /dev/null
+++ b/Engine/Plugins/Experimental/AnimNextAnimGraph/Source/AnimNextAnimGraph/Private/Graph/AnimNextModuleAnimGraphComponent.cpp
@@ -0,0 +1,45 @@
+// Copyright Epic Games, Inc. All Rights Reserved.
+
+#include "Graph/AnimNextModuleAnimGraphComponent.h"
+
+#include "Graph/AnimNextGraphInstance.h"
+#include "Graph/AnimNextAnimationGraph.h"
+
+FAnimNextModuleAnimGraphComponent::FAnimNextModuleAnimGraphComponent(FAnimNextModuleInstance& InOwnerInstance)
+	: FAnimNextModuleInstanceComponent(InOwnerInstance)
+{
+}
+
+TWeakPtr<FAnimNextGraphInstance> FAnimNextModuleAnimGraphComponent::AllocateInstance(const UAnimNextAnimationGraph* InAnimationGraph, FAnimNextGraphInstance* InParentInstance, FName InEntryPoint)
+{
+	check(InAnimationGraph);
+
+	TSharedPtr<FAnimNextGraphInstance> NewInstance = InAnimationGraph->AllocateInstance(&GetModuleInstance(), InParentInstance, InEntryPoint);
+	if(!NewInstance.IsValid())
+	{
+		return TWeakPtr<FAnimNextGraphInstance>();
+	}
+
+	GraphInstances.Add(NewInstance);
+	return NewInstance;
+}
+
+void FAnimNextModuleAnimGraphComponent::ReleaseInstance(TWeakPtr<FAnimNextGraphInstance> InWeakInstance)
+{
+	TSharedPtr<FAnimNextGraphInstance> PinnedInstance = InWeakInstance.Pin();
+	if(PinnedInstance.IsValid())
+	{
+		check(GraphInstances.Contains(PinnedInstance));			// Should not be releasing this instance if it is not owned by this module
+		GraphInstances.Remove(PinnedInstance);
+		check(PinnedInstance.GetSharedReferenceCount() == 1);	// This should be the final reference, all others should be weak
+	}
+}
+
+void FAnimNextModuleAnimGraphComponent::AddStructReferencedObjects(class FReferenceCollector& Collector)
+{
+	for (const TSharedPtr<FAnimNextGraphInstance>& GraphInstance : GraphInstances)
+	{
+		GraphInstance->AddStructReferencedObjects(Collector);
+	}
+}
+
diff --git a/Engine/Plugins/Experimental/AnimNextAnimGraph/Source/AnimNextAnimGraph/Private/Graph/RigUnit_AnimNextGraphEvaluator.cpp b/Engine/Plugins/Experimental/AnimNextAnimGraph/Source/AnimNextAnimGraph/Private/Graph/RigUnit_AnimNextGraphEvaluator.cpp
--- a/Engine/Plugins/Experimental/AnimNextAnimGraph/Source/AnimNextAnimGraph/Private/Graph/RigUnit_AnimNextGraphEvaluator.cpp
+++ b/Engine/Plugins/Experimental/AnimNextAnimGraph/Source/AnimNextAnimGraph/Private/Graph/RigUnit_AnimNextGraphEvaluator.cpp
@@ -2,6 +2,7 @@
 
 #include "Graph/RigUnit_AnimNextGraphEvaluator.h"
 #include "Graph/AnimNextGraphContextData.h"
+#include "Graph/AnimNextGraphEvaluatorExecuteDefinition.h"
 #include "Graph/AnimNextGraphLatentPropertiesContextData.h"
 #include "TraitCore/LatentPropertyHandle.h"
 
diff --git a/Engine/Plugins/Experimental/AnimNextAnimGraph/Source/AnimNextAnimGraph/Private/Graph/RigUnit_AnimNextRunAnimationGraph_v1.cpp b/Engine/Plugins/Experimental/AnimNextAnimGraph/Source/AnimNextAnimGraph/Private/Graph/RigUnit_AnimNextRunAnimationGraph_v1.cpp
new file mode 100644
--- /dev/null
+++ b/Engine/Plugins/Experimental/AnimNextAnimGraph/Source/AnimNextAnimGraph/Private/Graph/RigUnit_AnimNextRunAnimationGraph_v1.cpp
@@ -0,0 +1,141 @@
+// Copyright Epic Games, Inc. All Rights Reserved.
+
+#include "Graph/RigUnit_AnimNextRunAnimationGraph_v1.h"
+
+#include "Units/RigUnitContext.h"
+#include "AnimNextAnimGraphStats.h"
+#include "TraitInterfaces/IEvaluate.h"
+#include "TraitInterfaces/IUpdate.h"
+#include "Graph/AnimNextModuleAnimGraphComponent.h"
+#include "Module/AnimNextModuleInstance.h"
+#include "Graph/AnimNextAnimationGraph.h"
+#include "Graph/AnimNextGraphInstance.h"
+#include "Graph/RigVMTrait_AnimNextPublicVariables.h"
+#include "TraitCore/TraitEventList.h"
+
+#include UE_INLINE_GENERATED_CPP_BY_NAME(RigUnit_AnimNextRunAnimationGraph_v1)
+
+DEFINE_STAT(STAT_AnimNext_Run_Graph);
+
+FRigUnit_AnimNextRunAnimationGraph_v1_Execute()
+{
+	SCOPE_CYCLE_COUNTER(STAT_AnimNext_Run_Graph);
+
+	using namespace UE::AnimNext;
+
+	const FAnimNextModuleContextData& ModuleContextData = ExecuteContext.GetContextData<FAnimNextModuleContextData>();
+	FAnimNextModuleInstance& ModuleInstance = ModuleContextData.GetModuleInstance();
+
+	if (!ReferencePose.ReferencePose.IsValid())
+	{
+		return;
+	}
+
+	FAnimNextModuleAnimGraphComponent& AnimationGraphComponent = ModuleInstance.GetComponent<FAnimNextModuleAnimGraphComponent>();
+
+	if(Graph == nullptr)
+	{
+		AnimationGraphComponent.ReleaseInstance(WorkData.WeakInstance);
+		return;
+	}
+
+	// Release the instance if the graph has changed
+	if(WorkData.WeakInstance.IsValid() && !WorkData.WeakInstance.Pin()->UsesAnimationGraph(Graph))
+	{
+		AnimationGraphComponent.ReleaseInstance(WorkData.WeakInstance);
+	}
+
+	// Lazily (re-)allocate graph instance if required
+	if(!WorkData.WeakInstance.IsValid())
+	{
+		WorkData.WeakInstance = AnimationGraphComponent.AllocateInstance(Graph);
+	}
+
+	if(!WorkData.WeakInstance.IsValid())
+	{
+		return;
+	}
+
+	// Take a strong reference to the host instance, we are going to run it
+	TSharedRef<FAnimNextGraphInstance> InstanceRef = WorkData.WeakInstance.Pin().ToSharedRef();
+
+	if(InstanceRef->RequiresPublicVariableBinding())
+	{
+		// Bind any public variables we find
+		TArray<FPublicVariablesTraitToDataInterfaceHostAdapter, TInlineAllocator<4>> Adapters;
+		Adapters.Reserve(ExecuteContext.GetTraits().Num());
+		TArray<IDataInterfaceHost*, TInlineAllocator<4>> Hosts;
+		Hosts.Reserve(ExecuteContext.GetTraits().Num());
+		for(const FRigVMTraitScope& TraitScope : ExecuteContext.GetTraits())
+		{
+			const FRigVMTrait_AnimNextPublicVariables* VariablesTrait = TraitScope.GetTrait<FRigVMTrait_AnimNextPublicVariables>();
+			if(VariablesTrait == nullptr)
+			{
+				continue;
+			}
+
+			FPublicVariablesTraitToDataInterfaceHostAdapter& Adapter = Adapters.Emplace_GetRef(*VariablesTrait, TraitScope);
+			Hosts.Add(&Adapter);
+		}
+
+		InstanceRef->BindPublicVariables(Hosts);
+	}
+
+	// Propagate delta time
+	FAnimNextExecuteContext& AnimNextExecuteContext = InstanceRef->GetExtendedExecuteContext().GetPublicDataSafe<FAnimNextExecuteContext>();
+	AnimNextExecuteContext.SetDeltaTime(ExecuteContext.GetDeltaTime());
+
+	const UE::AnimNext::FReferencePose& RefPose = ReferencePose.ReferencePose.GetRef<UE::AnimNext::FReferencePose>();
+
+	// TODO: Currently forcing additive flag to false here
+	if (Result.LODPose.ShouldPrepareForLOD(RefPose, LOD, false))
+	{
+		Result.LODPose.PrepareForLOD(RefPose, LOD, true, false);
+	}
+
+	ensure(Result.LODPose.LODLevel == LOD);
+
+	// Every graph in a schedule will see the same input events (if they were queued before the schedule started)
+	FTraitEventList InputEventList;
+	FTraitEventList OutputEventList;
+
+	// A schedule can contain multiple graphs, we copy the input event list since it might be appended to during our update
+	{
+		FReadScopeLock ReadLock(ModuleInstance.EventListLock);
+		InputEventList = ModuleInstance.InputEventList;
+	}
+
+	// Track how many input events we started with, we'll append the new ones
+	const int32 NumOriginalInputEvents = InputEventList.Num();
+
+	// Internally we use memstack allocation, so we need a mark here
+	FMemStack& MemStack = FMemStack::Get();
+	FMemMark MemMark(MemStack);
+
+	// We allocate a dummy buffer to trigger the allocation of a large chunk if this is the first mark
+	// This reduces churn internally by avoiding a chunk to be repeatedly allocated and freed as we push/pop marks
+	MemStack.Alloc(size_t(FPageAllocator::SmallPageSize) + 1, 16);
+
+	UpdateGraph(InstanceRef.Get(), ExecuteContext.GetDeltaTime(), InputEventList, OutputEventList);
+	EvaluateGraph(InstanceRef.Get(), RefPose, LOD, Result);
+
+	// We might have appended new input/output events, append them
+	{
+		const int32 NumInputEvents = InputEventList.Num();
+
+		FWriteScopeLock WriteLock(ModuleInstance.EventListLock);
+
+		// Append the new input events
+		for (int32 EventIndex = NumOriginalInputEvents; EventIndex < NumInputEvents; ++EventIndex)
+		{
+			FAnimNextTraitEventPtr& Event = InputEventList[EventIndex];
+			if (Event->IsValid())
+			{
+				ModuleInstance.InputEventList.Push(MoveTemp(Event));
+			}
+		}
+
+		// Append our output events
+		ModuleInstance.OutputEventList.Append(OutputEventList);
+	}
+}
diff --git a/Engine/Plugins/Experimental/AnimNextAnimGraph/Source/AnimNextAnimGraph/Private/Graph/RigUnit_AnimNextRunAnimationGraph_v2.cpp b/Engine/Plugins/Experimental/AnimNextAnimGraph/Source/AnimNextAnimGraph/Private/Graph/RigUnit_AnimNextRunAnimationGraph_v2.cpp
new file mode 100644
--- /dev/null
+++ b/Engine/Plugins/Experimental/AnimNextAnimGraph/Source/AnimNextAnimGraph/Private/Graph/RigUnit_AnimNextRunAnimationGraph_v2.cpp
@@ -0,0 +1,168 @@
+// Copyright Epic Games, Inc. All Rights Reserved.
+
+#include "Graph/RigUnit_AnimNextRunAnimationGraph_v2.h"
+
+#include "AnimNextAnimGraphSettings.h"
+#include "Units/RigUnitContext.h"
+#include "AnimNextAnimGraphStats.h"
+#include "TraitInterfaces/IEvaluate.h"
+#include "TraitInterfaces/IUpdate.h"
+#include "Module/AnimNextModuleInstance.h"
+#include "Graph/AnimNextAnimationGraph.h"
+#include "Graph/AnimNextGraphInstance.h"
+#include "Graph/AnimNextModuleAnimGraphComponent.h"
+#include "Graph/RigVMTrait_AnimNextPublicVariables.h"
+#include "Injection/GraphInstanceInjectionComponent.h"
+#include "Injection/InjectionDataInterfaceHostAdapter.h"
+#include "TraitCore/TraitEventList.h"
+#include "Injection/AnimNextModuleInjectionComponent.h"
+
+#include UE_INLINE_GENERATED_CPP_BY_NAME(RigUnit_AnimNextRunAnimationGraph_v2)
+
+FRigUnit_AnimNextRunAnimationGraph_v2_Execute()
+{
+	SCOPE_CYCLE_COUNTER(STAT_AnimNext_Run_Graph);
+
+	using namespace UE::AnimNext;
+
+	const FAnimNextModuleContextData& ModuleContextData = ExecuteContext.GetContextData<FAnimNextModuleContextData>();
+	FAnimNextModuleInstance& ModuleInstance = ModuleContextData.GetModuleInstance();
+
+	if (!ReferencePose.ReferencePose.IsValid())
+	{
+		return;
+	}
+
+	const UE::AnimNext::FReferencePose& RefPose = ReferencePose.ReferencePose.GetRef<UE::AnimNext::FReferencePose>();
+
+	// TODO: Currently forcing additive flag to false here
+	if (Result.LODPose.ShouldPrepareForLOD(RefPose, LOD, false))
+	{
+		Result.LODPose.PrepareForLOD(RefPose, LOD, true, false);
+	}
+
+	ensure(Result.LODPose.LODLevel == LOD);
+
+	// Get a host to run this graph
+	const UAnimNextAnimationGraph* HostGraph = Graph.HostGraph ? Graph.HostGraph.Get() : GetDefault<UAnimNextAnimGraphSettings>()->GetDefaultRunGraphHost();// GetHostGraphToRun(ExecuteContext);
+	FAnimNextModuleAnimGraphComponent& AnimationGraphComponent = ModuleInstance.GetComponent<FAnimNextModuleAnimGraphComponent>();
+	if(HostGraph == nullptr)
+	{
+		AnimationGraphComponent.ReleaseInstance(WorkData.WeakHost);
+		return;
+	}
+
+	// Release the instance if the host graph has changed
+	if(WorkData.WeakHost.IsValid() && !WorkData.WeakHost.Pin()->UsesAnimationGraph(HostGraph))
+	{
+		AnimationGraphComponent.ReleaseInstance(WorkData.WeakHost);
+	}
+
+	// Lazily (re-)allocate graph instance if required
+	if(!WorkData.WeakHost.IsValid())
+	{
+		WorkData.WeakHost = AnimationGraphComponent.AllocateInstance(HostGraph);
+	}
+
+	if(!WorkData.WeakHost.IsValid())
+	{
+		return;
+	}
+
+	// Take a strong reference to the host instance, we are going to run it
+	TSharedRef<FAnimNextGraphInstance> HostInstanceRef = WorkData.WeakHost.Pin().ToSharedRef();
+	FAnimNextGraphInstance& HostInstance = HostInstanceRef.Get();
+
+	// Apply bindings
+	if(HostInstance.RequiresPublicVariableBinding())
+	{
+		TArray<FPublicVariablesTraitToDataInterfaceHostAdapter, TInlineAllocator<4>> Adapters;
+		Adapters.Reserve(ExecuteContext.GetTraits().Num());
+		TArray<IDataInterfaceHost*, TInlineAllocator<4>> Hosts;
+		Hosts.Reserve(ExecuteContext.GetTraits().Num() + 1);
+		FInjectionDataInterfaceHostAdapter InjectionAdapter;
+
+		// Bind any pin-overriden public variable traits we find
+		for(const FRigVMTraitScope& TraitScope : ExecuteContext.GetTraits())
+		{
+			const FRigVMTrait_AnimNextPublicVariables* VariablesTrait = TraitScope.GetTrait<FRigVMTrait_AnimNextPublicVariables>();
+			if(VariablesTrait == nullptr)
+			{
+				continue;
+			}
+
+			FPublicVariablesTraitToDataInterfaceHostAdapter& Adapter = Adapters.Emplace_GetRef(*VariablesTrait, TraitScope);
+			Hosts.Add(&Adapter);
+		}
+
+		// Ensure we have an injection component, as this node acts as a top-level injection site
+		ModuleInstance.GetComponent<FAnimNextModuleInjectionComponent>();
+
+		// Inject the default graph if there is any (_after_ pin-exposed variables so they cant override it)
+		FGraphInstanceInjectionComponent& HostInjectionComponent = HostInstance.GetComponent<FGraphInstanceInjectionComponent>();
+		FName GraphName;
+		TStructView<FAnimNextAnimGraph> GraphInstanceView;
+		if(HostInjectionComponent.GetInjectionInfo().GetDefaultInjectableGraph(GraphName, GraphInstanceView))
+		{
+			WorkData.InjectedGraph = Graph;
+			GraphInstanceView = TStructView<FAnimNextAnimGraph>(WorkData.InjectedGraph);
+			InjectionAdapter = FInjectionDataInterfaceHostAdapter(HostInstance, GraphName, GraphInstanceView);
+			Hosts.Add(&InjectionAdapter);
+		}
+
+		HostInstance.BindPublicVariables(Hosts);
+	}
+	else
+	{
+		// Re-copy the graph into the injected work data, as it may have been updated by bindings/injection
+		WorkData.InjectedGraph = Graph;
+	}
+
+	// Propagate delta time
+	FAnimNextExecuteContext& AnimNextExecuteContext = HostInstance.GetExtendedExecuteContext().GetPublicDataSafe<FAnimNextExecuteContext>();
+	AnimNextExecuteContext.SetDeltaTime(ExecuteContext.GetDeltaTime());
+
+	// Every graph in a schedule will see the same input events (if they were queued before the schedule started)
+	FTraitEventList InputEventList;
+	FTraitEventList OutputEventList;
+
+	// A module can contain multiple graphs, we copy the input event list since it might be appended to during our update
+	{
+		FReadScopeLock ReadLock(ModuleInstance.EventListLock);
+		InputEventList = ModuleInstance.InputEventList;
+	}
+
+	// Track how many input events we started with, we'll append the new ones
+	const int32 NumOriginalInputEvents = InputEventList.Num();
+
+	// Internally we use memstack allocation, so we need a mark here
+	FMemStack& MemStack = FMemStack::Get();
+	FMemMark MemMark(MemStack);
+
+	// We allocate a dummy buffer to trigger the allocation of a large chunk if this is the first mark
+	// This reduces churn internally by avoiding a chunk to be repeatedly allocated and freed as we push/pop marks
+	MemStack.Alloc(size_t(FPageAllocator::SmallPageSize) + 1, 16);
+
+	UpdateGraph(HostInstance, ExecuteContext.GetDeltaTime(), InputEventList, OutputEventList);
+	EvaluateGraph(HostInstance, RefPose, LOD, Result);
+
+	// We might have appended new input/output events, append them
+	{
+		const int32 NumInputEvents = InputEventList.Num();
+
+		FWriteScopeLock WriteLock(ModuleInstance.EventListLock);
+
+		// Append the new input events
+		for (int32 EventIndex = NumOriginalInputEvents; EventIndex < NumInputEvents; ++EventIndex)
+		{
+			FAnimNextTraitEventPtr& Event = InputEventList[EventIndex];
+			if (Event->IsValid())
+			{
+				ModuleInstance.InputEventList.Push(MoveTemp(Event));
+			}
+		}
+
+		// Append our output events
+		ModuleInstance.OutputEventList.Append(OutputEventList);
+	}
+}
diff --git a/Engine/Plugins/Experimental/AnimNextAnimGraph/Source/AnimNextAnimGraph/Private/Graph/RigUnit_AnimNextSetNotifyContext.cpp b/Engine/Plugins/Experimental/AnimNextAnimGraph/Source/AnimNextAnimGraph/Private/Graph/RigUnit_AnimNextSetNotifyContext.cpp
--- a/Engine/Plugins/Experimental/AnimNextAnimGraph/Source/AnimNextAnimGraph/Private/Graph/RigUnit_AnimNextSetNotifyContext.cpp
+++ b/Engine/Plugins/Experimental/AnimNextAnimGraph/Source/AnimNextAnimGraph/Private/Graph/RigUnit_AnimNextSetNotifyContext.cpp
@@ -6,6 +6,6 @@
 FRigUnit_AnimNextSetNotifyContext_Execute()
 {
 	FAnimNextModuleInstance& ModuleInstance = ExecuteContext.GetContextData<FAnimNextModuleContextData>().GetModuleInstance();
-	FModuleInstanceComponent_NotifyDispatcher& NotifyDispatcher = ModuleInstance.GetComponent<FModuleInstanceComponent_NotifyDispatcher>();
+	FAnimNextNotifyDispatcherComponent& NotifyDispatcher = ModuleInstance.GetComponent<FAnimNextNotifyDispatcherComponent>();
 	NotifyDispatcher.SkeletalMeshComponent = SkeletalMeshComponent;
 }
\ No newline at end of file
diff --git a/Engine/Plugins/Experimental/AnimNextAnimGraph/Source/AnimNextAnimGraph/Private/Injection/AnimNextModuleInjectionComponent.cpp b/Engine/Plugins/Experimental/AnimNextAnimGraph/Source/AnimNextAnimGraph/Private/Injection/AnimNextModuleInjectionComponent.cpp
new file mode 100644
--- /dev/null
+++ b/Engine/Plugins/Experimental/AnimNextAnimGraph/Source/AnimNextAnimGraph/Private/Injection/AnimNextModuleInjectionComponent.cpp
@@ -0,0 +1,199 @@
+﻿// Copyright Epic Games, Inc. All Rights Reserved.
+
+#include "Injection/AnimNextModuleInjectionComponent.h"
+
+#include "InjectionEvents.h"
+#include "Graph/AnimNextAnimGraph.h"
+#include "Graph/AnimNextGraphInstance.h"
+#include "Logging/StructuredLog.h"
+#include "Module/AnimNextModuleInstance.h"
+#include "Module/ModuleTaskContext.h"
+
+
+FAnimNextModuleInjectionComponent::FAnimNextModuleInjectionComponent(FAnimNextModuleInstance& InOwnerInstance)
+	: FAnimNextModuleInstanceComponent(InOwnerInstance)
+	, InjectionInfo(InOwnerInstance)
+{
+	// Register re-injection for each (user) tick function
+	for(UE::AnimNext::FModuleEventTickFunction& TickFunction : InOwnerInstance.GetTickFunctions())
+	{
+		if(TickFunction.bUserEvent)
+		{
+			TickFunction.OnPreModuleEvent.AddStatic(&FAnimNextModuleInjectionComponent::OnReapplyInjection);
+		}
+	}
+}
+
+void FAnimNextModuleInjectionComponent::OnTraitEvent(FAnimNextTraitEvent& Event)
+{
+	if(UE::AnimNext::FInjection_InjectEvent* InjectionEvent = Event.AsType<UE::AnimNext::FInjection_InjectEvent>())
+	{
+		OnInjectionEvent(*InjectionEvent);
+	}
+	else if(UE::AnimNext::FInjection_UninjectEvent* UninjectionEvent = Event.AsType<UE::AnimNext::FInjection_UninjectEvent>())
+	{
+		OnUninjectionEvent(*UninjectionEvent);
+	}
+}
+
+void FAnimNextModuleInjectionComponent::OnInjectionEvent(UE::AnimNext::FInjection_InjectEvent& InEvent)
+{
+	using namespace UE::AnimNext;
+
+	if(InEvent.IsHandled())
+	{
+		return;
+	}
+
+	const FInjectionRequestArgs& RequestArgs = InEvent.Request->GetArgs();
+
+	FName FoundName = NAME_None;
+	TStructView<FAnimNextAnimGraph> InjectableGraph = InjectionInfo.FindInjectableGraphInstance(RequestArgs.Name, FoundName);
+	if(!InjectableGraph.IsValid())
+	{
+		UE_LOGFMT(LogAnimation, Warning, "Could not find injection site {SiteName} for injection request", RequestArgs.Name);
+		InEvent.MarkConsumed();
+		return;
+	}
+
+	// Correct the name found above as we may have targeted NAME_None (any)
+	if(RequestArgs.Name != FoundName)
+	{
+		InEvent.Request->GetMutableArgs().Name = FoundName;
+	}
+
+	// Mark as handled so any additional trait events dont get processed at the module level
+	InEvent.MarkHandled();
+
+	// Store request as it will need to be re-applied each frame to ensure that bindings do not override it
+	FInjectionRecord& InjectionRecord = CurrentRequests.FindOrAdd(InEvent.Request->GetArgs().Name);
+
+	FAnimNextAnimGraph& Graph = InjectableGraph.Get<FAnimNextAnimGraph>();
+	switch(RequestArgs.Type)
+	{
+	case EAnimNextInjectionType::InjectObject:
+		// Update graph and bump serial number to identify this injection routing
+		ensureAlways(RequestArgs.Object != nullptr);
+		Graph.Asset = RequestArgs.Object;
+		Graph.InjectionSerialNumber = IncrementSerialNumber();
+		InEvent.SerialNumber = Graph.InjectionSerialNumber;
+		InjectionRecord.SerialNumber = Graph.InjectionSerialNumber;
+		InjectionRecord.GraphRequest = InEvent.Request;
+
+		// Note we dont consume here, as we want the event to forward to the injection site trait
+		break;
+	case EAnimNextInjectionType::EvaluationModifier:
+		// We dont increment serial number when applying evaluation modifiers as we dont want to trigger a graph instantiation
+		ensureAlways(RequestArgs.EvaluationModifier != nullptr);
+		Graph.EvaluationModifier = RequestArgs.EvaluationModifier;
+		InjectionRecord.ModifierRequest = InEvent.Request;
+
+		// Update status
+		{
+			auto StatusUpdateEvent = MakeTraitEvent<FInjection_StatusUpdateEvent>();
+			StatusUpdateEvent->Request = InEvent.Request;
+			StatusUpdateEvent->Status = EInjectionStatus::Playing;
+			GetModuleInstance().QueueOutputTraitEvent(StatusUpdateEvent);
+		}
+
+		// Evaluation modifiers are consumed straight away
+		InEvent.MarkConsumed();
+		break;
+	}
+}
+
+void FAnimNextModuleInjectionComponent::OnUninjectionEvent(UE::AnimNext::FInjection_UninjectEvent& InEvent)
+{
+	using namespace UE::AnimNext;
+
+	const FInjectionRequestArgs& RequestArgs = InEvent.Request->GetArgs();
+
+	// Injection modifiers can have private access to take over low-level tasks
+	FName FoundName = NAME_None;
+	TStructView<FAnimNextAnimGraph> InjectableGraph = InjectionInfo.FindInjectableGraphInstance(RequestArgs.Name, FoundName);
+	if(!InjectableGraph.IsValid())
+	{
+		UE_LOGFMT(LogAnimation, Warning, "Could not find injection site {SiteName} for un-injection request", RequestArgs.Name);
+		InEvent.MarkConsumed();
+		return;
+	}
+
+	FInjectionRecord& InjectionRecord = CurrentRequests.FindOrAdd(FoundName);
+
+	// Update graph and bump serial number to identify this un-injection routing
+	FAnimNextAnimGraph& Graph = InjectableGraph.Get<FAnimNextAnimGraph>();
+	switch(RequestArgs.Type)
+	{
+	case EAnimNextInjectionType::InjectObject:
+		// Update graph and bump serial number to identify this injection routing
+		Graph.Asset = nullptr;
+		Graph.InjectionSerialNumber = IncrementSerialNumber();
+		InEvent.SerialNumber = Graph.InjectionSerialNumber;
+		InjectionRecord.GraphRequest.Reset();
+
+		// Note we dont consume here, as we want the event to forward to the injection site trait
+		break;
+	case EAnimNextInjectionType::EvaluationModifier:
+		// We dont increment serial number when applying evaluation modifiers as we dont want to trigger a graph instantiation
+		Graph.EvaluationModifier = nullptr;
+		InjectionRecord.ModifierRequest.Reset();
+
+		// Update status
+		{
+			auto StatusUpdateEvent = MakeTraitEvent<FInjection_StatusUpdateEvent>();
+			StatusUpdateEvent->Request = InEvent.Request;
+			StatusUpdateEvent->Status = EInjectionStatus::Completed;
+			GetModuleInstance().QueueOutputTraitEvent(StatusUpdateEvent);
+		}
+		
+		// Evaluation modifiers are consumed straight away
+		InEvent.MarkConsumed();
+		break;
+	}
+
+	if(!InjectionRecord.IsValid())
+	{
+		CurrentRequests.Remove(InEvent.Request->GetArgs().Name);
+	}
+}
+
+void FAnimNextModuleInjectionComponent::OnReapplyInjection(const UE::AnimNext::FModuleTaskContext& InContext)
+{
+	using namespace UE::AnimNext;
+
+	FAnimNextModuleInjectionComponent& Component = InContext.ModuleInstance->GetComponent<FAnimNextModuleInjectionComponent>();
+	
+	for(const TPair<FName, FInjectionRecord>& RequestPair : Component.CurrentRequests)
+	{
+		FName FoundName = NAME_None;
+		TStructView<FAnimNextAnimGraph> InjectableGraph = Component.InjectionInfo.FindInjectableGraphInstance(RequestPair.Key, FoundName);
+		if(!InjectableGraph.IsValid())
+		{
+			continue;
+		}
+
+		// Re-apply this request, as it may have been overwritten by subsequent bindings/calculations
+		FAnimNextAnimGraph& Graph = InjectableGraph.Get<FAnimNextAnimGraph>();
+		if(RequestPair.Value.GraphRequest.IsValid())
+		{
+			const FInjectionRequestArgs& GraphRequestArgs = RequestPair.Value.GraphRequest->GetArgs();
+			Graph.Asset = GraphRequestArgs.Object;
+			Graph.InjectionSerialNumber = RequestPair.Value.SerialNumber;
+		}
+		if(RequestPair.Value.ModifierRequest.IsValid())
+		{
+			const FInjectionRequestArgs& ModifierRequestArgs = RequestPair.Value.ModifierRequest->GetArgs();
+			Graph.EvaluationModifier = ModifierRequestArgs.EvaluationModifier;
+		}
+	}
+}
+
+uint32 FAnimNextModuleInjectionComponent::IncrementSerialNumber()
+{
+	// Avoid zero as this is 'invalid' and  will ensure at injection sites (indicating incorrect routing)
+	if(++SerialNumber == 0)
+	{
+		++SerialNumber;
+	}
+	return SerialNumber;
+}
diff --git a/Engine/Plugins/Experimental/AnimNextAnimGraph/Source/AnimNextAnimGraph/Private/Injection/AnimNextModuleInjectionComponent.h b/Engine/Plugins/Experimental/AnimNextAnimGraph/Source/AnimNextAnimGraph/Private/Injection/AnimNextModuleInjectionComponent.h
new file mode 100644
--- /dev/null
+++ b/Engine/Plugins/Experimental/AnimNextAnimGraph/Source/AnimNextAnimGraph/Private/Injection/AnimNextModuleInjectionComponent.h
@@ -0,0 +1,64 @@
+﻿// Copyright Epic Games, Inc. All Rights Reserved.
+
+#pragma once
+
+#include "InjectionInfo.h"
+#include "DataInterface/AnimNextDataInterfaceHost.h"
+#include "Injection/InjectionRequest.h"
+#include "Module/ModuleInstanceComponent.h"
+#include "AnimNextModuleInjectionComponent.generated.h"
+
+namespace UE::AnimNext
+{
+	struct FInjection_InjectEvent;
+	struct FInjection_UninjectEvent;
+	struct FModuleTaskContext;
+}
+
+// Module component that holds info about injection sites and routes injection requests
+USTRUCT()
+struct FAnimNextModuleInjectionComponent : public FAnimNextModuleInstanceComponent
+{
+	GENERATED_BODY()
+
+	ANIMNEXT_DECLARE_MODULE_INSTANCE_COMPONENT(FAnimNextModuleInjectionComponent);
+
+	FAnimNextModuleInjectionComponent() = default;
+	explicit FAnimNextModuleInjectionComponent(FAnimNextModuleInstance& InOwnerInstance);
+
+	// Get the cached injection info for our module 
+	const UE::AnimNext::FInjectionInfo& GetInjectionInfo() const { return InjectionInfo; }
+
+private:
+	// FAnimNextModuleInstanceComponent interface
+	virtual void OnTraitEvent(FAnimNextTraitEvent& Event) override;
+
+	void OnInjectionEvent(UE::AnimNext::FInjection_InjectEvent& InEvent);
+	void OnUninjectionEvent(UE::AnimNext::FInjection_UninjectEvent& InEvent);
+
+	static void OnReapplyInjection(const UE::AnimNext::FModuleTaskContext& InContext);
+
+	uint32 IncrementSerialNumber();
+
+private:
+	// Info for injection
+	UE::AnimNext::FInjectionInfo InjectionInfo;
+
+	struct FInjectionRecord
+	{
+		bool IsValid() const
+		{
+			return GraphRequest.IsValid() || ModifierRequest.IsValid();
+		}
+
+		TSharedPtr<UE::AnimNext::FInjectionRequest> GraphRequest;
+		TSharedPtr<UE::AnimNext::FInjectionRequest> ModifierRequest;
+		uint32 SerialNumber = 0;
+	};
+	
+	// Currently-injected requests
+	TMap<FName, FInjectionRecord> CurrentRequests;
+
+	// Serial number used to identify forwarded requests
+	uint32 SerialNumber = 0;
+};
diff --git a/Engine/Plugins/Experimental/AnimNextAnimGraph/Source/AnimNextAnimGraph/Private/Injection/GraphInstanceInjectionComponent.cpp b/Engine/Plugins/Experimental/AnimNextAnimGraph/Source/AnimNextAnimGraph/Private/Injection/GraphInstanceInjectionComponent.cpp
new file mode 100644
--- /dev/null
+++ b/Engine/Plugins/Experimental/AnimNextAnimGraph/Source/AnimNextAnimGraph/Private/Injection/GraphInstanceInjectionComponent.cpp
@@ -0,0 +1,15 @@
+﻿// Copyright Epic Games, Inc. All Rights Reserved.
+
+#include "GraphInstanceInjectionComponent.h"
+#include "Graph/AnimNextGraphInstance.h"
+
+namespace UE::AnimNext
+{
+
+FGraphInstanceInjectionComponent::FGraphInstanceInjectionComponent(FAnimNextGraphInstance& InOwnerInstance)
+	: FGraphInstanceComponent(InOwnerInstance)
+	, InjectionInfo(InOwnerInstance)
+{
+}
+
+}
diff --git a/Engine/Plugins/Experimental/AnimNextAnimGraph/Source/AnimNextAnimGraph/Private/Injection/GraphInstanceInjectionComponent.h b/Engine/Plugins/Experimental/AnimNextAnimGraph/Source/AnimNextAnimGraph/Private/Injection/GraphInstanceInjectionComponent.h
new file mode 100644
--- /dev/null
+++ b/Engine/Plugins/Experimental/AnimNextAnimGraph/Source/AnimNextAnimGraph/Private/Injection/GraphInstanceInjectionComponent.h
@@ -0,0 +1,29 @@
+﻿// Copyright Epic Games, Inc. All Rights Reserved.
+
+#pragma once
+
+#include "CoreMinimal.h"
+#include "InjectionInfo.h"
+#include "Graph/GraphInstanceComponent.h"
+
+namespace UE::AnimNext
+{
+	/**
+	 * FGraphInstanceInjectionComponent
+	 * This component maintains injection info for a graph
+	 */
+	struct FGraphInstanceInjectionComponent : public FGraphInstanceComponent
+	{
+		DECLARE_ANIM_GRAPH_INSTANCE_COMPONENT(FGraphInstanceInjectionComponent)
+
+		explicit FGraphInstanceInjectionComponent(FAnimNextGraphInstance& InOwnerInstance);
+
+		const FInjectionInfo& GetInjectionInfo() const
+		{
+			return InjectionInfo;
+		}
+
+	private:
+		FInjectionInfo InjectionInfo;
+	};
+}
diff --git a/Engine/Plugins/Experimental/AnimNextAnimGraph/Source/AnimNextAnimGraph/Private/Injection/InjectionCallbackProxy.cpp b/Engine/Plugins/Experimental/AnimNextAnimGraph/Source/AnimNextAnimGraph/Private/Injection/InjectionCallbackProxy.cpp
new file mode 100644
--- /dev/null
+++ b/Engine/Plugins/Experimental/AnimNextAnimGraph/Source/AnimNextAnimGraph/Private/Injection/InjectionCallbackProxy.cpp
@@ -0,0 +1,130 @@
+﻿// Copyright Epic Games, Inc. All Rights Reserved.
+
+#include "Injection/InjectionCallbackProxy.h"
+#include "Animation/AnimSequence.h"
+#include "Component/AnimNextComponent.h"
+#include "GraphInterfaces/AnimNextNativeDataInterface_AnimSequencePlayer.h"
+#include "Injection/InjectionUtils.h"
+
+#include UE_INLINE_GENERATED_CPP_BY_NAME(InjectionCallbackProxy)
+
+UInjectionCallbackProxy* UInjectionCallbackProxy::CreateProxyObjectForInjection(
+	UAnimNextComponent* AnimNextComponent,
+	FName SiteName,
+	UObject* Object,
+	UAnimNextComponent* BindingComponent,
+	FInstancedStruct Payload,
+	UE::AnimNext::FInjectionBlendSettings BlendInSettings,
+	UE::AnimNext::FInjectionBlendSettings BlendOutSettings)
+{
+	UInjectionCallbackProxy* Proxy = NewObject<UInjectionCallbackProxy>();
+	Proxy->SetFlags(RF_StrongRefOnFrame);
+	Proxy->Inject(AnimNextComponent, SiteName, Object, BindingComponent, MoveTemp(Payload), BlendInSettings, BlendOutSettings);
+	return Proxy;
+}
+
+bool UInjectionCallbackProxy::Inject(
+	UAnimNextComponent* AnimNextComponent,
+	FName SiteName,
+	UObject* Object,
+	UAnimNextComponent* BindingComponent,
+	FInstancedStruct&& Payload,
+	const UE::AnimNext::FInjectionBlendSettings& BlendInSettings,
+	const UE::AnimNext::FInjectionBlendSettings& BlendOutSettings)
+{
+	if (AnimNextComponent == nullptr)
+	{
+		return false;
+	}
+
+	UE::AnimNext::FInjectionRequestArgs RequestArgs;
+	RequestArgs.Name = SiteName;
+	RequestArgs.Object = Object;
+	RequestArgs.BlendInSettings = BlendInSettings;
+	RequestArgs.BlendOutSettings = BlendOutSettings;
+	RequestArgs.BindingModuleHandle = BindingComponent ? BindingComponent->GetModuleHandle() : UE::AnimNext::FModuleHandle();
+	RequestArgs.Payload.AddNative(MoveTemp(Payload));
+
+	UE::AnimNext::FInjectionLifetimeEvents LifetimeEvents;
+	LifetimeEvents.OnCompleted.BindUObject(this, &UInjectionCallbackProxy::OnInjectionCompleted);
+	LifetimeEvents.OnInterrupted.BindUObject(this, &UInjectionCallbackProxy::OnInjectionInterrupted);
+	LifetimeEvents.OnBlendingOut.BindUObject(this, &UInjectionCallbackProxy::OnInjectionBlendingOut);
+
+	PlayingRequest = UE::AnimNext::FInjectionUtils::Inject(AnimNextComponent, AnimNextComponent->GetModuleHandle(), MoveTemp(RequestArgs), MoveTemp(LifetimeEvents));
+
+	bWasInterrupted = false;
+
+	if (!PlayingRequest.IsValid())
+	{
+		OnInterrupted.Broadcast();
+		Reset();
+	}
+
+	return PlayingRequest.IsValid();
+}
+
+EUninjectionResult UInjectionCallbackProxy::Uninject()
+{
+	if(!PlayingRequest.IsValid())
+	{
+		return EUninjectionResult::Failed;
+	}
+
+	UE::AnimNext::FInjectionUtils::Uninject(PlayingRequest);
+
+	return EUninjectionResult::Succeeded;
+}
+
+void UInjectionCallbackProxy::Cancel()
+{
+	Super::Cancel();
+	Uninject();
+}
+
+void UInjectionCallbackProxy::OnInjectionCompleted(const UE::AnimNext::FInjectionRequest& Request)
+{
+	if (!bWasInterrupted)
+	{
+		const UE::AnimNext::EInjectionStatus Status = Request.GetStatus();
+		check(!EnumHasAnyFlags(Status, UE::AnimNext::EInjectionStatus::Interrupted));
+
+		if (EnumHasAnyFlags(Status, UE::AnimNext::EInjectionStatus::Expired))
+		{
+			OnInterrupted.Broadcast();
+		}
+		else
+		{
+			OnCompleted.Broadcast();
+		}
+	}
+
+	Reset();
+}
+
+void UInjectionCallbackProxy::OnInjectionInterrupted(const UE::AnimNext::FInjectionRequest& Request)
+{
+	bWasInterrupted = true;
+
+	OnInterrupted.Broadcast();
+}
+
+void UInjectionCallbackProxy::OnInjectionBlendingOut(const UE::AnimNext::FInjectionRequest& Request)
+{
+	if (!bWasInterrupted)
+	{
+		OnBlendOut.Broadcast();
+	}
+}
+
+void UInjectionCallbackProxy::Reset()
+{
+	PlayingRequest = nullptr;
+	bWasInterrupted = false;
+}
+
+void UInjectionCallbackProxy::BeginDestroy()
+{
+	Reset();
+
+	Super::BeginDestroy();
+}
diff --git a/Engine/Plugins/Experimental/AnimNextAnimGraph/Source/AnimNextAnimGraph/Private/Injection/InjectionDataInterfaceHostAdapter.cpp b/Engine/Plugins/Experimental/AnimNextAnimGraph/Source/AnimNextAnimGraph/Private/Injection/InjectionDataInterfaceHostAdapter.cpp
new file mode 100644
--- /dev/null
+++ b/Engine/Plugins/Experimental/AnimNextAnimGraph/Source/AnimNextAnimGraph/Private/Injection/InjectionDataInterfaceHostAdapter.cpp
@@ -0,0 +1,32 @@
+﻿// Copyright Epic Games, Inc. All Rights Reserved.
+
+#include "InjectionDataInterfaceHostAdapter.h"
+#include "Graph/AnimNextAnimationGraph.h"
+#include "Graph/AnimNextAnimGraph.h"
+#include "Graph/AnimNextGraphInstance.h"
+
+namespace UE::AnimNext
+{
+
+const UAnimNextDataInterface* FInjectionDataInterfaceHostAdapter::GetDataInterface() const
+{
+	return HostInstance->GetAnimationGraph();
+}
+
+uint8* FInjectionDataInterfaceHostAdapter::GetMemoryForVariable(int32 InVariableIndex, FName InVariableName, const FProperty* InVariableProperty) const
+{
+	if(InVariableName != Name)
+	{
+		return nullptr;
+	}
+
+	const FStructProperty* StructProperty = CastField<FStructProperty>(InVariableProperty);
+	if(StructProperty == nullptr || StructProperty->Struct != FAnimNextAnimGraph::StaticStruct())
+	{
+		return nullptr;
+	}
+
+	return GraphInstance.GetMemory();
+}
+
+}
diff --git a/Engine/Plugins/Experimental/AnimNextAnimGraph/Source/AnimNextAnimGraph/Private/Injection/InjectionDataInterfaceHostAdapter.h b/Engine/Plugins/Experimental/AnimNextAnimGraph/Source/AnimNextAnimGraph/Private/Injection/InjectionDataInterfaceHostAdapter.h
new file mode 100644
--- /dev/null
+++ b/Engine/Plugins/Experimental/AnimNextAnimGraph/Source/AnimNextAnimGraph/Private/Injection/InjectionDataInterfaceHostAdapter.h
@@ -0,0 +1,35 @@
+﻿// Copyright Epic Games, Inc. All Rights Reserved.
+
+#pragma once
+
+#include "CoreMinimal.h"
+#include "DataInterface/AnimNextDataInterfaceHost.h"
+#include "StructUtils/StructView.h"
+
+struct FAnimNextAnimGraph;
+
+namespace UE::AnimNext
+{
+
+// Interface adapter used to inject a graph into a host.
+// Uses the host's data interface and the name of the variable in the host's data interface to apply the supplied graph instance
+struct FInjectionDataInterfaceHostAdapter : public IDataInterfaceHost
+{
+	FInjectionDataInterfaceHostAdapter() = default;
+	
+	FInjectionDataInterfaceHostAdapter(FAnimNextGraphInstance& InHostInstance, FName InName, TStructView<FAnimNextAnimGraph> InGraphInstance)
+		: HostInstance(&InHostInstance)
+		, Name(InName)
+		, GraphInstance(InGraphInstance)
+	{}
+
+	// IDataInterfaceHost interface
+	virtual const UAnimNextDataInterface* GetDataInterface() const override;
+	virtual uint8* GetMemoryForVariable(int32 InVariableIndex, FName InVariableName, const FProperty* InVariableProperty) const override;
+
+	FAnimNextGraphInstance* HostInstance = nullptr;
+	FName Name;
+	TStructView<FAnimNextAnimGraph> GraphInstance;
+};
+
+}
diff --git a/Engine/Plugins/Experimental/AnimNextAnimGraph/Source/AnimNextAnimGraph/Private/Injection/InjectionInfo.cpp b/Engine/Plugins/Experimental/AnimNextAnimGraph/Source/AnimNextAnimGraph/Private/Injection/InjectionInfo.cpp
new file mode 100644
--- /dev/null
+++ b/Engine/Plugins/Experimental/AnimNextAnimGraph/Source/AnimNextAnimGraph/Private/Injection/InjectionInfo.cpp
@@ -0,0 +1,96 @@
+﻿// Copyright Epic Games, Inc. All Rights Reserved.
+
+#include "InjectionInfo.h"
+
+#include "DataInterface/AnimNextDataInterfaceInstance.h"
+#include "Graph/AnimNextAnimGraph.h"
+#include "StructUtils/PropertyBag.h"
+
+namespace UE::AnimNext
+{
+
+FInjectionInfo::FInjectionInfo(const FAnimNextDataInterfaceInstance& InInstance)
+	: Instance(&InInstance)
+{
+	CacheInfo();
+}
+
+void FInjectionInfo::CacheInfo() const
+{
+	check(Instance);
+	check(Instance->GetDataInterface());
+
+	DefaultInjectableGraphIndex = INDEX_NONE;
+	InjectableGraphs.Reset();
+
+	const UPropertyBag* Struct = Instance->Variables.GetPropertyBagStruct();
+	const uint8* Memory = Instance->Variables.GetValue().GetMemory();
+	if(Struct == nullptr || Memory == nullptr)
+	{
+		return;
+	}
+
+	// Find the default value in our set of variables
+	// Note we dont recurse here as the struct views we cache cannot be relocated (e.g. held in an array) due to them being referenced
+	// elsewhere via raw ptrs in RigVM memory handles
+	int32 VariableIndex = 0;
+	for (TPropertyValueIterator<FProperty> It(Struct, Memory, EPropertyValueIteratorFlags::NoRecursion); It; ++It, ++VariableIndex)
+	{
+		const FStructProperty* Property = CastField<FStructProperty>(It.Key());
+		if(Property == nullptr)
+		{
+			continue;
+		}
+
+		if(Property->Struct->IsChildOf(FAnimNextAnimGraph::StaticStruct()))
+		{
+			const FAnimNextAnimGraph* Value = static_cast<const FAnimNextAnimGraph*>(It.Value());
+			FName GraphName = Property->GetFName();
+			TStructView<FAnimNextAnimGraph> GraphStructView(const_cast<FAnimNextAnimGraph&>(*Value));
+
+			int32 GraphIndex = InjectableGraphs.Emplace(GraphName, GraphStructView);
+
+			if(VariableIndex == Instance->GetDataInterface()->DefaultInjectionSiteIndex)
+			{
+				DefaultInjectableGraphIndex = GraphIndex;
+			}
+		}
+	}
+
+	if(DefaultInjectableGraphIndex == INDEX_NONE && InjectableGraphs.Num() > 0)
+	{
+		DefaultInjectableGraphIndex = 0;
+	}
+}
+
+bool FInjectionInfo::GetDefaultInjectableGraph(FName& OutName, TStructView<FAnimNextAnimGraph>& OutGraph) const
+{
+	if(DefaultInjectableGraphIndex == INDEX_NONE || !InjectableGraphs.IsValidIndex(DefaultInjectableGraphIndex))
+	{
+		return false;
+	}
+
+	const FInjectableGraphInfo& InjectableGraph = InjectableGraphs[DefaultInjectableGraphIndex];
+	OutName = InjectableGraph.Name;
+	OutGraph = InjectableGraph.StructView;
+
+	return OutName != NAME_None && OutGraph.IsValid();
+}
+
+TStructView<FAnimNextAnimGraph> FInjectionInfo::FindInjectableGraphInstance(FName InName, FName& OutRealName) const
+{
+	const int32 InjectableGraphNum = InjectableGraphs.Num();
+	for(int32 InjectableGraphIndex = 0; InjectableGraphIndex < InjectableGraphNum; ++InjectableGraphIndex)
+	{
+		const FInjectableGraphInfo& InjectableGraphInfo = InjectableGraphs[InjectableGraphIndex];
+		if((InName == NAME_None && InjectableGraphIndex == DefaultInjectableGraphIndex) || InName == InjectableGraphInfo.Name)
+		{
+			OutRealName = InjectableGraphInfo.Name;
+			return InjectableGraphInfo.StructView;
+		}
+	}
+
+	return TStructView<FAnimNextAnimGraph>();
+}
+
+}
diff --git a/Engine/Plugins/Experimental/AnimNextAnimGraph/Source/AnimNextAnimGraph/Private/Injection/InjectionInfo.h b/Engine/Plugins/Experimental/AnimNextAnimGraph/Source/AnimNextAnimGraph/Private/Injection/InjectionInfo.h
new file mode 100644
--- /dev/null
+++ b/Engine/Plugins/Experimental/AnimNextAnimGraph/Source/AnimNextAnimGraph/Private/Injection/InjectionInfo.h
@@ -0,0 +1,72 @@
+﻿// Copyright Epic Games, Inc. All Rights Reserved.
+
+#pragma once
+
+#include "CoreMinimal.h"
+#include "StructUtils/StructView.h"
+
+struct FAnimNextDataInterfaceInstance;
+struct FAnimNextAnimGraph;
+struct FAnimNextModuleInjectionComponent;
+
+namespace UE::AnimNext
+{
+
+// Info used to track injection sites for a data interface instance (graph, module etc)
+struct FInjectionInfo
+{
+	FInjectionInfo() = default;
+
+	explicit FInjectionInfo(const FAnimNextDataInterfaceInstance& InInstance);
+
+	// Get the default injectable graph data
+	bool GetDefaultInjectableGraph(FName& OutName, TStructView<FAnimNextAnimGraph>& OutGraph) const;
+
+	// Find an injectable graph instance by name.
+	// @param   InName            The name of the injection site. If NAME_None is passed then the default site will be returned
+	// @param   OutRealName       The name of the actual injection site we found (in the case NAME_None was passed)
+	TStructView<FAnimNextAnimGraph> FindInjectableGraphInstance(FName InName, FName& OutRealName) const;
+
+	// Iterate each injectable graph instance
+	template<typename PredicateType>
+	void ForEachInjectableGraphInstance(PredicateType InPredicate) const
+	{
+		for(const FInjectableGraphInfo& InjectableGraphInfo : InjectableGraphs)
+		{
+			InPredicate(InjectableGraphInfo.Name, InjectableGraphInfo.StructView);
+		}
+	}
+
+private:
+	void CacheInfo() const;
+
+private:
+	friend ::FAnimNextModuleInjectionComponent;
+
+	struct FInjectableGraphInfo
+	{
+		FInjectableGraphInfo(FName InName, TStructView<FAnimNextAnimGraph> InStructView)
+			: StructView(InStructView)
+			, Name(InName)
+		{}
+
+		TStructView<FAnimNextAnimGraph> StructView;
+		FName Name;
+	};
+
+	// Lookup of injectable graphs by name
+	// Lookup is a linear search for now as we expect there to be 'not many' of these things. If we find this is a bottleneck, we can evaluate the
+	// performance/memory cost of a TMap
+	mutable TArray<FInjectableGraphInfo> InjectableGraphs;
+
+	// Instance we are tracking
+	const FAnimNextDataInterfaceInstance* Instance = nullptr;
+
+	// The name of the default injectable graph for our module (user-adjustable via SetDefaultInjectionSite) 
+	FName DefaultInjectableGraphName;
+
+	// Index into InjectableGraphs of the default graph
+	mutable int32 DefaultInjectableGraphIndex = INDEX_NONE;
+};
+
+}
diff --git a/Engine/Plugins/Experimental/AnimNextAnimGraph/Source/AnimNextAnimGraph/Private/Injection/InjectionUtils.cpp b/Engine/Plugins/Experimental/AnimNextAnimGraph/Source/AnimNextAnimGraph/Private/Injection/InjectionUtils.cpp
new file mode 100644
--- /dev/null
+++ b/Engine/Plugins/Experimental/AnimNextAnimGraph/Source/AnimNextAnimGraph/Private/Injection/InjectionUtils.cpp
@@ -0,0 +1,166 @@
+﻿// Copyright Epic Games, Inc. All Rights Reserved.
+
+#include "Injection/InjectionUtils.h"
+#include "GraphInterfaces/AnimNextNativeDataInterface_AnimSequencePlayer.h"
+#include "Animation/AnimSequence.h"
+#include "Component/AnimNextComponent.h"
+#include "DataInterface/AnimNextDataInterfaceInstance.h"
+#include "Graph/AnimNextGraphInstance.h"
+#include "Logging/StructuredLog.h"
+#include "UObject/ObjectKey.h"
+
+namespace UE::AnimNext
+{
+	FInjectionRequestPtr FInjectionUtils::Inject(UAnimNextComponent* InComponent, FInjectionRequestArgs&& InArgs, FInjectionLifetimeEvents&& InLifetimeEvents)
+	{
+		return Inject(InComponent, InComponent->GetModuleHandle(), MoveTemp(InArgs), MoveTemp(InLifetimeEvents));
+	}
+
+	FInjectionRequestPtr FInjectionUtils::Inject(UObject* InHost, FModuleHandle InHandle, FInjectionRequestArgs&& InArgs, FInjectionLifetimeEvents&& InLifetimeEvents)
+	{
+		check(IsInGameThread());
+
+		FInjectionRequestPtr Request = MakeInjectionRequest();
+		if(Request->Inject(MoveTemp(InArgs), MoveTemp(InLifetimeEvents), InHost, InHandle))
+		{
+			return Request;
+		}
+		return nullptr;
+	}
+
+	void FInjectionUtils::Uninject(FInjectionRequestPtr InInjectionRequest)
+	{
+		if(!InInjectionRequest.IsValid())
+		{
+			return;
+		}
+		InInjectionRequest->Uninject();
+	}
+
+	FInjectionRequestPtr FInjectionUtils::PlayAnim(
+		UAnimNextComponent* InComponent,
+		FName InSiteName,
+		UAnimSequence* InAnimSequence,
+		float InPlayRate,
+		float InStartPosition,
+		const FInjectionBlendSettings& InBlendInSettings,
+		const FInjectionBlendSettings& InBlendOutSettings,
+		FInjectionLifetimeEvents&& InLifetimeEvents)
+	{
+		return PlayAnim(
+			InComponent,
+			InComponent->GetModuleHandle(),
+			InSiteName,
+			InAnimSequence,
+			InPlayRate,
+			InStartPosition,
+			InBlendInSettings,
+			InBlendOutSettings,
+			MoveTemp(InLifetimeEvents));
+	}
+
+	FInjectionRequestPtr FInjectionUtils::PlayAnim(
+		UObject* InHost,
+		FModuleHandle InModuleHandle,
+		FName InSiteName,
+		UAnimSequence* InAnimSequence,
+		float InPlayRate, float InStartPosition,
+		const FInjectionBlendSettings& InBlendInSettings,
+		const FInjectionBlendSettings& InBlendOutSettings,
+		FInjectionLifetimeEvents&& InLifetimeEvents)
+	{
+		check(IsInGameThread());
+
+		FInstancedStruct Payload;
+		Payload.InitializeAs<FAnimNextNativeDataInterface_AnimSequencePlayer>();
+		FAnimNextNativeDataInterface_AnimSequencePlayer& PlayNativeInterface = Payload.GetMutable<FAnimNextNativeDataInterface_AnimSequencePlayer>();
+		PlayNativeInterface.AnimSequence = InAnimSequence;
+		PlayNativeInterface.PlayRate = InPlayRate;
+		PlayNativeInterface.StartPosition = InStartPosition;
+
+		FInjectionRequestArgs RequestArgs;
+		RequestArgs.Name = InSiteName;
+		RequestArgs.Object = InAnimSequence;
+		RequestArgs.Type = EAnimNextInjectionType::InjectObject;
+		RequestArgs.BlendInSettings = InBlendInSettings;
+		RequestArgs.BlendOutSettings = InBlendOutSettings;
+		RequestArgs.Payload.AddNative(MoveTemp(Payload));
+
+		return Inject(InHost, InModuleHandle, MoveTemp(RequestArgs), MoveTemp(InLifetimeEvents));
+	}
+
+	FInjectionRequestPtr FInjectionUtils::InjectAsset(
+		UAnimNextComponent* InComponent,
+		FName InSiteName,
+		UObject* InAsset,
+		FAnimNextDataInterfacePayload&& InPayload,
+		UAnimNextComponent* InBindingComponent,
+		const FInjectionBlendSettings& InBlendInSettings,
+		const FInjectionBlendSettings& InBlendOutSettings,
+		FInjectionLifetimeEvents&& InLifetimeEvents)
+	{
+		return InjectAsset(
+			InComponent,
+			InComponent->GetModuleHandle(),
+			InSiteName,
+			InAsset,
+			MoveTemp(InPayload),
+			InBindingComponent ? InBindingComponent->GetModuleHandle() : FModuleHandle(),
+			InBlendInSettings,
+			InBlendOutSettings,
+			MoveTemp(InLifetimeEvents));
+	}
+
+	FInjectionRequestPtr FInjectionUtils::InjectAsset(
+		UObject* InHost,
+		FModuleHandle InModuleHandle,
+		FName InSiteName,
+		UObject* InAsset,
+		FAnimNextDataInterfacePayload&& InPayload,
+		FModuleHandle InBindingModuleHandle,
+		const FInjectionBlendSettings& InBlendInSettings,
+		const FInjectionBlendSettings& InBlendOutSettings,
+		FInjectionLifetimeEvents&& InLifetimeEvents)
+	{
+		check(IsInGameThread());
+
+		FInjectionRequestArgs RequestArgs;
+		RequestArgs.Name = InSiteName;
+		RequestArgs.Object = InAsset;
+		RequestArgs.BindingModuleHandle = InBindingModuleHandle;
+		RequestArgs.Type = EAnimNextInjectionType::InjectObject;
+		RequestArgs.BlendInSettings = InBlendInSettings;
+		RequestArgs.BlendOutSettings = InBlendOutSettings;
+		RequestArgs.Payload = MoveTemp(InPayload);
+
+		return Inject(InHost, InModuleHandle, MoveTemp(RequestArgs), MoveTemp(InLifetimeEvents));
+	}
+
+	FInjectionRequestPtr FInjectionUtils::InjectEvaluationModifier(
+		UAnimNextComponent* InComponent,
+		const TSharedRef<IEvaluationModifier>& InEvaluationModifier,
+		FName InSiteName)
+	{
+		return InjectEvaluationModifier(
+			InComponent,
+			InComponent->GetModuleHandle(),
+			InEvaluationModifier,
+			InSiteName);
+	}
+
+	FInjectionRequestPtr FInjectionUtils::InjectEvaluationModifier(
+		UObject* InHost,
+		FModuleHandle InModuleHandle,
+		const TSharedRef<IEvaluationModifier>& InEvaluationModifier,
+		FName InSiteName)
+	{
+		check(IsInGameThread());
+
+		FInjectionRequestArgs RequestArgs;
+		RequestArgs.Name = InSiteName;
+		RequestArgs.Type = EAnimNextInjectionType::EvaluationModifier;
+		RequestArgs.EvaluationModifier = InEvaluationModifier;
+
+		return Inject(InHost, InModuleHandle, MoveTemp(RequestArgs));
+	}
+}
diff --git a/Engine/Plugins/Experimental/AnimNextAnimGraph/Source/AnimNextAnimGraph/Private/Injection/ModuleInjectionDataInterfaceAdapter.cpp b/Engine/Plugins/Experimental/AnimNextAnimGraph/Source/AnimNextAnimGraph/Private/Injection/ModuleInjectionDataInterfaceAdapter.cpp
new file mode 100644
--- /dev/null
+++ b/Engine/Plugins/Experimental/AnimNextAnimGraph/Source/AnimNextAnimGraph/Private/Injection/ModuleInjectionDataInterfaceAdapter.cpp
@@ -0,0 +1,47 @@
+﻿// Copyright Epic Games, Inc. All Rights Reserved.
+
+#include "ModuleInjectionDataInterfaceAdapter.h"
+
+#include "AnimNextPool.h"
+#include "Module/AnimNextModuleInstance.h"
+#include "Module/ModuleGuard.h"
+
+namespace UE::AnimNext
+{
+
+FModuleInjectionDataInterfaceAdapter::FModuleInjectionDataInterfaceAdapter(FAnimNextModuleInstance* InModuleInstance, FModuleHandle InOtherModuleHandle)
+{
+	if(InModuleInstance == nullptr || !InOtherModuleHandle.IsValid() || InModuleInstance->Pool == nullptr)
+	{
+		return;
+	}
+
+	const FAnimNextModuleInstance* FoundOtherModuleInstance = InModuleInstance->Pool->TryGet(InOtherModuleHandle);
+	if(FoundOtherModuleInstance == nullptr)
+	{
+		return;
+	}
+
+	if(!ensureAlways(InModuleInstance->IsPrerequisite(*FoundOtherModuleInstance)))
+	{
+		return;
+	}
+
+	OtherModuleInstance = FoundOtherModuleInstance;
+}
+
+const UAnimNextDataInterface* FModuleInjectionDataInterfaceAdapter::GetDataInterface() const
+{
+	return OtherModuleInstance ? OtherModuleInstance->GetDataInterface() : nullptr;
+}
+
+uint8* FModuleInjectionDataInterfaceAdapter::GetMemoryForVariable(int32 InVariableIndex, FName InVariableName, const FProperty* InVariableProperty) const
+{
+	if(OtherModuleInstance == nullptr)
+	{
+		return nullptr;
+	}
+	return OtherModuleInstance->GetMemoryForVariable(InVariableIndex, InVariableName, InVariableProperty);
+}
+
+}
diff --git a/Engine/Plugins/Experimental/AnimNextAnimGraph/Source/AnimNextAnimGraph/Private/Injection/ModuleInjectionDataInterfaceAdapter.h b/Engine/Plugins/Experimental/AnimNextAnimGraph/Source/AnimNextAnimGraph/Private/Injection/ModuleInjectionDataInterfaceAdapter.h
new file mode 100644
--- /dev/null
+++ b/Engine/Plugins/Experimental/AnimNextAnimGraph/Source/AnimNextAnimGraph/Private/Injection/ModuleInjectionDataInterfaceAdapter.h
@@ -0,0 +1,30 @@
+﻿// Copyright Epic Games, Inc. All Rights Reserved.
+
+#pragma once
+
+#include "CoreMinimal.h"
+#include "DataInterface/AnimNextDataInterfaceHost.h"
+#include "Module/ModuleHandle.h"
+#include "Module/ModuleGuard.h"
+
+namespace UE::AnimNext
+{
+
+// Adapter allowing external systems to implement a data interface
+// Used at instantiation time (and in editor, re-binding after compilation time) to bind directly to a host's memory for an interfaces variables
+struct FModuleInjectionDataInterfaceAdapter : public IDataInterfaceHost
+{
+public:
+	FModuleInjectionDataInterfaceAdapter() = default;
+
+	FModuleInjectionDataInterfaceAdapter(FAnimNextModuleInstance* InModuleInstance, FModuleHandle InOtherModuleHandle);
+
+private:
+	// IDataInterfaceHost interface
+	virtual const UAnimNextDataInterface* GetDataInterface() const override;
+	virtual uint8* GetMemoryForVariable(int32 InVariableIndex, FName InVariableName, const FProperty* InVariableProperty) const override;
+
+	const FAnimNextModuleInstance* OtherModuleInstance = nullptr;
+};
+
+}
diff --git a/Engine/Plugins/Experimental/AnimNextAnimGraph/Source/AnimNextAnimGraph/Private/PlayAnim/PlayAnimCallbackProxy.cpp b/Engine/Plugins/Experimental/AnimNextAnimGraph/Source/AnimNextAnimGraph/Private/PlayAnim/PlayAnimCallbackProxy.cpp
--- a/Engine/Plugins/Experimental/AnimNextAnimGraph/Source/AnimNextAnimGraph/Private/PlayAnim/PlayAnimCallbackProxy.cpp
+++ b/Engine/Plugins/Experimental/AnimNextAnimGraph/Source/AnimNextAnimGraph/Private/PlayAnim/PlayAnimCallbackProxy.cpp
@@ -2,7 +2,9 @@
 
 #include "PlayAnim/PlayAnimCallbackProxy.h"
 #include "Animation/AnimSequence.h"
+#include "Component/AnimNextComponent.h"
 #include "GraphInterfaces/AnimNextNativeDataInterface_AnimSequencePlayer.h"
+#include "Injection/InjectionUtils.h"
 
 #include UE_INLINE_GENERATED_CPP_BY_NAME(PlayAnimCallbackProxy)
 
@@ -16,82 +18,49 @@ UPlayAnimCallbackProxy::UPlayAnimCallbackProxy(const FObjectInitializer& ObjectI
 
 UPlayAnimCallbackProxy* UPlayAnimCallbackProxy::CreateProxyObjectForPlayAnim(
 	UAnimNextComponent* AnimNextComponent,
-	FName SlotName,
+	FName SiteName,
 	UAnimSequence* AnimSequence,
 	float PlayRate,
 	float StartPosition,
-	UE::AnimNext::FPlayAnimBlendSettings BlendInSettings,
-	UE::AnimNext::FPlayAnimBlendSettings BlendOutSettings)
+	UE::AnimNext::FInjectionBlendSettings BlendInSettings,
+	UE::AnimNext::FInjectionBlendSettings BlendOutSettings)
 {
 	UPlayAnimCallbackProxy* Proxy = NewObject<UPlayAnimCallbackProxy>();
 	Proxy->SetFlags(RF_StrongRefOnFrame);
-	Proxy->Play(AnimNextComponent, SlotName, AnimSequence, PlayRate, StartPosition, BlendInSettings, BlendOutSettings);
-	return Proxy;
-}
-
-UPlayAnimCallbackProxy* UPlayAnimCallbackProxy::CreateProxyObjectForPlayAsset(
-	UAnimNextComponent* AnimNextComponent,
-	FName SlotName,
-	UObject* Asset,
-	const FInstancedStruct& Payload,
-	UE::AnimNext::FPlayAnimBlendSettings BlendInSettings,
-	UE::AnimNext::FPlayAnimBlendSettings BlendOutSettings)
-{
-	UPlayAnimCallbackProxy* Proxy = NewObject<UPlayAnimCallbackProxy>();
-	Proxy->SetFlags(RF_StrongRefOnFrame);
-	FInstancedStruct PayloadCopy(Payload);
-	Proxy->Play(AnimNextComponent, SlotName, Asset, MoveTemp(PayloadCopy), BlendInSettings, BlendOutSettings);
+	Proxy->Play(AnimNextComponent, SiteName, AnimSequence, PlayRate, StartPosition, BlendInSettings, BlendOutSettings);
 	return Proxy;
 }
 
 bool UPlayAnimCallbackProxy::Play(
 	UAnimNextComponent* AnimNextComponent,
-	FName SlotName,
+	FName SiteName,
 	UAnimSequence* AnimSequence,
 	float PlayRate,
 	float StartPosition,
-	const UE::AnimNext::FPlayAnimBlendSettings& BlendInSettings,
-	const UE::AnimNext::FPlayAnimBlendSettings& BlendOutSettings)
-{
-	FInstancedStruct NativeInterface;
-	NativeInterface.InitializeAs<FAnimNextNativeDataInterface_AnimSequencePlayer>();
-	FAnimNextNativeDataInterface_AnimSequencePlayer& PlayNativeInterface = NativeInterface.GetMutable<FAnimNextNativeDataInterface_AnimSequencePlayer>();
-	PlayNativeInterface.AnimSequence = AnimSequence;
-	PlayNativeInterface.PlayRate = PlayRate;
-	PlayNativeInterface.StartPosition = StartPosition;
-	return Play(AnimNextComponent, SlotName, AnimSequence, MoveTemp(NativeInterface), BlendInSettings, BlendOutSettings);
-}
-
-bool UPlayAnimCallbackProxy::Play(
-	UAnimNextComponent* AnimNextComponent,
-	FName SlotName,
-	UObject* Object,
-	FInstancedStruct&& Payload,
-	const UE::AnimNext::FPlayAnimBlendSettings& BlendInSettings,
-	const UE::AnimNext::FPlayAnimBlendSettings& BlendOutSettings)
+	const UE::AnimNext::FInjectionBlendSettings& BlendInSettings,
+	const UE::AnimNext::FInjectionBlendSettings& BlendOutSettings)
 {
 	bool bPlayedSuccessfully = false;
 	if (AnimNextComponent != nullptr)
 	{
-		UE::AnimNext::FPlayAnimRequestArgs RequestArgs;
-		RequestArgs.SlotName = SlotName;
-		RequestArgs.Object = Object;
-		RequestArgs.BlendInSettings = BlendInSettings;
-		RequestArgs.BlendOutSettings = BlendOutSettings;
-		RequestArgs.Payload.AddNative(MoveTemp(Payload));
-
-		auto Request = UE::AnimNext::MakePlayAnimRequest();
-		Request->OnCompleted.BindUObject(this, &UPlayAnimCallbackProxy::OnPlayAnimCompleted);
-		Request->OnInterrupted.BindUObject(this, &UPlayAnimCallbackProxy::OnPlayAnimInterrupted);
-		Request->OnBlendingOut.BindUObject(this, &UPlayAnimCallbackProxy::OnPlayAnimBlendingOut);
-
-		bPlayedSuccessfully = Request->Play(MoveTemp(RequestArgs), AnimNextComponent);
-
-		PlayingRequest = Request;
+		UE::AnimNext::FInjectionLifetimeEvents LifetimeEvents;
+		LifetimeEvents.OnCompleted.BindUObject(this, &UPlayAnimCallbackProxy::OnPlayAnimCompleted);
+		LifetimeEvents.OnInterrupted.BindUObject(this, &UPlayAnimCallbackProxy::OnPlayAnimInterrupted);
+		LifetimeEvents.OnBlendingOut.BindUObject(this, &UPlayAnimCallbackProxy::OnPlayAnimBlendingOut);
+
+		PlayingRequest = UE::AnimNext::FInjectionUtils::PlayAnim(
+			AnimNextComponent,
+			SiteName,
+			AnimSequence,
+			PlayRate,
+			StartPosition,
+			BlendInSettings,
+			BlendOutSettings,
+			MoveTemp(LifetimeEvents));
 		bWasInterrupted = false;
 	}
 
-	if (!bPlayedSuccessfully)
+	if (!PlayingRequest.IsValid())
 	{
 		OnInterrupted.Broadcast();
 		Reset();
@@ -100,14 +69,14 @@ bool UPlayAnimCallbackProxy::Play(
 	return bPlayedSuccessfully;
 }
 
-void UPlayAnimCallbackProxy::OnPlayAnimCompleted(const UE::AnimNext::FPlayAnimRequest& Request)
+void UPlayAnimCallbackProxy::OnPlayAnimCompleted(const UE::AnimNext::FInjectionRequest& Request)
 {
 	if (!bWasInterrupted)
 	{
-		const UE::AnimNext::EPlayAnimStatus Status = Request.GetStatus();
-		check(!EnumHasAnyFlags(Status, UE::AnimNext::EPlayAnimStatus::Interrupted));
+		const UE::AnimNext::EInjectionStatus Status = Request.GetStatus();
+		check(!EnumHasAnyFlags(Status, UE::AnimNext::EInjectionStatus::Interrupted));
 
-		if (EnumHasAnyFlags(Status, UE::AnimNext::EPlayAnimStatus::Expired))
+		if (EnumHasAnyFlags(Status, UE::AnimNext::EInjectionStatus::Expired))
 		{
 			OnInterrupted.Broadcast();
 		}
@@ -120,14 +89,14 @@ void UPlayAnimCallbackProxy::OnPlayAnimCompleted(const UE::AnimNext::FPlayAnimRe
 	Reset();
 }
 
-void UPlayAnimCallbackProxy::OnPlayAnimInterrupted(const UE::AnimNext::FPlayAnimRequest& Request)
+void UPlayAnimCallbackProxy::OnPlayAnimInterrupted(const UE::AnimNext::FInjectionRequest& Request)
 {
 	bWasInterrupted = true;
 
 	OnInterrupted.Broadcast();
 }
 
-void UPlayAnimCallbackProxy::OnPlayAnimBlendingOut(const UE::AnimNext::FPlayAnimRequest& Request)
+void UPlayAnimCallbackProxy::OnPlayAnimBlendingOut(const UE::AnimNext::FInjectionRequest& Request)
 {
 	if (!bWasInterrupted)
 	{
diff --git a/Engine/Plugins/Experimental/AnimNextAnimGraph/Source/AnimNextAnimGraph/Private/PlayAnim/PlayAnimSlotTrait.cpp b/Engine/Plugins/Experimental/AnimNextAnimGraph/Source/AnimNextAnimGraph/Private/PlayAnim/PlayAnimSlotTrait.cpp
--- a/Engine/Plugins/Experimental/AnimNextAnimGraph/Source/AnimNextAnimGraph/Private/PlayAnim/PlayAnimSlotTrait.cpp
+++ b/Engine/Plugins/Experimental/AnimNextAnimGraph/Source/AnimNextAnimGraph/Private/PlayAnim/PlayAnimSlotTrait.cpp
@@ -31,14 +31,14 @@ namespace UE::AnimNext
 		GeneratorMacro(IUpdateTraversal) \
 
 	#define TRAIT_EVENT_ENUMERATOR(GeneratorMacro) \
-		GeneratorMacro(FPlayAnimSlotTrait::OnPlayEvent) \
-		GeneratorMacro(FPlayAnimSlotTrait::OnStopEvent) \
+		GeneratorMacro(FPlayAnimSlotTrait::OnInjectEvent) \
+		GeneratorMacro(FPlayAnimSlotTrait::OnUninjectEvent) \
 
 	GENERATE_ANIM_TRAIT_IMPLEMENTATION(FPlayAnimSlotTrait, TRAIT_INTERFACE_ENUMERATOR, NULL_ANIM_TRAIT_INTERFACE_ENUMERATOR, TRAIT_EVENT_ENUMERATOR)
 	#undef TRAIT_INTERFACE_ENUMERATOR
 	#undef TRAIT_EVENT_ENUMERATOR
 
-	void FPlayAnimSlotTrait::FPlayAnimSlotRequest::Initialize(FPlayAnimRequestPtr InRequest, const FPlayAnimBlendSettings& InBlendSettings, const UAnimNextAnimationGraph* InAnimationGraph, FAnimNextDataInterfacePayload&& InPayload)
+	void FPlayAnimSlotTrait::FPlayAnimSlotRequest::Initialize(FInjectionRequestPtr InRequest, const FInjectionBlendSettings& InBlendSettings, const UAnimNextAnimationGraph* InAnimationGraph, FAnimNextDataInterfacePayload&& InPayload)
 	{
 		Request = InRequest;
 		BlendSettings = InBlendSettings;
@@ -81,13 +81,13 @@ namespace UE::AnimNext
 		return InstanceData.SlotRequests.AddDefaulted();
 	}
 
-	ETraitStackPropagation FPlayAnimSlotTrait::OnPlayEvent(const FExecutionContext& Context, FTraitBinding& Binding, FPlayAnim_PlayEvent& Event) const
+	ETraitStackPropagation FPlayAnimSlotTrait::OnInjectEvent(const FExecutionContext& Context, FTraitBinding& Binding, FInjection_InjectEvent& Event) const
 	{
 		const FSharedData* SharedData = Binding.GetSharedData<FSharedData>();
 		const FName SlotName = SharedData->GetSlotName(Binding);
 
-		const FAnimNextPlayAnimRequestArgs& RequestArgs = Event.Request->GetArgs();
-		if (SlotName == RequestArgs.SlotName)
+		const FAnimNextInjectionRequestArgs& RequestArgs = Event.Request->GetArgs();
+		if (SlotName == RequestArgs.Name)
 		{
 			FInstanceData* InstanceData = Binding.GetInstanceData<FInstanceData>();
 
@@ -103,13 +103,13 @@ namespace UE::AnimNext
 		return ETraitStackPropagation::Continue;
 	}
 
-	ETraitStackPropagation FPlayAnimSlotTrait::OnStopEvent(const FExecutionContext& Context, FTraitBinding& Binding, FPlayAnim_StopEvent& Event) const
+	ETraitStackPropagation FPlayAnimSlotTrait::OnUninjectEvent(const FExecutionContext& Context, FTraitBinding& Binding, FInjection_UninjectEvent& Event) const
 	{
 		const FSharedData* SharedData = Binding.GetSharedData<FSharedData>();
 		const FName SlotName = SharedData->GetSlotName(Binding);
 
-		const FAnimNextPlayAnimRequestArgs& RequestArgs = Event.Request->GetArgs();
-		if (SlotName == RequestArgs.SlotName)
+		const FAnimNextInjectionRequestArgs& RequestArgs = Event.Request->GetArgs();
+		if (SlotName == RequestArgs.Name)
 		{
 			FInstanceData* InstanceData = Binding.GetInstanceData<FInstanceData>();
 
@@ -157,19 +157,19 @@ namespace UE::AnimNext
 		bool bJustTransitioned = false;
 		if (InstanceData->PendingRequest.IsValid() || !bHasActiveSubGraph)
 		{
-			const FPlayAnimRequestPtr Request = InstanceData->PendingRequest.Request;
+			const FInjectionRequestPtr Request = InstanceData->PendingRequest.Request;
 
 			// Clear it out now in case we early out below
 			InstanceData->PendingRequest.Reset();
 
-			FPlayAnimBlendSettings BlendSettings;
+			FInjectionBlendSettings BlendSettings;
 			const UAnimNextAnimationGraph* AnimationGraph = nullptr;
 			FAnimNextDataInterfacePayload Payload;
 
 			if (Request)
 			{
 				// This is a new pending request, lookup the sub-graph to use with any factory trait we have and the desired animation object
-				FAnimNextPlayAnimRequestArgs& RequestArgs = Request->GetMutableArgs();
+				FAnimNextInjectionRequestArgs& RequestArgs = Request->GetMutableArgs();
 				if (RequestArgs.Object)
 				{
 					AnimationGraph = IGraphFactory::GetGraphFromObjectWithFallback(Context, Binding, RequestArgs.Object);
@@ -205,9 +205,9 @@ namespace UE::AnimNext
 				const FPlayAnimSlotRequest& OldSlotRequest = InstanceData->SlotRequests[InstanceData->CurrentlyActiveRequestIndex];
 				if (OldSlotRequest.State == EPlayAnimRequestState::Active)
 				{
-					auto StatusUpdateEvent = MakeTraitEvent<FPlayAnim_StatusUpdateEvent>();
+					auto StatusUpdateEvent = MakeTraitEvent<FInjection_StatusUpdateEvent>();
 					StatusUpdateEvent->Request = OldSlotRequest.Request;
-					StatusUpdateEvent->Status = EPlayAnimStatus::Playing | EPlayAnimStatus::Interrupted;
+					StatusUpdateEvent->Status = EInjectionStatus::Playing | EInjectionStatus::Interrupted;
 
 					Context.RaiseOutputTraitEvent(StatusUpdateEvent);
 				}
@@ -264,7 +264,7 @@ namespace UE::AnimNext
 			}
 
 			{
-				auto TimelineUpdateEvent = MakeTraitEvent<FPlayAnim_TimelineUpdateEvent>();
+				auto TimelineUpdateEvent = MakeTraitEvent<FInjection_TimelineUpdateEvent>();
 				TimelineUpdateEvent->Request = SlotRequest.Request;
 
 				// We don't have too many options here:
@@ -286,16 +286,16 @@ namespace UE::AnimNext
 
 			if (ActiveSlotRequest.State == EPlayAnimRequestState::Active)
 			{
-				const FAnimNextPlayAnimRequestArgs& RequestArgs = ActiveSlotRequest.Request->GetArgs();
+				const FAnimNextInjectionRequestArgs& RequestArgs = ActiveSlotRequest.Request->GetArgs();
 
 				const float BlendOutTime = RequestArgs.BlendOutSettings.Blend.BlendTime;
 				if (CurrentRequestTimeLeft <= BlendOutTime)
 				{
 					// We are ready to start blending out
 					{
-						auto StatusUpdateEvent = MakeTraitEvent<FPlayAnim_StatusUpdateEvent>();
+						auto StatusUpdateEvent = MakeTraitEvent<FInjection_StatusUpdateEvent>();
 						StatusUpdateEvent->Request = ActiveSlotRequest.Request;
-						StatusUpdateEvent->Status = EPlayAnimStatus::BlendingOut;
+						StatusUpdateEvent->Status = EInjectionStatus::BlendingOut;
 
 						Context.RaiseOutputTraitEvent(StatusUpdateEvent);
 					}
@@ -304,7 +304,7 @@ namespace UE::AnimNext
 					const int32 FreeRequestIndex = FindFreeRequestIndexOrAdd(*InstanceData);
 
 					FPlayAnimSlotRequest& FreeSlotRequest = InstanceData->SlotRequests[FreeRequestIndex];
-					FreeSlotRequest.Initialize(FPlayAnimRequestPtr(), RequestArgs.BlendOutSettings, nullptr, FAnimNextDataInterfacePayload());
+					FreeSlotRequest.Initialize(FInjectionRequestPtr(), RequestArgs.BlendOutSettings, nullptr, FAnimNextDataInterfacePayload());
 
 					const int32 OldChildIndex = InstanceData->CurrentlyActiveRequestIndex;
 					const int32 NewChildIndex = FreeRequestIndex;
@@ -399,9 +399,9 @@ namespace UE::AnimNext
 
 				{
 					// Queue our status update
-					auto StatusUpdateEvent = MakeTraitEvent<FPlayAnim_StatusUpdateEvent>();
+					auto StatusUpdateEvent = MakeTraitEvent<FInjection_StatusUpdateEvent>();
 					StatusUpdateEvent->Request = SlotRequest.Request;
-					StatusUpdateEvent->Status = EPlayAnimStatus::Playing;
+					StatusUpdateEvent->Status = EInjectionStatus::Playing;
 
 					Context.RaiseOutputTraitEvent(StatusUpdateEvent);
 				}
@@ -430,9 +430,9 @@ namespace UE::AnimNext
 
 				{
 					// Queue our status update
-					auto StatusUpdateEvent = MakeTraitEvent<FPlayAnim_StatusUpdateEvent>();
+					auto StatusUpdateEvent = MakeTraitEvent<FInjection_StatusUpdateEvent>();
 					StatusUpdateEvent->Request = SlotRequest.Request;
-					StatusUpdateEvent->Status = EPlayAnimStatus::Completed;
+					StatusUpdateEvent->Status = EInjectionStatus::Completed;
 
 					Context.RaiseOutputTraitEvent(StatusUpdateEvent);
 				}
@@ -496,7 +496,7 @@ namespace UE::AnimNext
 		if (InstanceData->SlotRequests.IsValidIndex(ChildIndex))
 		{
 			const FPlayAnimSlotRequest& SlotRequest = InstanceData->SlotRequests[ChildIndex];
-			if (SlotRequest.BlendSettings.BlendMode == EAnimNextPlayAnimBlendMode::Inertialization)
+			if (SlotRequest.BlendSettings.BlendMode == EAnimNextInjectionBlendMode::Inertialization)
 			{
 				return SlotRequest.BlendSettings.Blend.BlendTime;
 			}
diff --git a/Engine/Plugins/Experimental/AnimNextAnimGraph/Source/AnimNextAnimGraph/Private/PlayAnim/PlayAnimSlotTrait.h b/Engine/Plugins/Experimental/AnimNextAnimGraph/Source/AnimNextAnimGraph/Private/PlayAnim/PlayAnimSlotTrait.h
--- a/Engine/Plugins/Experimental/AnimNextAnimGraph/Source/AnimNextAnimGraph/Private/PlayAnim/PlayAnimSlotTrait.h
+++ b/Engine/Plugins/Experimental/AnimNextAnimGraph/Source/AnimNextAnimGraph/Private/PlayAnim/PlayAnimSlotTrait.h
@@ -13,20 +13,10 @@
 #include "TraitInterfaces/IUpdate.h"
 #include "Graph/AnimNextAnimationGraph.h"
 #include "Graph/AnimNextGraphInstance.h"
-#include "PlayAnim/PlayAnimEvents.h"
+#include "Injection/InjectionEvents.h"
 
 #include "PlayAnimSlotTrait.generated.h"
 
-USTRUCT(meta = (DisplayName = "PlayAnim Chooser Parameters"))
-struct FAnimNextPlayAnimChooserParameters
-{
-	GENERATED_BODY()
-
-	// The class type of the animation object
-	UPROPERTY()
-	TObjectPtr<UClass> AnimationObjectType;
-};
-
 USTRUCT(meta = (DisplayName = "PlayAnim Slot"))
 struct FAnimNextPlayAnimSlotTraitSharedData : public FAnimNextTraitSharedData
 {
@@ -77,7 +67,7 @@ namespace UE::AnimNext
 		struct FPlayAnimPendingRequest
 		{
 			// The PlayAnim request
-			FPlayAnimRequestPtr Request;
+			FInjectionRequestPtr Request;
 
 			// Whether or not a Stop request was issued
 			bool bStop = false;
@@ -96,10 +86,10 @@ namespace UE::AnimNext
 		struct FPlayAnimSlotRequest
 		{
 			// The PlayAnim request
-			FPlayAnimRequestPtr Request;
+			FInjectionRequestPtr Request;
 
 			// The blend settings to use
-			FPlayAnimBlendSettings BlendSettings;
+			FInjectionBlendSettings BlendSettings;
 
 			// The module used by the graph instance, as selected by the chooser
 			TObjectPtr<const UAnimNextAnimationGraph> AnimationGraph;
@@ -121,7 +111,7 @@ namespace UE::AnimNext
 			bool bWasRelevant = false;
 
 			// Initializes a request to begin playing
-			void Initialize(FPlayAnimRequestPtr InRequest, const FPlayAnimBlendSettings& InBlendSettings, const UAnimNextAnimationGraph* InAnimationGraph, FAnimNextDataInterfacePayload&& InPayload);
+			void Initialize(FInjectionRequestPtr InRequest, const FInjectionBlendSettings& InBlendSettings, const UAnimNextAnimationGraph* InAnimationGraph, FAnimNextDataInterfacePayload&& InPayload);
 		};
 
 		using FSharedData = FAnimNextPlayAnimSlotTraitSharedData;
@@ -146,8 +136,8 @@ namespace UE::AnimNext
 		static int32 FindFreeRequestIndexOrAdd(FInstanceData& InstanceData);
 
 		// Event handling impl
-		ETraitStackPropagation OnPlayEvent(const FExecutionContext& Context, FTraitBinding& Binding, FPlayAnim_PlayEvent& Event) const;
-		ETraitStackPropagation OnStopEvent(const FExecutionContext& Context, FTraitBinding& Binding, FPlayAnim_StopEvent& Event) const;
+		ETraitStackPropagation OnInjectEvent(const FExecutionContext& Context, FTraitBinding& Binding, FInjection_InjectEvent& Event) const;
+		ETraitStackPropagation OnUninjectEvent(const FExecutionContext& Context, FTraitBinding& Binding, FInjection_UninjectEvent& Event) const;
 
 		// IUpdate impl
 		virtual void PreUpdate(FUpdateTraversalContext& Context, const TTraitBinding<IUpdate>& Binding, const FTraitUpdateState& TraitState) const override;
diff --git a/Engine/Plugins/Experimental/AnimNextAnimGraph/Source/AnimNextAnimGraph/Private/TraitCore/ExecutionContext.cpp b/Engine/Plugins/Experimental/AnimNextAnimGraph/Source/AnimNextAnimGraph/Private/TraitCore/ExecutionContext.cpp
--- a/Engine/Plugins/Experimental/AnimNextAnimGraph/Source/AnimNextAnimGraph/Private/TraitCore/ExecutionContext.cpp
+++ b/Engine/Plugins/Experimental/AnimNextAnimGraph/Source/AnimNextAnimGraph/Private/TraitCore/ExecutionContext.cpp
@@ -12,6 +12,7 @@
 #include "Graph/AnimNextAnimationGraph.h"
 #include "Graph/AnimNextGraphInstance.h"
 #include "Components/SceneComponent.h"
+#include "Module/AnimNextModuleInstance.h"
 
 namespace UE::AnimNext
 {
@@ -83,26 +84,12 @@ namespace UE::AnimNext
 	{
 	}
 
-	FExecutionContext::FExecutionContext(FAnimNextGraphInstancePtr& InGraphInstance)
-		: FExecutionContext()
-	{
-		BindTo(InGraphInstance);
-	}
-
 	FExecutionContext::FExecutionContext(FAnimNextGraphInstance& InGraphInstance)
 		: FExecutionContext()
 	{
 		BindTo(InGraphInstance);
 	}
 
-	void FExecutionContext::BindTo(const FAnimNextGraphInstancePtr& InGraphInstance)
-	{
-		if (FAnimNextGraphInstance* InstanceImpl = InGraphInstance.GetImpl())
-		{
-			BindTo(*InstanceImpl);
-		}
-	}
-
 	void FExecutionContext::BindTo(FAnimNextGraphInstance& InGraphInstance)
 	{
 		FAnimNextGraphInstance* InRootGraphInstance = InGraphInstance.GetRootGraphInstance();
@@ -127,16 +114,6 @@ namespace UE::AnimNext
 		return RootGraphInstance != nullptr;
 	}
 
-	bool FExecutionContext::IsBoundTo(const FAnimNextGraphInstancePtr& InGraphInstance) const
-	{
-		if (FAnimNextGraphInstance* GraphInstance = InGraphInstance.GetImpl())
-		{
-			return RootGraphInstance == GraphInstance->GetRootGraphInstance();
-		}
-
-		return false;
-	}
-
 	bool FExecutionContext::IsBoundTo(const FAnimNextGraphInstance& InGraphInstance) const
 	{
 		return RootGraphInstance == InGraphInstance.GetRootGraphInstance();
