AnimNext: Injection

diff --git a/Engine/Plugins/Experimental/AnimNextAnimGraph/Source/AnimNextAnimGraph/Private/TraitInterfaces/IBlendStack.cpp b/Engine/Plugins/Experimental/AnimNextAnimGraph/Source/AnimNextAnimGraph/Private/TraitInterfaces/IBlendStack.cpp
--- a/Engine/Plugins/Experimental/AnimNextAnimGraph/Source/AnimNextAnimGraph/Private/TraitInterfaces/IBlendStack.cpp
+++ b/Engine/Plugins/Experimental/AnimNextAnimGraph/Source/AnimNextAnimGraph/Private/TraitInterfaces/IBlendStack.cpp
@@ -11,34 +11,46 @@ namespace UE::AnimNext
 #if WITH_EDITOR
 	const FText& IBlendStack::GetDisplayName() const
 	{
-		static FText InterfaceName = NSLOCTEXT("TraitInterfaces", "TraitInterface_IBlendStack_Name", "Sub Graph");
+		static FText InterfaceName = NSLOCTEXT("TraitInterfaces", "TraitInterface_IBlendStack_Name", "Blend Stack");
 		return InterfaceName;
 	}
 	const FText& IBlendStack::GetDisplayShortName() const
 	{
-		static FText InterfaceShortName = NSLOCTEXT("TraitInterfaces", "TraitInterface_IBlendStack_ShortName", "SBG");
+		static FText InterfaceShortName = NSLOCTEXT("TraitInterfaces", "TraitInterface_IBlendStack_ShortName", "BST");
 		return InterfaceShortName;
 	}
 #endif // WITH_EDITOR
 
-	void IBlendStack::PushGraph(FExecutionContext& Context, const TTraitBinding<IBlendStack>& Binding, IBlendStack::FGraphRequest&& GraphRequest) const
+	int32 IBlendStack::PushGraph(FExecutionContext& Context, const TTraitBinding<IBlendStack>& Binding, IBlendStack::FGraphRequest&& GraphRequest) const
 	{
 		TTraitBinding<IBlendStack> SuperBinding;
 		if (Binding.GetStackInterfaceSuper(SuperBinding))
 		{
-			SuperBinding.PushGraph(Context, MoveTemp(GraphRequest));
+			return SuperBinding.PushGraph(Context, MoveTemp(GraphRequest));
 		}
+		return INDEX_NONE;
 	}
 
-	void IBlendStack::GetActiveGraph(FExecutionContext& Context, const TTraitBinding<IBlendStack>& Binding, IBlendStack::FGraphRequestPtr& OutGraphRequest) const
+	int32 IBlendStack::GetActiveGraph(FExecutionContext& Context, const TTraitBinding<IBlendStack>& Binding, IBlendStack::FGraphRequestPtr& OutGraphRequest) const
 	{
 		TTraitBinding<IBlendStack> SuperBinding;
 		if (Binding.GetStackInterfaceSuper(SuperBinding))
 		{
-			SuperBinding.GetActiveGraph(Context, OutGraphRequest);
-			return;
+			return SuperBinding.GetActiveGraph(Context, OutGraphRequest);
 		}
 
 		OutGraphRequest = nullptr;
+		return INDEX_NONE;
+	}
+
+	IBlendStack::FGraphRequestPtr IBlendStack::GetGraph(FExecutionContext& Context, const TTraitBinding<IBlendStack>& Binding, int32 InChildIndex) const
+	{
+		TTraitBinding<IBlendStack> SuperBinding;
+		if (Binding.GetStackInterfaceSuper(SuperBinding))
+		{
+			return SuperBinding.GetGraph(Context, InChildIndex);
+		}
+
+		return nullptr;
 	}
 }
diff --git a/Engine/Plugins/Experimental/AnimNextAnimGraph/Source/AnimNextAnimGraph/Private/Traits/BlendStackTrait.cpp b/Engine/Plugins/Experimental/AnimNextAnimGraph/Source/AnimNextAnimGraph/Private/Traits/BlendStackTrait.cpp
--- a/Engine/Plugins/Experimental/AnimNextAnimGraph/Source/AnimNextAnimGraph/Private/Traits/BlendStackTrait.cpp
+++ b/Engine/Plugins/Experimental/AnimNextAnimGraph/Source/AnimNextAnimGraph/Private/Traits/BlendStackTrait.cpp
@@ -8,9 +8,12 @@
 #include "DataInterface/DataInterfaceStructAdapter.h"
 #include "Graph/AnimNextGraphInstance.h"
 #include "GraphInterfaces/AnimNextNativeDataInterface_AnimSequencePlayer.h"
+#include "Injection/ModuleInjectionDataInterfaceAdapter.h"
+#include "Logging/StructuredLog.h"
 #include "TraitCore/NodeInstance.h"
 #include "TraitCore/ExecutionContext.h"
 #include "TraitInterfaces/IGraphFactory.h"
+#include "TraitInterfaces/IInertializerBlend.h"
 #include "TraitInterfaces/IPayloadFactory.h"
 
 namespace UE::AnimNext
@@ -25,10 +28,12 @@ namespace UE::AnimNext
 		GeneratorMacro(IGarbageCollection) \
 		GeneratorMacro(IHierarchy) \
 		GeneratorMacro(ISmoothBlend) \
+		GeneratorMacro(IInertializerBlend) \
 		GeneratorMacro(IAttributeProvider) \
 		GeneratorMacro(ITimeline) \
 		GeneratorMacro(IUpdateTraversal) \
 		GeneratorMacro(IBlendStack) \
+		GeneratorMacro(IUpdate) \
 
 	GENERATE_ANIM_TRAIT_IMPLEMENTATION(FBlendStackCoreTrait, TRAIT_INTERFACE_ENUMERATOR, NULL_ANIM_TRAIT_INTERFACE_ENUMERATOR, NULL_ANIM_TRAIT_EVENT_ENUMERATOR)
 	#undef TRAIT_INTERFACE_ENUMERATOR
@@ -93,6 +98,24 @@ namespace UE::AnimNext
 		return InstanceData.ChildGraphs.AddDefaulted();
 	}
 
+	void FBlendStackCoreTrait::PreUpdate(FUpdateTraversalContext& Context, const TTraitBinding<IUpdate>& Binding, const FTraitUpdateState& TraitState) const
+	{
+		IUpdate::PreUpdate(Context, Binding, TraitState);
+#if WITH_EDITOR
+		// In editor we may need to rebind our data if an outer hosting module was recompiled
+		FInstanceData* InstanceData = Binding.GetInstanceData<FInstanceData>();
+		for (FGraphState& ChildGraph : InstanceData->ChildGraphs)
+		{
+			if(ChildGraph.State == EGraphState::Active && ChildGraph.Instance.IsValid() && ChildGraph.Instance->RequiresPublicVariableBinding())
+			{
+				FModuleInjectionDataInterfaceAdapter ModuleAdapter(Context.GetRootGraphInstance().GetModuleInstance(), ChildGraph.Request.BindingModuleHandle);
+				IDataInterfaceHost* DataInterfaceHost = &ModuleAdapter;
+				ChildGraph.Instance->BindPublicVariables(ChildGraph.Request.GraphPayload.Get(), TConstArrayView<IDataInterfaceHost*>(&DataInterfaceHost, 1));
+			}
+		}
+#endif
+	}
+
 	uint32 FBlendStackCoreTrait::GetNumChildren(const FExecutionContext& Context, const TTraitBinding<IHierarchy>& Binding) const
 	{
 		const FInstanceData* InstanceData = Binding.GetInstanceData<FInstanceData>();
@@ -133,6 +156,25 @@ namespace UE::AnimNext
 		}
 	}
 
+	float FBlendStackCoreTrait::GetBlendWeight(FExecutionContext& Context, const TTraitBinding<IDiscreteBlend>& Binding, int32 ChildIndex) const
+	{
+		const FInstanceData* InstanceData = Binding.GetInstanceData<FInstanceData>();
+
+		if (ChildIndex == InstanceData->CurrentlyActiveGraphIndex)
+		{
+			return 1.0f;	// Active child has full weight
+		}
+		else if (InstanceData->ChildGraphs.IsValidIndex(ChildIndex))
+		{
+			return 0.0f;	// Other children have no weight
+		}
+		else
+		{
+			// Invalid child index
+			return -1.0f;
+		}
+	}
+
 	int32 FBlendStackCoreTrait::GetBlendDestinationChildIndex(FExecutionContext& Context, const TTraitBinding<IDiscreteBlend>& Binding) const
 	{
 		const FInstanceData* InstanceData = Binding.GetInstanceData<FInstanceData>();
@@ -158,10 +200,26 @@ namespace UE::AnimNext
 		{
 			FGraphState& Graph = InstanceData->ChildGraphs[ChildIndex];
 
-			//@TODO: Remove or implement entry points once we decide if we still need them.
-			static const FName EntryPoint = NAME_None;
-			Graph.Instance = Graph.Request.AnimationGraph->AllocateInstance(nullptr, &Binding.GetTraitPtr().GetNodeInstance()->GetOwner(), EntryPoint);
-			Graph.ChildPtr = Graph.Instance->GetGraphRootPtr();
+			switch(Graph.Request.Type)
+			{
+			case EGraphRequestType::Owned:
+				{
+					//@TODO: Remove or implement entry points once we decide if we still need them.
+					const FName EntryPoint = NAME_None;
+					FAnimNextGraphInstance& Owner = Binding.GetTraitPtr().GetNodeInstance()->GetOwner();
+					Graph.Instance = Graph.Request.AnimationGraph->AllocateInstance(Owner.GetModuleInstance(), &Owner, EntryPoint);
+					Graph.ChildPtr = Graph.Instance->GetGraphRootPtr();
+					FModuleInjectionDataInterfaceAdapter ModuleAdapter(Context.GetRootGraphInstance().GetModuleInstance(), Graph.Request.BindingModuleHandle);
+					IDataInterfaceHost* DataInterfaceHost = &ModuleAdapter;
+					Graph.Instance->BindPublicVariables(Graph.Request.GraphPayload.Get(), TConstArrayView<IDataInterfaceHost*>(&DataInterfaceHost, 1));
+					break;
+				}
+			case EGraphRequestType::Child:
+				{
+					Graph.ChildPtr = Graph.Request.ChildPtr;
+					break;
+				}
+			}
 		}
 	}
 
@@ -177,13 +235,44 @@ namespace UE::AnimNext
 		}
 	}
 
+	float FBlendStackCoreTrait::GetBlendTime(FExecutionContext& Context, const TTraitBinding<IInertializerBlend>& Binding, int32 ChildIndex) const
+	{
+		const FInstanceData* InstanceData = Binding.GetInstanceData<FInstanceData>();
+		if (InstanceData->ChildGraphs.IsValidIndex(ChildIndex))
+		{
+			const FGraphState& Graph = InstanceData->ChildGraphs[ChildIndex];
+			if (Graph.Request.BlendMode == EBlendMode::Inertialization)
+			{
+				return Graph.Request.BlendArgs.BlendTime;
+			}
+			else
+			{
+				// Not an inertializing blend
+				return 0.0f;
+			}
+		}
+		else
+		{
+			// Unknown child
+			return 0.0f;
+		}
+	}
+
 	float FBlendStackCoreTrait::GetBlendTime(FExecutionContext& Context, const TTraitBinding<ISmoothBlend>& Binding, int32 ChildIndex) const
 	{
 		const FInstanceData* InstanceData = Binding.GetInstanceData<FInstanceData>();
 		if (InstanceData->ChildGraphs.IsValidIndex(ChildIndex))
 		{
 			const FGraphState& Graph = InstanceData->ChildGraphs[ChildIndex];
-			return Graph.Request.BlendArgs.BlendTime;
+			if (Graph.Request.BlendMode == EBlendMode::Standard)
+			{
+				return Graph.Request.BlendArgs.BlendTime;
+			}
+			else
+			{
+				// Not a standard blend
+				return 0.0f;
+			}
 		}
 		else
 		{
@@ -295,15 +384,45 @@ namespace UE::AnimNext
 		}
 	}
 
-	void FBlendStackCoreTrait::PushGraph(FExecutionContext& Context, const TTraitBinding<IBlendStack>& Binding, IBlendStack::FGraphRequest&& GraphRequest) const
+	int32 FBlendStackCoreTrait::PushGraph(FExecutionContext& Context, const TTraitBinding<IBlendStack>& Binding, IBlendStack::FGraphRequest&& GraphRequest) const
 	{
-		if (GraphRequest.AnimationGraph == nullptr)
+		//@TODO: Add depth limit and saturation policies
+
+		// Validate request
+		switch(GraphRequest.Type)
 		{
-			return;
+		case EGraphRequestType::Owned:
+			{
+				if (GraphRequest.AnimationGraph == nullptr)
+				{
+					return INDEX_NONE;
+				}
+
+				// Check for re-entrancy and early-out if we are linking back to the current instance or one of its parents
+				const FName EntryPoint = GraphRequest.AnimationGraph->DefaultEntryPoint;
+				const FAnimNextGraphInstance* OwnerGraphInstance = &Binding.GetTraitPtr().GetNodeInstance()->GetOwner();
+				while (OwnerGraphInstance != nullptr)
+				{
+					if (OwnerGraphInstance->UsesAnimationGraph(GraphRequest.AnimationGraph) && OwnerGraphInstance->UsesEntryPoint(EntryPoint))
+					{
+						UE_LOGFMT(LogAnimation, Warning, "Ignoring PushGraph request for {0}, re-entrancy detected", GraphRequest.AnimationGraph->GetFName());
+						return INDEX_NONE;
+					}
+
+					OwnerGraphInstance = OwnerGraphInstance->GetParentGraphInstance();
+				}
+				break;
+			}
+		case EGraphRequestType::Child:
+			{
+				if (!GraphRequest.ChildPtr.IsValid())
+				{
+					return INDEX_NONE;
+				}
+				break;
+			}
 		}
 
-		//@TODO: Add depth limit and saturation policies
-		
 		FInstanceData* InstanceData = Binding.GetInstanceData<FInstanceData>();
 
 		const int32 OldChildIndex = InstanceData->CurrentlyActiveGraphIndex;
@@ -311,16 +430,16 @@ namespace UE::AnimNext
 		FBlendStackCoreTrait::FGraphState& Graph = InstanceData->ChildGraphs[NewChildIndex];
 		Graph.Initialize(MoveTemp(GraphRequest));
 
+		InstanceData->CurrentlyActiveGraphIndex = NewChildIndex;
+
 		TTraitBinding<IDiscreteBlend> DiscreteBlendTrait;
 		Binding.GetStackInterface(DiscreteBlendTrait);
 		DiscreteBlendTrait.OnBlendTransition(Context, OldChildIndex, NewChildIndex);
 
-		InstanceData->CurrentlyActiveGraphIndex = NewChildIndex;
-
-		Graph.Instance->BindPublicVariables(Graph.Request.GraphPayload.Get());
+		return NewChildIndex;
 	}
 
-	void FBlendStackCoreTrait::GetActiveGraph(FExecutionContext& Context, const TTraitBinding<IBlendStack>& Binding, IBlendStack::FGraphRequestPtr& OutGraphRequest) const
+	int32 FBlendStackCoreTrait::GetActiveGraph(FExecutionContext& Context, const TTraitBinding<IBlendStack>& Binding, IBlendStack::FGraphRequestPtr& OutGraphRequest) const
 	{
 		FInstanceData* InstanceData = Binding.GetInstanceData<FInstanceData>();
 		const int32 CurrentlyActiveGraphIndex = InstanceData->CurrentlyActiveGraphIndex;
@@ -328,14 +447,27 @@ namespace UE::AnimNext
 		{
 			FGraphState& GraphState = InstanceData->ChildGraphs[CurrentlyActiveGraphIndex];
 			OutGraphRequest = &GraphState.Request;
-			return;
+			return CurrentlyActiveGraphIndex;
 		}
 
-		return IBlendStack::GetActiveGraph(Context, Binding, OutGraphRequest);
+		return INDEX_NONE;
+	}
+
+	IBlendStack::FGraphRequestPtr FBlendStackCoreTrait::GetGraph(FExecutionContext& Context, const TTraitBinding<IBlendStack>& Binding, int32 InChildIndex) const
+	{
+		FInstanceData* InstanceData = Binding.GetInstanceData<FInstanceData>();
+		if(!InstanceData->ChildGraphs.IsValidIndex(InChildIndex))
+		{
+			return nullptr;
+		}
+
+		return &InstanceData->ChildGraphs[InChildIndex].Request;
 	}
 
 	void FBlendStackTrait::PreUpdate(FUpdateTraversalContext& Context, const TTraitBinding<IUpdate>& Binding, const FTraitUpdateState& TraitState) const
 	{
+		FBlendStackCoreTrait::PreUpdate(Context, Binding, TraitState);
+
 		FInstanceData* InstanceData = Binding.GetInstanceData<FInstanceData>();
 		const FSharedData* SharedData = Binding.GetSharedData<FSharedData>();
  
@@ -370,6 +502,8 @@ namespace UE::AnimNext
 
 	void FBlendStackRequesterTrait::PreUpdate(FUpdateTraversalContext& Context, const TTraitBinding<IUpdate>& Binding, const FTraitUpdateState& TraitState) const
 	{
+		IUpdate::PreUpdate(Context, Binding, TraitState);
+
 		const FSharedData* SharedData = Binding.GetSharedData<FSharedData>();
 
 		const UObject* DesiredObject = SharedData->GetObject(Binding);
diff --git a/Engine/Plugins/Experimental/AnimNextAnimGraph/Source/AnimNextAnimGraph/Private/Traits/CallFunction.cpp b/Engine/Plugins/Experimental/AnimNextAnimGraph/Source/AnimNextAnimGraph/Private/Traits/CallFunction.cpp
--- a/Engine/Plugins/Experimental/AnimNextAnimGraph/Source/AnimNextAnimGraph/Private/Traits/CallFunction.cpp
+++ b/Engine/Plugins/Experimental/AnimNextAnimGraph/Source/AnimNextAnimGraph/Private/Traits/CallFunction.cpp
@@ -8,6 +8,7 @@
 #include "TraitCore/ExecutionContext.h"
 #include "TraitCore/NodeInstance.h"
 #include "AnimNextAnimGraphStats.h"
+#include "AnimNextExecuteContext.h"
 
 DEFINE_STAT(STAT_AnimNext_CallFunction);
 
diff --git a/Engine/Plugins/Experimental/AnimNextAnimGraph/Source/AnimNextAnimGraph/Private/Traits/NotifyDispatcher.cpp b/Engine/Plugins/Experimental/AnimNextAnimGraph/Source/AnimNextAnimGraph/Private/Traits/NotifyDispatcher.cpp
--- a/Engine/Plugins/Experimental/AnimNextAnimGraph/Source/AnimNextAnimGraph/Private/Traits/NotifyDispatcher.cpp
+++ b/Engine/Plugins/Experimental/AnimNextAnimGraph/Source/AnimNextAnimGraph/Private/Traits/NotifyDispatcher.cpp
@@ -49,7 +49,7 @@ namespace UE::AnimNext
 			// Ensure we have a handler component on the module
 			if(FAnimNextModuleInstance* ModuleInstance = Context.GetRootGraphInstance().GetModuleInstance())
 			{
-				ModuleInstance->GetComponent<FModuleInstanceComponent_NotifyDispatcher>();
+				ModuleInstance->GetComponent<FAnimNextNotifyDispatcherComponent>();
 			}
 
 			auto NotifyDispatchEvent = MakeTraitEvent<FNotifyDispatchEvent>();
@@ -110,7 +110,7 @@ namespace UE::AnimNext
 	}
 }
 
-void FModuleInstanceComponent_NotifyDispatcher::OnTraitEvent(FAnimNextTraitEvent& Event)
+void FAnimNextNotifyDispatcherComponent::OnTraitEvent(FAnimNextTraitEvent& Event)
 {
 	if(UE::AnimNext::FNotifyDispatchEvent* NotifyDispatchEvent = Event.AsType<UE::AnimNext::FNotifyDispatchEvent>())
 	{
@@ -118,7 +118,7 @@ void FModuleInstanceComponent_NotifyDispatcher::OnTraitEvent(FAnimNextTraitEvent
 	}
 }
 
-void FModuleInstanceComponent_NotifyDispatcher::OnEndExecution(float InDeltaTime)
+void FAnimNextNotifyDispatcherComponent::OnEndExecution(float InDeltaTime)
 {
 	using namespace UE::AnimNext;
 	
@@ -242,7 +242,7 @@ void FModuleInstanceComponent_NotifyDispatcher::OnEndExecution(float InDeltaTime
 	NotifyQueue.AnimNotifies.Reset();
 }
 
-void FModuleInstanceComponent_NotifyDispatcher::TriggerSingleAnimNotify(float InDeltaTime, UE::AnimNext::FNotifyQueueDispatchEvent& Dispatcher, const FAnimNotifyEventReference& EventReference)
+void FAnimNextNotifyDispatcherComponent::TriggerSingleAnimNotify(float InDeltaTime, UE::AnimNext::FNotifyQueueDispatchEvent& Dispatcher, const FAnimNotifyEventReference& EventReference)
 {
 	// This is for non 'state' anim notifies.
 	const FAnimNotifyEvent* AnimNotifyEvent = EventReference.GetNotify();
@@ -263,6 +263,12 @@ void FModuleInstanceComponent_NotifyDispatcher::TriggerSingleAnimNotify(float In
 		{
 			// Named notifies can be handled by a custom module event on our worker thread, no need to queue
 			ModuleInstance.RunRigVMEvent(AnimNotifyEvent->NotifyName, InDeltaTime);
+
+			// Also dispatch to any prerequisite modules
+			ModuleInstance.ForEachPrerequisite([AnimNotifyEvent, InDeltaTime](FAnimNextModuleInstance& InPrerequisiteInstance)
+			{
+				InPrerequisiteInstance.RunRigVMEvent(AnimNotifyEvent->NotifyName, InDeltaTime);
+			});
 		}
 	}
 }
\ No newline at end of file
diff --git a/Engine/Plugins/Experimental/AnimNextAnimGraph/Source/AnimNextAnimGraph/Private/Traits/NotifyDispatcher.h b/Engine/Plugins/Experimental/AnimNextAnimGraph/Source/AnimNextAnimGraph/Private/Traits/NotifyDispatcher.h
--- a/Engine/Plugins/Experimental/AnimNextAnimGraph/Source/AnimNextAnimGraph/Private/Traits/NotifyDispatcher.h
+++ b/Engine/Plugins/Experimental/AnimNextAnimGraph/Source/AnimNextAnimGraph/Private/Traits/NotifyDispatcher.h
@@ -33,17 +33,17 @@ struct FAnimNextNotifyDispatcherTraitSharedData : public FAnimNextTraitSharedDat
 };
 
 USTRUCT()
-struct FModuleInstanceComponent_NotifyDispatcher : public FModuleInstanceComponent
+struct FAnimNextNotifyDispatcherComponent : public FAnimNextModuleInstanceComponent
 {
 	GENERATED_BODY()
 
-	FModuleInstanceComponent_NotifyDispatcher() = default;
+	FAnimNextNotifyDispatcherComponent() = default;
 
-	FModuleInstanceComponent_NotifyDispatcher(FAnimNextModuleInstance& InOwnerInstance)
-		: FModuleInstanceComponent(InOwnerInstance)
+	FAnimNextNotifyDispatcherComponent(FAnimNextModuleInstance& InOwnerInstance)
+		: FAnimNextModuleInstanceComponent(InOwnerInstance)
 	{}
 
-	// FModuleInstanceComponent interface
+	// FAnimNextModuleInstanceComponent interface
 	virtual void OnTraitEvent(FAnimNextTraitEvent& Event) override;
 	virtual void OnEndExecution(float InDeltaTime) override;
 
diff --git a/Engine/Plugins/Experimental/AnimNextAnimGraph/Source/AnimNextAnimGraph/Public/Graph/AnimNextAnimationGraph.h b/Engine/Plugins/Experimental/AnimNextAnimGraph/Source/AnimNextAnimGraph/Public/Graph/AnimNextAnimationGraph.h
new file mode 100644
--- /dev/null
+++ b/Engine/Plugins/Experimental/AnimNextAnimGraph/Source/AnimNextAnimGraph/Public/Graph/AnimNextAnimationGraph.h
@@ -0,0 +1,162 @@
+// Copyright Epic Games, Inc. All Rights Reserved.
+
+#pragma once
+
+#include "CoreMinimal.h"
+#include "Graph/AnimNextGraphEntryPoint.h"
+#include "DataInterface/AnimNextDataInterface.h"
+#include "RigVMCore/RigVM.h"
+#include "TraitCore/TraitPtr.h"
+#include "TraitCore/TraitHandle.h"
+#include "TraitCore/EntryPointHandle.h"
+#include "RigVMHost.h"
+#include "AnimNextGraphEvaluatorExecuteDefinition.h"
+
+#include "AnimNextAnimationGraph.generated.h"
+
+class UEdGraph;
+class UAnimNextModule;
+class UAnimGraphNode_AnimNextGraph;
+struct FAnimNode_AnimNextGraph;
+struct FRigUnit_AnimNextGraphEvaluator;
+struct FAnimNextGraphInstance;
+struct FAnimNextScheduleGraphTask;
+struct FAnimNextEditorParam;
+struct FAnimNextParam;
+struct FAnimNextModuleInstance;
+struct FAnimNextModuleAnimGraphComponent;
+
+namespace UE::AnimNext
+{
+	struct FContext;
+	struct FExecutionContext;
+	class FAnimNextModuleImpl;
+	struct FTestUtils;
+	struct FParametersProxy;
+	struct FPlayAnimSlotTrait;
+	struct FBlendStackCoreTrait;
+}
+
+namespace UE::AnimNext::UncookedOnly
+{
+	struct FUtils;
+}
+
+namespace UE::AnimNext::Editor
+{
+	class FModuleEditor;
+	class FVariableCustomization;
+}
+
+namespace UE::AnimNext::AnimGraph
+{
+	class FAnimNextAnimGraphModule;
+}
+
+namespace UE::AnimNext::Graph
+{
+	extern ANIMNEXTANIMGRAPH_API const FName EntryPointName;
+	extern ANIMNEXTANIMGRAPH_API const FName ResultName;
+}
+
+// A user-created collection of animation logic & data
+UCLASS(BlueprintType)
+class ANIMNEXTANIMGRAPH_API UAnimNextAnimationGraph : public UAnimNextDataInterface
+{
+	GENERATED_BODY()
+
+public:
+	UAnimNextAnimationGraph(const FObjectInitializer& ObjectInitializer);
+
+	// UObject interface
+	virtual void Serialize(FArchive& Ar) override;
+
+	// Allocates an instance of the graph with an optional module and/or parent graph instance
+	// @param	InModuleInstance			The module instance to use
+	// @param	InParentGraphInstance		The parent graph instance to use
+	// @param	InEntryPoint				The entry point to use. If this is NAME_None then the default entry point for this graph is used
+	// @return the allocated instance. This can be invalid if the entry point was not found or the graph's root trait is invalid.
+	TSharedPtr<FAnimNextGraphInstance> AllocateInstance(FAnimNextModuleInstance* InModuleInstance = nullptr, FAnimNextGraphInstance* InParentGraphInstance = nullptr, FName InEntryPoint = NAME_None) const;
+
+private:
+	// Loads the graph data from the provided archive buffer and returns true on success, false otherwise
+	bool LoadFromArchiveBuffer(const TArray<uint8>& SharedDataArchiveBuffer);
+
+#if WITH_EDITORONLY_DATA
+	// During graph compilation, if we have existing graph instances, we freeze them by releasing their memory before thawing them
+	// Freezing is a partial release of resources that retains the necessary information to re-create things safely
+	void FreezeGraphInstances();
+
+	// During graph compilation, once compilation is done we thaw existing graph instances to reallocate their memory
+	void ThawGraphInstances();
+#endif
+	
+	friend class UAnimNextAnimationGraphFactory;
+	friend class UAnimNextAnimationGraph_EditorData;
+	friend class UAnimNextVariableEntry;
+	friend struct UE::AnimNext::UncookedOnly::FUtils;
+	friend class UE::AnimNext::Editor::FModuleEditor;
+	friend struct UE::AnimNext::FTestUtils;
+	friend FAnimNextGraphInstance;
+	friend class UAnimGraphNode_AnimNextGraph;
+	friend UE::AnimNext::FExecutionContext;
+	friend struct FAnimNextScheduleGraphTask;
+	friend UE::AnimNext::AnimGraph::FAnimNextAnimGraphModule;
+	friend class UE::AnimNext::Editor::FVariableCustomization;
+	friend struct UE::AnimNext::FParametersProxy;
+	friend struct UE::AnimNext::FPlayAnimSlotTrait;
+	friend struct UE::AnimNext::FBlendStackCoreTrait;
+	friend struct FAnimNextModuleAnimGraphComponent;
+
+protected:
+#if WITH_EDITORONLY_DATA
+	mutable FCriticalSection GraphInstancesLock;
+
+	// This is a list of live graph instances that have been allocated, used in the editor to reset instances when we re-compile/live edit
+	mutable TSet<FAnimNextGraphInstance*> GraphInstances;
+#endif
+
+	// This is the execute method definition used by a graph to evaluate latent pins
+	UPROPERTY()
+	FAnimNextGraphEvaluatorExecuteDefinition ExecuteDefinition;
+
+	// Data for each entry point in this graph
+	UPROPERTY()
+	TArray<FAnimNextGraphEntryPoint> EntryPoints;
+
+	// This is a resolved handle to the root trait in our graph, for each entry point 
+	TMap<FName, FAnimNextTraitHandle> ResolvedRootTraitHandles;
+
+	// This is an index into EntryPoints, for each entry point
+	TMap<FName, int32> ResolvedEntryPoints;
+
+	// This is the graph shared data used by the trait system, the output of FTraitReader
+	// We de-serialize manually into this buffer from the archive buffer, this is never saved on disk
+	TArray<uint8> SharedDataBuffer;
+
+	// This is a list of all referenced UObjects in the graph shared data
+	// We collect all the references here to make it quick and easy for the GC to query them
+	// It means that object references in the graph shared data are not visited at runtime by the GC (they are immutable)
+	// The shared data serialization archive stores indices to these to perform UObject serialization
+	UPROPERTY()
+	TArray<TObjectPtr<UObject>> GraphReferencedObjects;
+
+	// This is a list of all referenced soft objects in the graph shared data
+	// Used to serialize the soft objects correctly as we dont use FArchiveUObject
+	UPROPERTY()
+	TArray<FSoftObjectPath> GraphReferencedSoftObjects;
+
+	// The entry point that this graph defaults to using
+	UPROPERTY(EditAnywhere, Category = "Graph")
+	FName DefaultEntryPoint = TEXT("Root");
+
+	// Default state for this graph
+	UPROPERTY()
+	FAnimNextGraphState DefaultState;
+
+#if WITH_EDITORONLY_DATA
+	// This buffer holds the output of the FTraitWriter post compilation
+	// We serialize it manually and it is discarded at runtime
+	TArray<uint8> SharedDataArchiveBuffer;
+#endif
+};
diff --git a/Engine/Plugins/Experimental/AnimNextAnimGraph/Source/AnimNextAnimGraph/Public/Graph/AnimNextGraphInstance.h b/Engine/Plugins/Experimental/AnimNextAnimGraph/Source/AnimNextAnimGraph/Public/Graph/AnimNextGraphInstance.h
--- a/Engine/Plugins/Experimental/AnimNextAnimGraph/Source/AnimNextAnimGraph/Public/Graph/AnimNextGraphInstance.h
+++ b/Engine/Plugins/Experimental/AnimNextAnimGraph/Source/AnimNextAnimGraph/Public/Graph/AnimNextGraphInstance.h
@@ -14,7 +14,6 @@
 
 class FReferenceCollector;
 
-struct FAnimNextGraphInstancePtr;
 struct FRigUnit_AnimNextGraphEvaluator;
 class UAnimNextAnimationGraph;
 class UAnimNextModule;
@@ -22,6 +21,10 @@ struct FAnimNextModuleInstance;
 class FRigVMTraitScope;
 struct FRigVMExtendedExecuteContext;
 struct FAnimNextExecuteContext;
+struct FAnimNextModuleInjectionComponent;
+struct FRigUnit_AnimNextRunAnimationGraph_v1;
+struct FRigUnit_AnimNextRunAnimationGraph_v2;
+struct FAnimNextStateTreeRigVMConditionBase;
 
 namespace UE::AnimNext
 {
@@ -32,6 +35,11 @@ namespace UE::AnimNext
 	struct FTraitStackBinding;
 	struct FPlayAnimSlotTrait;
 	struct FBlendStackCoreTrait;
+	struct FInjectionRequest;
+	struct FInjectionSiteTrait;
+	struct FStateTreeTrait;
+	struct FCachedBindingInfo;
+	struct FInjectionUtils;
 }
 
 using GraphInstanceComponentMapType = TMap<FName, TSharedPtr<UE::AnimNext::FGraphInstanceComponent>>;
@@ -111,10 +119,9 @@ struct ANIMNEXTANIMGRAPH_API FAnimNextGraphInstance : public FAnimNextDataInterf
 	// Called each time the graph updates
 	void Update();
 
-	// Get the hosting instance, if any, that owns us
-	FAnimNextDataInterfaceInstance* GetHost() const;
+	// Get a reference to the internal public variables
+	const FInstancedPropertyBag& GetVariables() const { return Variables; }
 
-	FInstancedPropertyBag& GetVariables() { return Variables; }
 private:
 	// Returns a pointer to the specified component, or nullptr if not found
 	UE::AnimNext::FGraphInstanceComponent* TryGetComponent(int32 ComponentNameHash, FName ComponentName) const;
@@ -138,12 +145,18 @@ private:
 	void OnModuleCompiled(UAnimNextModule* InModule);
 #endif
 
-	template<typename HostType>
-	bool BindToHostHelper(const HostType& InHost, bool bInAutoBind);
+	// Check if variable binding is required
+	bool RequiresPublicVariableBinding() const
+	{
+		return PublicVariablesState == EPublicVariablesState::Unbound;
+	}
+
+	// Helper function used to bind to cached bindings after they are built
+	bool BindToCachedBindings();
 
 	// Bind the variables in the supplied traits in scope to their respective public variables, represented by structs held on a host
 	// Note: struct views are mutable here as bindings allow writes!
-	void BindPublicVariables(TArrayView<FStructView> InHostStructs);
+	void BindPublicVariables(TArrayView<FStructView> InHostStructs, TConstArrayView<UE::AnimNext::IDataInterfaceHost*> InHosts = TConstArrayView<UE::AnimNext::IDataInterfaceHost*>());
 
 	// Bind the variables in the supplied traits in scope to their respective public variables, so they point at host memory
 	void BindPublicVariables(TConstArrayView<UE::AnimNext::IDataInterfaceHost*> InHosts);
@@ -151,6 +164,9 @@ private:
 	// Unbind any public variables that were pointing at host memory and re-point them at the internal defaults
 	void UnbindPublicVariables();
 
+	// Get a reference to the internal public variables
+	FInstancedPropertyBag& GetMutableVariables() { return Variables; }
+
 	// The entry point in Graph that this instance corresponds to 
 	FName EntryPoint;
 
@@ -160,22 +176,44 @@ private:
 	// The module instance that owns the root, us and the components
 	FAnimNextModuleInstance* ModuleInstance = nullptr;
 
-	// The graph instance that owns us
-	FAnimNextGraphInstance* ParentGraphInstance = nullptr;
-
 	// The root graph instance that owns us and the components
 	FAnimNextGraphInstance* RootGraphInstance = nullptr;
 
-#if WITH_EDITORONLY_DATA
-	struct FCachedVariableBinding
+	struct FCachedDataInterfaceBinding
 	{
-		FName VariableName;
-		uint8* Memory = nullptr;
+		struct FVariable
+		{
+			FVariable() = default;
+
+			FVariable(const FName& InVariableName, int32 InVariableIndex, const FProperty* InProperty, uint8* InMemory)
+				: VariableName(InVariableName)
+				, VariableIndex(InVariableIndex)
+				, Property(InProperty)
+				, Memory(InMemory)
+			{
+			}
+
+			FName VariableName;
+			int32 VariableIndex = INDEX_NONE;
+			const FProperty* Property = nullptr;
+			uint8* Memory = nullptr;
+		};
+
+		const UAnimNextDataInterface* DataInterface = nullptr;
+		TArray<FVariable> CachedBindings;
 	};
 
-	// Cached public variable bindings used to correctly thaw instances with input pin bindings
-	TArray<FCachedVariableBinding> CachedVariableBindings;
-#endif
+	void UpdateCachedBindingsForHost();
+	void UpdateCachedBindingsForHost(const FAnimNextModuleInstance& InHost);
+	void UpdateCachedBindingsForHost(const UE::AnimNext::IDataInterfaceHost& InHost);
+
+	// Helper function used to cache a host's bindings, overriding any of the already cached bindings that are valid.
+	// Template to allow binding to both FAnimNextDataInterfaceInstance and IDataInterfaceHost without an abstraction between the two.
+	template<typename HostType>
+	void UpdateCachedBindingsForHostHelper(const HostType& InHost);
+
+	// Bindings we have to our outer host(s), per data interface
+	TArray<FCachedDataInterfaceBinding> CachedBindings;
 
 	// Graph instance components that persist from update to update
 	GraphInstanceComponentMapType Components;
@@ -196,10 +234,17 @@ private:
 	friend UAnimNextAnimationGraph;			// The graph is the one that allocates instances
 	friend FRigUnit_AnimNextGraphEvaluator;	// We evaluate the instance
 	friend UE::AnimNext::FExecutionContext;
-	friend FAnimNextGraphInstancePtr;
 	friend UE::AnimNext::FTraitStackBinding;
 	friend UE::AnimNext::FPlayAnimSlotTrait;
 	friend UE::AnimNext::FBlendStackCoreTrait;
+	friend UE::AnimNext::FInjectionRequest;
+	friend UE::AnimNext::FInjectionSiteTrait;
+	friend UE::AnimNext::FStateTreeTrait;	// Temp - remove: Needs to access mutable variables to copy-in data OnBecomeRelevant
+	friend FAnimNextModuleInjectionComponent;
+	friend FRigUnit_AnimNextRunAnimationGraph_v1;
+	friend FRigUnit_AnimNextRunAnimationGraph_v2;
+	friend UE::AnimNext::FInjectionUtils;
+	friend FAnimNextStateTreeRigVMConditionBase;	// Temp - remove: Needs to access mutable variables to copy-in data for function shim
 };
 
 template<>
diff --git a/Engine/Plugins/Experimental/AnimNextAnimGraph/Source/AnimNextAnimGraph/Public/Injection/InjectionUtils.h b/Engine/Plugins/Experimental/AnimNextAnimGraph/Source/AnimNextAnimGraph/Public/Injection/InjectionUtils.h
new file mode 100644
--- /dev/null
+++ b/Engine/Plugins/Experimental/AnimNextAnimGraph/Source/AnimNextAnimGraph/Public/Injection/InjectionUtils.h
@@ -0,0 +1,133 @@
+﻿// Copyright Epic Games, Inc. All Rights Reserved.
+
+#pragma once
+
+#include "CoreMinimal.h"
+#include "InjectionRequest.h"
+
+struct FAnimNextDataInterfaceInstance;
+class UAnimNextComponent;
+
+namespace UE::AnimNext
+{
+	struct IEvaluationModifier;
+}
+
+namespace UE::AnimNext
+{
+	// A collection of wrapper functions for common injection operations
+	struct FInjectionUtils
+	{
+		// Injects an animation sequence into a running AnimNext module
+		// @param InComponent          The UAnimNextComponent that hosts the module being played on
+		// @param InHost               The object that hosts the module (e.g. UAnimNextComponent) being played on
+		// @param InModuleHandle       The module handle being injected into
+		// @param InSiteName           The injection site to 'play' the animation sequence on
+		// @param InAnimSequence       The animation sequence to play
+		// @param InPlayRate           The play rate multiplier to use
+		// @param InStartPosition      The start position to use
+		// @param InBlendInSettings    The blend in settings
+		// @param InBlendOutSettings   The blend in settings
+		// @param InLifetimeEvents     Delegates called for various request lifetime events
+		// @return the new injection request for tracking, or an invalid request if injection failed. This can be discarded for fire-and-forget style behavior.
+		ANIMNEXTANIMGRAPH_API static FInjectionRequestPtr PlayAnim(
+			UAnimNextComponent* InComponent,
+			FName InSiteName,
+			UAnimSequence* InAnimSequence,
+			float InPlayRate = 1.0f,
+			float InStartPosition = 0.0f,
+			const FInjectionBlendSettings& InBlendInSettings = FInjectionBlendSettings(),
+			const FInjectionBlendSettings& InBlendOutSettings = FInjectionBlendSettings(),
+			FInjectionLifetimeEvents&& InLifetimeEvents = FInjectionLifetimeEvents());
+
+		// @see PlayAnim
+		ANIMNEXTANIMGRAPH_API static FInjectionRequestPtr PlayAnim(
+			UObject* InHost,
+			FModuleHandle InModuleHandle,
+			FName InSiteName,
+			UAnimSequence* InAnimSequence,
+			float InPlayRate = 1.0f,
+			float InStartPosition = 0.0f,
+			const FInjectionBlendSettings& InBlendInSettings = FInjectionBlendSettings(),
+			const FInjectionBlendSettings& InBlendOutSettings = FInjectionBlendSettings(),
+			FInjectionLifetimeEvents&& InLifetimeEvents = FInjectionLifetimeEvents());
+
+		// Injects an asset into a running AnimNext module
+		// @param InComponent          The UAnimNextComponent that hosts the module being played on
+		// @param InHost               The object that hosts the module (e.g. UAnimNextComponent) being played on
+		// @param InHandle             The module handle being injected into
+		// @param InSiteName           The injection site to 'play' the asset on
+		// @param InAsset              The asset to 'play'
+		// @param InPayload            The payload for the asset to use
+		// @param InBindingComponent   A component to use to supply additional data interface bindings at the injection site
+		// @param InBindingModuleHandle A module to use to supply additional data interface bindings at the injection site
+		// @param InBlendInSettings    The blend in settings
+		// @param InBlendOutSettings   The blend in settings
+		// @param InLifetimeEvents     Delegates called for various request lifetime events
+		// @return the new injection request for tracking, or an invalid request if injection failed. This can be discarded for fire-and-forget style behavior.
+		ANIMNEXTANIMGRAPH_API static FInjectionRequestPtr InjectAsset(
+			UAnimNextComponent* InComponent,
+			FName InSiteName,
+			UObject* InAsset,
+			FAnimNextDataInterfacePayload&& InPayload = FAnimNextDataInterfacePayload(),
+			UAnimNextComponent* InBindingComponent = nullptr,
+			const FInjectionBlendSettings& InBlendInSettings = FInjectionBlendSettings(),
+			const FInjectionBlendSettings& InBlendOutSettings = FInjectionBlendSettings(),
+			FInjectionLifetimeEvents&& InLifetimeEvents = FInjectionLifetimeEvents());
+
+		// @see InjectAsset
+		ANIMNEXTANIMGRAPH_API static FInjectionRequestPtr InjectAsset(
+			UObject* InHost,
+			FModuleHandle InModuleHandle,
+			FName InSiteName,
+			UObject* InAsset,
+			FAnimNextDataInterfacePayload&& InPayload = FAnimNextDataInterfacePayload(),
+			FModuleHandle InBindingModuleHandle = FModuleHandle(),
+			const FInjectionBlendSettings& InBlendInSettings = FInjectionBlendSettings(),
+			const FInjectionBlendSettings& InBlendOutSettings = FInjectionBlendSettings(),
+			FInjectionLifetimeEvents&& InLifetimeEvents = FInjectionLifetimeEvents());
+
+		// Inject an evaluation modifier into a running AnimNext module to control animation evaluation at a low level.
+		// Does not require that injection sites are publicly visible.
+		// @param InComponent          The UAnimNextComponent that hosts the module being injected into
+		// @param InHost               The object that hosts the module (e.g. UAnimNextComponent) being injected into
+		// @param InModuleHandle       The module handle being injected into
+		// @param InEvaluationModifier The evaluation modifier to inject
+		// @param InSiteName           The injection site to use. If this is NAME_None, then the first site will be used.
+		// @return the new injection request for tracking, or an invalid request if injection failed. This can be discarded for fire-and-forget style behavior.
+		ANIMNEXTANIMGRAPH_API static FInjectionRequestPtr InjectEvaluationModifier(
+			UAnimNextComponent* InComponent,
+			const TSharedRef<IEvaluationModifier>& InEvaluationModifier,
+			FName InSiteName = NAME_None);
+
+		// @see InjectExternalGraph
+		ANIMNEXTANIMGRAPH_API static FInjectionRequestPtr InjectEvaluationModifier(
+			UObject* InHost,
+			FModuleHandle InModuleHandle,
+			const TSharedRef<IEvaluationModifier>& InEvaluationModifier,
+			FName InSiteName = NAME_None);
+
+		// Inject into a running AnimNext module (raw custom args)
+		// @param InComponent          The UAnimNextComponent that hosts the module being injected into
+		// @param InHost               The object that hosts the module (e.g. UAnimNextComponent) being injected into
+		// @param InHandle             The module handle being injected into
+		// @param InArgs               Arguments to the injection request
+		// @param InLifetimeEvents     Delegates called for various request lifetime events
+		// @return the new injection request for tracking, or an invalid request if injection failed. This can be discarded for fire-and-forget style behavior.
+		ANIMNEXTANIMGRAPH_API static FInjectionRequestPtr Inject(
+			UAnimNextComponent* InComponent,
+			FInjectionRequestArgs&& InArgs,
+			FInjectionLifetimeEvents&& InLifetimeEvents = FInjectionLifetimeEvents());
+
+		// @see Inject
+		ANIMNEXTANIMGRAPH_API static FInjectionRequestPtr Inject(
+			UObject* InHost,
+			FModuleHandle InHandle,
+			FInjectionRequestArgs&& InArgs,
+			FInjectionLifetimeEvents&& InLifetimeEvents = FInjectionLifetimeEvents());
+
+		// Uninject a previously injected request
+		// @param InInjectionRequest   The injection request to uninject
+		ANIMNEXTANIMGRAPH_API static void Uninject(FInjectionRequestPtr InInjectionRequest);
+	};
+}
diff --git a/Engine/Plugins/Experimental/AnimNextAnimGraph/Source/AnimNextAnimGraph/Public/TraitCore/ExecutionContext.h b/Engine/Plugins/Experimental/AnimNextAnimGraph/Source/AnimNextAnimGraph/Public/TraitCore/ExecutionContext.h
--- a/Engine/Plugins/Experimental/AnimNextAnimGraph/Source/AnimNextAnimGraph/Public/TraitCore/ExecutionContext.h
+++ b/Engine/Plugins/Experimental/AnimNextAnimGraph/Source/AnimNextAnimGraph/Public/TraitCore/ExecutionContext.h
@@ -14,7 +14,7 @@
 #include "TraitCore/TraitInterfaceUID.h"
 #include "TraitCore/LatentPropertyHandle.h"
 #include "TraitCore/NodeHandle.h"
-#include "Graph/AnimNextGraphInstancePtr.h"
+#include "Graph/AnimNextGraphInstance.h"
 #include "Graph/GraphInstanceComponent.h"
 #include "EngineDefines.h"
 
@@ -50,9 +50,6 @@ namespace UE::AnimNext
 		// Creates an unbound execution context
 		FExecutionContext();
 
-		// Creates an execution context and binds it to the specified weak graph instance
-		explicit FExecutionContext(FAnimNextGraphInstancePtr& InWeakGraphInstance);
-
 		// Creates an execution context and binds it to the specified graph instance
 		explicit FExecutionContext(FAnimNextGraphInstance& InGraphInstance);
 
@@ -72,9 +69,6 @@ namespace UE::AnimNext
 		// The following functions handle the context binding
 		// In order to be used, the execution context must be bound to a valid root graph instance
 
-		// Binds the execution context to the specified graph instance if it differs from the currently bound instance
-		void BindTo(const FAnimNextGraphInstancePtr& InWeakGraphInstance);
-
 		// Binds the execution context to the specified graph instance if it differs from the currently bound instance
 		void BindTo(FAnimNextGraphInstance& InGraphInstance);
 
@@ -84,10 +78,6 @@ namespace UE::AnimNext
 		// Returns whether or not this execution context is bound to a graph instance
 		bool IsBound() const;
 
-		// Returns whether or not this execution context is bound to the specified graph instance
-		// Returns true if the root graphs match
-		bool IsBoundTo(const FAnimNextGraphInstancePtr& InWeakGraphInstance) const;
-
 		// Returns whether or not this execution context is bound to the specified graph instance
 		// Returns true if the root graphs match
 		bool IsBoundTo(const FAnimNextGraphInstance& InGraphInstance) const;
diff --git a/Engine/Plugins/Experimental/AnimNextAnimGraph/Source/AnimNextAnimGraph/Public/TraitInterfaces/IBlendStack.h b/Engine/Plugins/Experimental/AnimNextAnimGraph/Source/AnimNextAnimGraph/Public/TraitInterfaces/IBlendStack.h
--- a/Engine/Plugins/Experimental/AnimNextAnimGraph/Source/AnimNextAnimGraph/Public/TraitInterfaces/IBlendStack.h
+++ b/Engine/Plugins/Experimental/AnimNextAnimGraph/Source/AnimNextAnimGraph/Public/TraitInterfaces/IBlendStack.h
@@ -9,7 +9,7 @@
 #include "AnimNextDataInterfacePayload.h"
 #include "IPayloadFactory.h"
 #include "Graph/AnimNextAnimationGraph.h"
-#include "Graph/AnimNextGraphInstancePtr.h"
+#include "Module/ModuleHandle.h"
 
 namespace UE::AnimNext
 {
@@ -22,6 +22,24 @@ namespace UE::AnimNext
 	{
 		DECLARE_ANIM_TRAIT_INTERFACE(IBlendStack, 0xf3468e64)
 
+		enum class EBlendMode : uint8
+		{
+			// Uses standard weight based blend
+			Standard,
+
+			// Uses inertialization. Requires an inertialization trait somewhere earlier in the graph.
+			Inertialization,
+		};
+
+		enum class EGraphRequestType : uint8
+		{
+			// Graph is allocated and owned by the blend stack
+			Owned,
+
+			// Graph is a child handle in the blend stacks graph instance
+			Child,
+		};
+
 		struct FGraphRequest
 		{
 			/**
@@ -41,15 +59,32 @@ namespace UE::AnimNext
 			// The object that was used for factory generation of the graph.
 			// Used as a key to allow identification of this request without inspecting payloads.
 			TObjectPtr<const UObject> FactoryObject;
+
+			// Child ptr when Type == Child
+			FTraitPtr ChildPtr;
+
+			// Additional module handle to use for data interface bindings
+			FModuleHandle BindingModuleHandle;
+
+			// The Blend mode to use
+			EBlendMode BlendMode = EBlendMode::Standard;
+
+			// What type of request this is
+			EGraphRequestType Type = EGraphRequestType::Owned;
 		};
 
 		typedef FGraphRequest* FGraphRequestPtr;
 
-		// Pushes a new subgraph along with blend settings defined in GraphRequest. Outputs the in-place created subgraph OutGraphInstance
-		virtual void PushGraph(FExecutionContext& Context, const TTraitBinding<IBlendStack>& Binding, IBlendStack::FGraphRequest&& GraphRequest) const;
-		
+		// Pushes a new subgraph along with blend settings defined in GraphRequest. Outputs the in-place created subgraph OutGraphInstance.
+		// @return the child index of the new graph, or INDEX_NONE if we could not push a graph
+		virtual int32 PushGraph(FExecutionContext& Context, const TTraitBinding<IBlendStack>& Binding, IBlendStack::FGraphRequest&& GraphRequest) const;
+
 		// Gets the graph request info from the most recent PushGraph
-		virtual void GetActiveGraph(FExecutionContext& Context, const TTraitBinding<IBlendStack>& Binding, IBlendStack::FGraphRequestPtr& OutGraphRequest) const;
+		// @return the child index of the active graph, or INDEX_NONE if no active graph was found
+		virtual int32 GetActiveGraph(FExecutionContext& Context, const TTraitBinding<IBlendStack>& Binding, IBlendStack::FGraphRequestPtr& OutGraphRequest) const;
+
+		// Gets the graph request for the specified child index
+		virtual IBlendStack::FGraphRequestPtr GetGraph(FExecutionContext& Context, const TTraitBinding<IBlendStack>& Binding, int32 InChildIndex) const;
 
 #if WITH_EDITOR
 		virtual const FText& GetDisplayName() const override;
@@ -64,17 +99,23 @@ namespace UE::AnimNext
 	struct TTraitBinding<IBlendStack> : FTraitBinding
 	{
 		// @see IBlendStack::PushGraph
-		void PushGraph(FExecutionContext& Context, IBlendStack::FGraphRequest&& GraphRequest) const
+		int32 PushGraph(FExecutionContext& Context, IBlendStack::FGraphRequest&& GraphRequest) const
 		{
-			GetInterface()->PushGraph(Context, *this, MoveTemp(GraphRequest));
+			return GetInterface()->PushGraph(Context, *this, MoveTemp(GraphRequest));
 		}
 
 		// @see IBlendStack::GetActiveGraph
-		void GetActiveGraph(FExecutionContext& Context, IBlendStack::FGraphRequestPtr& OutGraphRequest) const
+		int32 GetActiveGraph(FExecutionContext& Context, IBlendStack::FGraphRequestPtr& OutGraphRequest) const
 		{
-			GetInterface()->GetActiveGraph(Context, *this, OutGraphRequest);
+			return GetInterface()->GetActiveGraph(Context, *this, OutGraphRequest);
 		}
-		
+
+		// @see IBlendStack::GetGraph
+		IBlendStack::FGraphRequestPtr GetGraph(FExecutionContext& Context, int32 InChildIndex) const
+		{
+			return GetInterface()->GetGraph(Context, *this, InChildIndex);
+		}
+
 	protected:
 		const IBlendStack* GetInterface() const { return GetInterfaceTyped<IBlendStack>(); }
 	};
diff --git a/Engine/Plugins/Experimental/AnimNextAnimGraph/Source/AnimNextAnimGraph/Public/TraitInterfaces/IUpdate.h b/Engine/Plugins/Experimental/AnimNextAnimGraph/Source/AnimNextAnimGraph/Public/TraitInterfaces/IUpdate.h
--- a/Engine/Plugins/Experimental/AnimNextAnimGraph/Source/AnimNextAnimGraph/Public/TraitInterfaces/IUpdate.h
+++ b/Engine/Plugins/Experimental/AnimNextAnimGraph/Source/AnimNextAnimGraph/Public/TraitInterfaces/IUpdate.h
@@ -10,7 +10,6 @@
 
 #include <type_traits>
 
-struct FAnimNextGraphInstancePtr;
 class FMemStack;
 
 namespace UE::AnimNext
diff --git a/Engine/Plugins/Experimental/AnimNextAnimGraph/Source/AnimNextAnimGraphEditor/Private/AnimNextAnimGraphEditorModule.cpp b/Engine/Plugins/Experimental/AnimNextAnimGraph/Source/AnimNextAnimGraphEditor/Private/AnimNextAnimGraphEditorModule.cpp
--- a/Engine/Plugins/Experimental/AnimNextAnimGraph/Source/AnimNextAnimGraphEditor/Private/AnimNextAnimGraphEditorModule.cpp
+++ b/Engine/Plugins/Experimental/AnimNextAnimGraph/Source/AnimNextAnimGraphEditor/Private/AnimNextAnimGraphEditorModule.cpp
@@ -33,6 +33,8 @@
 #include "AnimNextTraitStackUnitNode.h"
 #include "RigVMCore/RigVMFunction.h"
 #include "EditorUtils.h"
+#include "Graph/AnimNextAnimGraph.h"
+#include "Graph/AnimNextGraphDetails.h"
 #include "Module/AnimNextEventGraphSchema.h"
 
 #define LOCTEXT_NAMESPACE "FAnimNextAnimGraphEditorModule"
@@ -57,6 +59,9 @@ class FAnimNextAnimGraphEditorModule : public IModuleInterface
 		PropertyModule.RegisterCustomPropertyTypeLayout(
 			FAnimNextCallFunctionSharedData::StaticStruct()->GetFName(),
 			FOnGetPropertyTypeCustomizationInstance::CreateLambda([] { return MakeShared<FCallFunctionSharedDataDetails>(); }));
+		PropertyModule.RegisterCustomPropertyTypeLayout(
+					FAnimNextAnimGraph::StaticStruct()->GetFName(),
+					FOnGetPropertyTypeCustomizationInstance::CreateLambda([] { return MakeShared<FAnimNextGraphDetails>(); }));
 
 		Workspace::IWorkspaceEditorModule& WorkspaceEditorModule = FModuleManager::Get().LoadModuleChecked<Workspace::IWorkspaceEditorModule>("WorkspaceEditor");
 
@@ -125,6 +130,7 @@ class FAnimNextAnimGraphEditorModule : public IModuleInterface
 		{
 			FPropertyEditorModule& PropertyModule = FModuleManager::LoadModuleChecked<FPropertyEditorModule>("PropertyEditor");
 			PropertyModule.UnregisterCustomClassLayout("AnimNextCallFunctionSharedData");
+			PropertyModule.UnregisterCustomClassLayout("AnimNextAnimGraph");
 		}
 
 		IModularFeatures::Get().UnregisterModularFeature(ITraitStackEditor::ModularFeatureName, TraitStackEditor.Get());
diff --git a/Engine/Plugins/Experimental/AnimNextAnimGraph/Source/AnimNextAnimGraphEditor/Private/Graph/AnimNextAnimationGraphFactory.cpp b/Engine/Plugins/Experimental/AnimNextAnimGraph/Source/AnimNextAnimGraphEditor/Private/Graph/AnimNextAnimationGraphFactory.cpp
--- a/Engine/Plugins/Experimental/AnimNextAnimGraph/Source/AnimNextAnimGraphEditor/Private/Graph/AnimNextAnimationGraphFactory.cpp
+++ b/Engine/Plugins/Experimental/AnimNextAnimGraph/Source/AnimNextAnimGraphEditor/Private/Graph/AnimNextAnimationGraphFactory.cpp
@@ -5,6 +5,7 @@
 #include "Graph/AnimNextAnimationGraph_EditorData.h"
 #include "UncookedOnlyUtils.h"
 #include "Entries/AnimNextAnimationGraphEntry.h"
+#include "Graph/RigUnit_AnimNextGraphRoot.h"
 
 UAnimNextAnimationGraphFactory::UAnimNextAnimationGraphFactory()
 {
diff --git a/Engine/Plugins/Experimental/AnimNextAnimGraph/Source/AnimNextAnimGraphEditor/Private/Graph/AnimNextAnimationGraphMenuExtensions.cpp b/Engine/Plugins/Experimental/AnimNextAnimGraph/Source/AnimNextAnimGraphEditor/Private/Graph/AnimNextAnimationGraphMenuExtensions.cpp
--- a/Engine/Plugins/Experimental/AnimNextAnimGraph/Source/AnimNextAnimGraphEditor/Private/Graph/AnimNextAnimationGraphMenuExtensions.cpp
+++ b/Engine/Plugins/Experimental/AnimNextAnimGraph/Source/AnimNextAnimGraphEditor/Private/Graph/AnimNextAnimationGraphMenuExtensions.cpp
@@ -11,7 +11,8 @@
 #include "EdGraph/EdGraphNode.h"
 #include "Framework/Application/SlateApplication.h"
 #include "Graph/RigDecorator_AnimNextCppTrait.h"
-#include "Graph/RigUnit_AnimNextRunAnimationGraph.h"
+#include "Graph/RigUnit_AnimNextRunAnimationGraph_v1.h"
+#include "Graph/RigUnit_AnimNextRunAnimationGraph_v2.h"
 #include "Graph/RigUnit_AnimNextTraitStack.h"
 #include "Graph/RigVMTrait_AnimNextPublicVariables.h"
 #include "RigVMModel/RigVMController.h"
@@ -66,7 +67,7 @@ void FAnimationGraphMenuExtensions::RegisterMenus()
 			if (const URigVMUnitNode* VMNode = Cast<URigVMUnitNode>(InModelNode))
 			{
 				const UScriptStruct* ScriptStruct = VMNode->GetScriptStruct();
-				return ScriptStruct == FRigUnit_AnimNextRunAnimationGraph::StaticStruct();
+				return ScriptStruct == FRigUnit_AnimNextRunAnimationGraph_v1::StaticStruct() || ScriptStruct == FRigUnit_AnimNextRunAnimationGraph_v2::StaticStruct();
 			}
 
 			return false;
diff --git a/Engine/Plugins/Experimental/AnimNextAnimGraph/Source/AnimNextAnimGraphTestSuite/Private/AnimNextAnimGraphEditorVariablesTest.cpp b/Engine/Plugins/Experimental/AnimNextAnimGraph/Source/AnimNextAnimGraphTestSuite/Private/AnimNextAnimGraphEditorVariablesTest.cpp
--- a/Engine/Plugins/Experimental/AnimNextAnimGraph/Source/AnimNextAnimGraphTestSuite/Private/AnimNextAnimGraphEditorVariablesTest.cpp
+++ b/Engine/Plugins/Experimental/AnimNextAnimGraph/Source/AnimNextAnimGraphTestSuite/Private/AnimNextAnimGraphEditorVariablesTest.cpp
@@ -12,6 +12,7 @@
 #include "Entries/AnimNextAnimationGraphEntry.h"
 #include "Entries/AnimNextEventGraphEntry.h"
 #include "Graph/AnimNextAnimationGraphFactory.h"
+#include "Graph/RigUnit_AnimNextGraphRoot.h"
 #include "Module/RigUnit_AnimNextModuleEvents.h"
 #if WITH_EDITOR
 #include "ScopedTransaction.h"
diff --git a/Engine/Plugins/Experimental/AnimNextAnimGraph/Source/AnimNextAnimGraphTestSuite/Private/AnimNextRuntimeTest.cpp b/Engine/Plugins/Experimental/AnimNextAnimGraph/Source/AnimNextAnimGraphTestSuite/Private/AnimNextRuntimeTest.cpp
--- a/Engine/Plugins/Experimental/AnimNextAnimGraph/Source/AnimNextAnimGraphTestSuite/Private/AnimNextRuntimeTest.cpp
+++ b/Engine/Plugins/Experimental/AnimNextAnimGraph/Source/AnimNextAnimGraphTestSuite/Private/AnimNextRuntimeTest.cpp
@@ -5,6 +5,7 @@
 
 #include "TraitCore/TraitReader.h"
 #include "Graph/AnimNextAnimationGraph.h"
+#include "Graph/RigUnit_AnimNextGraphEvaluator.h"
 #include "Graph/RigUnit_AnimNextGraphRoot.h"
 #include "Misc/AutomationTest.h"
 #include "Serialization/MemoryReader.h"
diff --git a/Engine/Plugins/Experimental/AnimNextAnimGraph/Source/AnimNextAnimGraphUncookedOnly/Internal/PlayAnim/K2Node_AnimNextPlayAsset.h b/Engine/Plugins/Experimental/AnimNextAnimGraph/Source/AnimNextAnimGraphUncookedOnly/Internal/PlayAnim/K2Node_AnimNextPlayAsset.h
--- a/Engine/Plugins/Experimental/AnimNextAnimGraph/Source/AnimNextAnimGraphUncookedOnly/Internal/PlayAnim/K2Node_AnimNextPlayAsset.h
+++ b/Engine/Plugins/Experimental/AnimNextAnimGraph/Source/AnimNextAnimGraphUncookedOnly/Internal/PlayAnim/K2Node_AnimNextPlayAsset.h
@@ -8,12 +8,12 @@
 
 #include "K2Node_AnimNextPlayAsset.generated.h"
 
-UCLASS()
-class UK2Node_AnimNextPlayAsset : public UK2Node_BaseAsyncTask
+UCLASS(deprecated)
+class UDEPRECATED_K2Node_AnimNextPlayAsset : public UK2Node_BaseAsyncTask
 {
 	GENERATED_BODY()
 
-	UK2Node_AnimNextPlayAsset();
+	UDEPRECATED_K2Node_AnimNextPlayAsset();
 
 	// UEdGraphNode Interface
 	virtual FText GetTooltipText() const override;
diff --git a/Engine/Plugins/Experimental/AnimNextAnimGraph/Source/AnimNextAnimGraphUncookedOnly/Private/Graph/AnimNextAnimationGraph_EditorData.cpp b/Engine/Plugins/Experimental/AnimNextAnimGraph/Source/AnimNextAnimGraphUncookedOnly/Private/Graph/AnimNextAnimationGraph_EditorData.cpp
--- a/Engine/Plugins/Experimental/AnimNextAnimGraph/Source/AnimNextAnimGraphUncookedOnly/Private/Graph/AnimNextAnimationGraph_EditorData.cpp
+++ b/Engine/Plugins/Experimental/AnimNextAnimGraph/Source/AnimNextAnimGraphUncookedOnly/Private/Graph/AnimNextAnimationGraph_EditorData.cpp
@@ -27,6 +27,8 @@
 #include "UObject/LinkerLoad.h"
 #include "Traits/CallFunction.h"
 #include "Entries/AnimNextRigVMAssetEntry.h"
+#include "Graph/RigUnit_AnimNextGraphEvaluator.h"
+#include "Graph/RigUnit_AnimNextGraphRoot.h"
 
 namespace UE::AnimNext::UncookedOnly::Private
 {
diff --git a/Engine/Plugins/Experimental/AnimNextAnimGraph/Source/AnimNextAnimGraphUncookedOnly/Private/PlayAnim/K2Node_AnimNextPlayAsset.cpp b/Engine/Plugins/Experimental/AnimNextAnimGraph/Source/AnimNextAnimGraphUncookedOnly/Private/PlayAnim/K2Node_AnimNextPlayAsset.cpp
--- a/Engine/Plugins/Experimental/AnimNextAnimGraph/Source/AnimNextAnimGraphUncookedOnly/Private/PlayAnim/K2Node_AnimNextPlayAsset.cpp
+++ b/Engine/Plugins/Experimental/AnimNextAnimGraph/Source/AnimNextAnimGraphUncookedOnly/Private/PlayAnim/K2Node_AnimNextPlayAsset.cpp
@@ -6,24 +6,24 @@
 
 #define LOCTEXT_NAMESPACE "K2Node_AnimNextPlayAsset"
 
-UK2Node_AnimNextPlayAsset::UK2Node_AnimNextPlayAsset()
+UDEPRECATED_K2Node_AnimNextPlayAsset::UDEPRECATED_K2Node_AnimNextPlayAsset()
 {
-	ProxyFactoryFunctionName = GET_FUNCTION_NAME_CHECKED(UPlayAnimCallbackProxy, CreateProxyObjectForPlayAsset);
+	ProxyFactoryFunctionName = NAME_None;
 	ProxyFactoryClass = UPlayAnimCallbackProxy::StaticClass();
 	ProxyClass = UPlayAnimCallbackProxy::StaticClass();
 }
 
-FText UK2Node_AnimNextPlayAsset::GetTooltipText() const
+FText UDEPRECATED_K2Node_AnimNextPlayAsset::GetTooltipText() const
 {
 	return LOCTEXT("K2Node_PlayAsset_Tooltip", "Plays an asset on an AnimNextComponent");
 }
 
-FText UK2Node_AnimNextPlayAsset::GetNodeTitle(ENodeTitleType::Type TitleType) const
+FText UDEPRECATED_K2Node_AnimNextPlayAsset::GetNodeTitle(ENodeTitleType::Type TitleType) const
 {
 	return LOCTEXT("PlayAsset", "Play Asset");
 }
 
-FText UK2Node_AnimNextPlayAsset::GetMenuCategory() const
+FText UDEPRECATED_K2Node_AnimNextPlayAsset::GetMenuCategory() const
 {
 	return LOCTEXT("PlayAssetCategory", "Animation|AnimNext");
 }
diff --git a/Engine/Plugins/Experimental/AnimNextStateTree/Source/AnimNextStateTree/Private/AnimStateTreeTrait.cpp b/Engine/Plugins/Experimental/AnimNextStateTree/Source/AnimNextStateTree/Private/AnimStateTreeTrait.cpp
--- a/Engine/Plugins/Experimental/AnimNextStateTree/Source/AnimNextStateTree/Private/AnimStateTreeTrait.cpp
+++ b/Engine/Plugins/Experimental/AnimNextStateTree/Source/AnimNextStateTree/Private/AnimStateTreeTrait.cpp
@@ -52,7 +52,7 @@ namespace UE::AnimNext
 			if(OwnerGraphInstance && InstanceData->StateTree->GetDefaultParameters().Identical(&OwnerGraphInstance->GetVariables(), 0))
 			{
 				const FInstancedPropertyBag& StateTreeParameters = InstanceData->StateTree->GetDefaultParameters();
-				FStructView MutableVariables = OwnerGraphInstance->GetVariables().GetMutableValue();
+				FStructView MutableVariables = OwnerGraphInstance->GetMutableVariables().GetMutableValue();
 				
 				const FRigVMExtendedExecuteContext& ExtendedExecuteContext = OwnerGraphInstance->GetExtendedExecuteContext();
 
diff --git a/Engine/Plugins/Experimental/AnimNextStateTree/Source/AnimNextStateTree/Private/Conditions/AnimNextStateTreeRigVMConditionBase.cpp b/Engine/Plugins/Experimental/AnimNextStateTree/Source/AnimNextStateTree/Private/Conditions/AnimNextStateTreeRigVMConditionBase.cpp
--- a/Engine/Plugins/Experimental/AnimNextStateTree/Source/AnimNextStateTree/Private/Conditions/AnimNextStateTreeRigVMConditionBase.cpp
+++ b/Engine/Plugins/Experimental/AnimNextStateTree/Source/AnimNextStateTree/Private/Conditions/AnimNextStateTreeRigVMConditionBase.cpp
@@ -187,7 +187,7 @@ bool FAnimNextStateTreeRigVMConditionBase::TestCondition(FStateTreeExecutionCont
 				void* SourceAddress = StateTreePropertyDesc.CachedProperty->ContainerPtrToValuePtr<void*>(StateTreeParamData.GetMemory());
 
 				const FPropertyBagPropertyDesc& AnimNextGraphDesc = AnimNextGraphPropertyDescs[RigVMExternalVariableIndex];
-				GraphInstance.GetVariables().SetValue(AnimNextGraphDesc.Name, StateTreePropertyDesc.CachedProperty, StateTreeParamData.GetMemory());
+				GraphInstance.GetMutableVariables().SetValue(AnimNextGraphDesc.Name, StateTreePropertyDesc.CachedProperty, StateTreeParamData.GetMemory());
 			}
 
 			FRigVMExtendedExecuteContext& ExtendedExecuteContext = GraphInstance.GetExtendedExecuteContext();
diff --git a/Engine/Plugins/Experimental/AnimNextStateTree/Source/AnimNextStateTree/Public/AnimNextStateTreeContext.h b/Engine/Plugins/Experimental/AnimNextStateTree/Source/AnimNextStateTree/Public/AnimNextStateTreeContext.h
--- a/Engine/Plugins/Experimental/AnimNextStateTree/Source/AnimNextStateTree/Public/AnimNextStateTreeContext.h
+++ b/Engine/Plugins/Experimental/AnimNextStateTree/Source/AnimNextStateTree/Public/AnimNextStateTreeContext.h
@@ -2,7 +2,6 @@
 
 #pragma once
 
-#include "Graph/AnimNextGraphInstancePtr.h"
 #include "StructUtils/InstancedStruct.h"
 
 #include "AnimNextStateTreeContext.generated.h"
