AnimNext: Fix crash when running graphs with multiple data interfaces

Correctly index cached bindings with the local index for the data interface, not the overall variable index for the host.

diff --git a/Engine/Plugins/Experimental/AnimNextAnimGraph/Source/AnimNextAnimGraph/Private/Graph/AnimNextGraphInstance.cpp b/Engine/Plugins/Experimental/AnimNextAnimGraph/Source/AnimNextAnimGraph/Private/Graph/AnimNextGraphInstance.cpp
--- a/Engine/Plugins/Experimental/AnimNextAnimGraph/Source/AnimNextAnimGraph/Private/Graph/AnimNextGraphInstance.cpp
+++ b/Engine/Plugins/Experimental/AnimNextAnimGraph/Source/AnimNextAnimGraph/Private/Graph/AnimNextGraphInstance.cpp
@@ -364,13 +364,18 @@ void FAnimNextGraphInstance::UpdateCachedBindingsForHostHelper(const HostType& I
 	}
 
 	const UPropertyBag* PropertyBag = HostDataInterface->GetPublicVariableDefaults().GetPropertyBagStruct();
+	if(PropertyBag == nullptr)
+	{
+		return;
+	}
+
 	for(const FAnimNextImplementedDataInterface& ImplementedInterface : HostDataInterface->GetImplementedInterfaces())
 	{
 		FCachedDataInterfaceBinding& CachedBinding = FindOrAddCachedInterfaceBindings(ImplementedInterface.DataInterface);
 
 		int32 VariableIndex = ImplementedInterface.VariableIndex;
 		const int32 EndVariableIndex = ImplementedInterface.VariableIndex + ImplementedInterface.NumVariables;
-		for(; VariableIndex < EndVariableIndex; ++VariableIndex)
+		for(int32 BindingIndex = 0; VariableIndex < EndVariableIndex && BindingIndex < ImplementedInterface.NumVariables; ++VariableIndex, ++BindingIndex)
 		{
 			const FPropertyBagPropertyDesc& Desc = PropertyBag->GetPropertyDescs()[VariableIndex];
 			uint8* HostMemory = InHost.GetMemoryForVariable(VariableIndex, Desc.Name, Desc.CachedProperty);
@@ -378,7 +383,7 @@ void FAnimNextGraphInstance::UpdateCachedBindingsForHostHelper(const HostType& I
 			{
 				continue;
 			}
-			CachedBinding.CachedBindings[VariableIndex] = FCachedDataInterfaceBinding::FVariable(Desc.Name, VariableIndex, Desc.CachedProperty, HostMemory);
+			CachedBinding.CachedBindings[BindingIndex] = FCachedDataInterfaceBinding::FVariable(Desc.Name, VariableIndex, Desc.CachedProperty, HostMemory);
 		}
 	}
 }
@@ -505,4 +510,3 @@ void FAnimNextGraphInstance::UnbindPublicVariables()
 
 	PublicVariablesState = EPublicVariablesState::Unbound;
 }
-
