AnimNext: Fix logged error caused by deprecated K2 node

diff --git a/Engine/Plugins/Experimental/AnimNextAnimGraph/Source/AnimNextAnimGraphUncookedOnly/Private/PlayAnim/K2Node_AnimNextPlayAsset.cpp b/Engine/Plugins/Experimental/AnimNextAnimGraph/Source/AnimNextAnimGraphUncookedOnly/Private/PlayAnim/K2Node_AnimNextPlayAsset.cpp
--- a/Engine/Plugins/Experimental/AnimNextAnimGraph/Source/AnimNextAnimGraphUncookedOnly/Private/PlayAnim/K2Node_AnimNextPlayAsset.cpp
+++ b/Engine/Plugins/Experimental/AnimNextAnimGraph/Source/AnimNextAnimGraphUncookedOnly/Private/PlayAnim/K2Node_AnimNextPlayAsset.cpp
@@ -8,7 +8,7 @@
 
 UDEPRECATED_K2Node_AnimNextPlayAsset::UDEPRECATED_K2Node_AnimNextPlayAsset()
 {
-	ProxyFactoryFunctionName = NAME_None;
+	ProxyFactoryFunctionName = GET_FUNCTION_NAME_CHECKED(UPlayAnimCallbackProxy, CreateProxyObjectForPlayAnim);
 	ProxyFactoryClass = UPlayAnimCallbackProxy::StaticClass();
 	ProxyClass = UPlayAnimCallbackProxy::StaticClass();
 }
