AnimNext: Fixed crash when live-recompiling animation graphs with cached bindings

diff --git a/Engine/Plugins/Experimental/AnimNextAnimGraph/Source/AnimNextAnimGraph/Private/Graph/AnimNextGraphInstance.cpp b/Engine/Plugins/Experimental/AnimNextAnimGraph/Source/AnimNextAnimGraph/Private/Graph/AnimNextGraphInstance.cpp
--- a/Engine/Plugins/Experimental/AnimNextAnimGraph/Source/AnimNextAnimGraph/Private/Graph/AnimNextGraphInstance.cpp
+++ b/Engine/Plugins/Experimental/AnimNextAnimGraph/Source/AnimNextAnimGraph/Private/Graph/AnimNextGraphInstance.cpp
@@ -197,6 +197,7 @@ void FAnimNextGraphInstance::Freeze()
 	GraphInstancePtr.Reset();
 	ExtendedExecuteContext.Reset();
 	Components.Empty();
+	CachedBindings.Reset();
 	PublicVariablesState = PublicVariablesState == EPublicVariablesState::Bound ? EPublicVariablesState::Unbound : EPublicVariablesState::None;
 	bHasUpdatedOnce = false;
 }
