AnimNext: Fix incorrect name being displayed for AnimNext Graph properties

diff --git a/Engine/Plugins/Experimental/AnimNextAnimGraph/Source/AnimNextAnimGraphEditor/Private/Graph/AnimNextGraphDetails.cpp b/Engine/Plugins/Experimental/AnimNextAnimGraph/Source/AnimNextAnimGraphEditor/Private/Graph/AnimNextGraphDetails.cpp
--- a/Engine/Plugins/Experimental/AnimNextAnimGraph/Source/AnimNextAnimGraphEditor/Private/Graph/AnimNextGraphDetails.cpp
+++ b/Engine/Plugins/Experimental/AnimNextAnimGraph/Source/AnimNextAnimGraphEditor/Private/Graph/AnimNextGraphDetails.cpp
@@ -18,7 +18,8 @@ void FAnimNextGraphDetails::CustomizeHeader(TSharedRef<IPropertyHandle> InProper
 void FAnimNextGraphDetails::CustomizeChildren(TSharedRef<IPropertyHandle> InPropertyHandle, IDetailChildrenBuilder& InChildBuilder, IPropertyTypeCustomizationUtils& InCustomizationUtils)
 {
 	TSharedRef<IPropertyHandle> AssetPropertyHandle = InPropertyHandle->GetChildHandle(GET_MEMBER_NAME_CHECKED(FAnimNextAnimGraph, Asset)).ToSharedRef();
-	InChildBuilder.AddProperty(AssetPropertyHandle);
+	InChildBuilder.AddProperty(AssetPropertyHandle)
+	.DisplayName(InPropertyHandle->GetPropertyDisplayName());
 
 	IDetailCategoryBuilder& CategoryBuilder = InChildBuilder.GetParentCategory();
 	TSharedRef<IPropertyHandle> HostGraphPropertyHandle = InPropertyHandle->GetChildHandle(GET_MEMBER_NAME_CHECKED(FAnimNextAnimGraph, HostGraph)).ToSharedRef();
