AnimNext: Fix crash when editing an existing trait stack to have no base trait

diff --git a/Engine/Plugins/Experimental/AnimNextAnimGraph/Source/AnimNextAnimGraphUncookedOnly/Private/Graph/AnimNextAnimationGraph_EditorData.cpp b/Engine/Plugins/Experimental/AnimNextAnimGraph/Source/AnimNextAnimGraphUncookedOnly/Private/Graph/AnimNextAnimationGraph_EditorData.cpp
--- a/Engine/Plugins/Experimental/AnimNextAnimGraph/Source/AnimNextAnimGraphUncookedOnly/Private/Graph/AnimNextAnimationGraph_EditorData.cpp
+++ b/Engine/Plugins/Experimental/AnimNextAnimGraph/Source/AnimNextAnimGraphUncookedOnly/Private/Graph/AnimNextAnimationGraph_EditorData.cpp
@@ -321,7 +321,7 @@ namespace UE::AnimNext::UncookedOnly::Private
 		}
 	}
 
-	FTraitGraph CollectGraphInfo(const UAnimNextAnimationGraph* InAnimationGraph, const URigVMGraph* VMGraph, URigVMController* VMController)
+	FTraitGraph CollectGraphInfo(FRigVMCompileSettings& InSettings, const UAnimNextAnimationGraph* InAnimationGraph, const URigVMGraph* VMGraph, URigVMController* VMController)
 	{
 		const TArray<URigVMNode*>& VMNodes = VMGraph->GetNodes();
 		URigVMUnitNode* VMRootNode = FindRootNode(VMNodes);
@@ -351,13 +351,24 @@ namespace UE::AnimNext::UncookedOnly::Private
 				if (ScriptStruct == FRigUnit_AnimNextTraitStack::StaticStruct())
 				{
 					FTraitStackMapping Mapping(VMNode);
+
+					bool bHasBaseTrait = false;
 					ForEachTraitInStack(VMNode,
-						[&Mapping](const URigVMNode* DecoratorStackNode, const URigVMPin* DecoratorPin, const FTrait* Trait)
+						[&Mapping, &bHasBaseTrait](const URigVMNode* DecoratorStackNode, const URigVMPin* DecoratorPin, const FTrait* Trait)
 						{
+							bHasBaseTrait |= Trait->GetTraitMode() == ETraitMode::Base;
 							Mapping.TraitEntries.Add(FTraitEntryMapping(DecoratorStackNode, DecoratorPin, Trait));
 						});
 
-					TraitGraph.TraitStackNodes.Add(MoveTemp(Mapping));
+					if(!bHasBaseTrait)
+					{
+						// Must have at least one base trait
+						InSettings.ASTSettings.Reportf(EMessageSeverity::Error, const_cast<URigVMUnitNode*>(VMUnitNode), TEXT("No base trait supplied for @@"));
+					}
+					else
+					{
+						TraitGraph.TraitStackNodes.Add(MoveTemp(Mapping));
+					}
 				}
 			}
 
@@ -567,7 +578,7 @@ void UAnimNextAnimationGraph_EditorData::RecompileVM()
 			for(const URigVMGraph* AnimGraph : AnimGraphs)
 			{
 				// Gather our trait stacks
-				FTraitGraph& TraitGraph = TraitGraphs.Add_GetRef(CollectGraphInfo(AnimationGraph, AnimGraph, TempController->GetControllerForGraph(AnimGraph)));
+				FTraitGraph& TraitGraph = TraitGraphs.Add_GetRef(CollectGraphInfo(Settings, AnimationGraph, AnimGraph, TempController->GetControllerForGraph(AnimGraph)));
 				check(!TraitGraph.TraitStackNodes.IsEmpty());
 
 				FAnimNextGraphEntryPoint& EntryPoint = AnimationGraph->EntryPoints.AddDefaulted_GetRef();
