AnimNext: Fix automated test failure due to newly-introduced error for invalid trait stacks

diff --git a/Engine/Plugins/Experimental/AnimNextAnimGraph/Source/AnimNextAnimGraphTestSuite/Private/AnimNextAnimGraphTraitGraphTest.cpp b/Engine/Plugins/Experimental/AnimNextAnimGraph/Source/AnimNextAnimGraphTestSuite/Private/AnimNextAnimGraphTraitGraphTest.cpp
--- a/Engine/Plugins/Experimental/AnimNextAnimGraph/Source/AnimNextAnimGraphTestSuite/Private/AnimNextAnimGraphTraitGraphTest.cpp
+++ b/Engine/Plugins/Experimental/AnimNextAnimGraph/Source/AnimNextAnimGraphTestSuite/Private/AnimNextAnimGraphTraitGraphTest.cpp
@@ -477,44 +477,51 @@ bool FAnimationAnimNextRuntimeTest_GraphExecute::RunTest(const FString& InParame
 		URigVMPin* BeginExecutePin = MainEntryPointNode->FindPin(GET_MEMBER_NAME_STRING_CHECKED(FRigUnit_AnimNextGraphRoot, Result));
 		UE_RETURN_ON_ERROR(BeginExecutePin != nullptr && BeginExecutePin->GetDirection() == ERigVMPinDirection::Input, "FAnimationAnimNextRuntimeTest_GraphExecute -> Failed to create entry point");
 
-		// Create an empty trait stack node
-		URigVMUnitNode* DecoratorStackNode = Controller->AddUnitNode(FRigUnit_AnimNextTraitStack::StaticStruct(), FRigVMStruct::ExecuteName, FVector2D(0.0f, 0.0f), FString(), false);
-		UE_RETURN_ON_ERROR(DecoratorStackNode != nullptr, "FAnimationAnimNextRuntimeTest_GraphExecute -> Failed to create trait stack node");
+		URigVMUnitNode* DecoratorStackNode = nullptr;
+		FString DisplayName;
+		{
+			// Suspend auto compilation until we have constructed a valid trait stack
+			TGuardValue<bool> SuspendCompile(EditorData->bAutoRecompileVM, false);
 
-		// Link our stack result to our entry point
-		Controller->AddLink(DecoratorStackNode->GetPins()[0], MainEntryPointNode->GetPins()[0]);
+			// Create an empty trait stack node
+			DecoratorStackNode = Controller->AddUnitNode(FRigUnit_AnimNextTraitStack::StaticStruct(), FRigVMStruct::ExecuteName, FVector2D(0.0f, 0.0f), FString(), false);
+			UE_RETURN_ON_ERROR(DecoratorStackNode != nullptr, "FAnimationAnimNextRuntimeTest_GraphExecute -> Failed to create trait stack node");
 
-		// Add a trait
-		const UScriptStruct* CppDecoratorStruct = FRigDecorator_AnimNextCppDecorator::StaticStruct();
-		UE_RETURN_ON_ERROR(CppDecoratorStruct != nullptr, "FAnimationAnimNextRuntimeTest_GraphExecute -> Failed to get find Cpp trait static struct");
+			// Link our stack result to our entry point
+			Controller->AddLink(DecoratorStackNode->GetPins()[0], MainEntryPointNode->GetPins()[0]);
 
-		const FTrait* Trait = FTraitRegistry::Get().Find(FTestTrait::TraitUID);
-		UE_RETURN_ON_ERROR(Trait != nullptr, "FAnimationAnimNextRuntimeTest_GraphExecute -> Failed to find test trait");
+			// Add a trait
+			const UScriptStruct* CppDecoratorStruct = FRigDecorator_AnimNextCppDecorator::StaticStruct();
+			UE_RETURN_ON_ERROR(CppDecoratorStruct != nullptr, "FAnimationAnimNextRuntimeTest_GraphExecute -> Failed to get find Cpp trait static struct");
 
-		UScriptStruct* ScriptStruct = Trait->GetTraitSharedDataStruct();
-		UE_RETURN_ON_ERROR(ScriptStruct != nullptr, "FAnimationAnimNextRuntimeTest_GraphExecute -> Failed to find trait shared data struct");
+			const FTrait* Trait = FTraitRegistry::Get().Find(FTestTrait::TraitUID);
+			UE_RETURN_ON_ERROR(Trait != nullptr, "FAnimationAnimNextRuntimeTest_GraphExecute -> Failed to find test trait");
 
-		FString DefaultValue;
-		{
-			const FRigDecorator_AnimNextCppDecorator DefaultCppDecoratorStructInstance;
-			FRigDecorator_AnimNextCppDecorator CppDecoratorStructInstance;
-			CppDecoratorStructInstance.DecoratorSharedDataStruct = ScriptStruct;
+			UScriptStruct* ScriptStruct = Trait->GetTraitSharedDataStruct();
+			UE_RETURN_ON_ERROR(ScriptStruct != nullptr, "FAnimationAnimNextRuntimeTest_GraphExecute -> Failed to find trait shared data struct");
 
-			UE_RETURN_ON_ERROR(CppDecoratorStructInstance.CanBeAddedToNode(DecoratorStackNode, nullptr), "FAnimationAnimNextRuntimeTest_GraphExecute -> Trait cannot be added to trait stack node");
+			FString DefaultValue;
+			{
+				const FRigDecorator_AnimNextCppDecorator DefaultCppDecoratorStructInstance;
+				FRigDecorator_AnimNextCppDecorator CppDecoratorStructInstance;
+				CppDecoratorStructInstance.DecoratorSharedDataStruct = ScriptStruct;
 
-			FRigDecorator_AnimNextCppDecorator::StaticStruct()->ExportText(DefaultValue, &CppDecoratorStructInstance, &DefaultCppDecoratorStructInstance, nullptr, PPF_None, nullptr);
-		}
+				UE_RETURN_ON_ERROR(CppDecoratorStructInstance.CanBeAddedToNode(DecoratorStackNode, nullptr), "FAnimationAnimNextRuntimeTest_GraphExecute -> Trait cannot be added to trait stack node");
 
-		FString DisplayNameMetadata;
-		ScriptStruct->GetStringMetaDataHierarchical(FRigVMStruct::DisplayNameMetaName, &DisplayNameMetadata);
-		const FString DisplayName = DisplayNameMetadata.IsEmpty() ? Trait->GetTraitName() : DisplayNameMetadata;
+				FRigDecorator_AnimNextCppDecorator::StaticStruct()->ExportText(DefaultValue, &CppDecoratorStructInstance, &DefaultCppDecoratorStructInstance, nullptr, PPF_None, nullptr);
+			}
 
-		const FName DecoratorName = Controller->AddTrait(
-			DecoratorStackNode->GetFName(),
-			*CppDecoratorStruct->GetPathName(),
-			*DisplayName,
-			DefaultValue, INDEX_NONE, true, true);
-		UE_RETURN_ON_ERROR(DecoratorName == DisplayName, TEXT("FAnimationAnimNextRuntimeTest_GraphExecute -> Unexpected trait name"));
+			FString DisplayNameMetadata;
+			ScriptStruct->GetStringMetaDataHierarchical(FRigVMStruct::DisplayNameMetaName, &DisplayNameMetadata);
+			DisplayName = DisplayNameMetadata.IsEmpty() ? Trait->GetTraitName() : DisplayNameMetadata;
+
+			const FName DecoratorName = Controller->AddTrait(
+				DecoratorStackNode->GetFName(),
+				*CppDecoratorStruct->GetPathName(),
+				*DisplayName,
+				DefaultValue, INDEX_NONE, true, true);
+			UE_RETURN_ON_ERROR(DecoratorName == DisplayName, TEXT("FAnimationAnimNextRuntimeTest_GraphExecute -> Unexpected trait name"));
+		}
 
 		URigVMPin* DecoratorPin = DecoratorStackNode->FindPin(*DisplayName);
 		UE_RETURN_ON_ERROR(DecoratorPin != nullptr, TEXT("FAnimationAnimNextRuntimeTest_GraphExecute -> Failed to find trait pin"));
@@ -584,44 +591,51 @@ bool FAnimationAnimNextRuntimeTest_GraphExecuteLatent::RunTest(const FString& In
 		URigVMPin* BeginExecutePin = MainEntryPointNode->FindPin(GET_MEMBER_NAME_STRING_CHECKED(FRigUnit_AnimNextGraphRoot, Result));
 		UE_RETURN_ON_ERROR(BeginExecutePin != nullptr && BeginExecutePin->GetDirection() == ERigVMPinDirection::Input, "FAnimationAnimNextRuntimeTest_GraphExecuteLatent -> Failed to create entry point");
 
-		// Create an empty trait stack node
-		URigVMUnitNode* DecoratorStackNode = Controller->AddUnitNode(FRigUnit_AnimNextTraitStack::StaticStruct(), FRigVMStruct::ExecuteName, FVector2D(0.0f, 0.0f), FString(), false);
-		UE_RETURN_ON_ERROR(DecoratorStackNode != nullptr, "FAnimationAnimNextRuntimeTest_GraphExecuteLatent -> Failed to create trait stack node");
+		URigVMUnitNode* DecoratorStackNode = nullptr;
+		FString DisplayName;
+		{
+			// Suspend auto compilation until we have constructed a valid trait stack
+			TGuardValue<bool> SuspendCompile(EditorData->bAutoRecompileVM, false);
 
-		// Link our stack result to our entry point
-		Controller->AddLink(DecoratorStackNode->GetPins()[0], MainEntryPointNode->GetPins()[0]);
+			// Create an empty trait stack node
+			DecoratorStackNode = Controller->AddUnitNode(FRigUnit_AnimNextTraitStack::StaticStruct(), FRigVMStruct::ExecuteName, FVector2D(0.0f, 0.0f), FString(), false);
+			UE_RETURN_ON_ERROR(DecoratorStackNode != nullptr, "FAnimationAnimNextRuntimeTest_GraphExecuteLatent -> Failed to create trait stack node");
 
-		// Add a trait
-		const UScriptStruct* CppDecoratorStruct = FRigDecorator_AnimNextCppDecorator::StaticStruct();
-		UE_RETURN_ON_ERROR(CppDecoratorStruct != nullptr, "FAnimationAnimNextRuntimeTest_GraphExecuteLatent -> Failed to get find Cpp trait static struct");
+			// Link our stack result to our entry point
+			Controller->AddLink(DecoratorStackNode->GetPins()[0], MainEntryPointNode->GetPins()[0]);
 
-		const FTrait* Trait = FTraitRegistry::Get().Find(FTestTrait::TraitUID);
-		UE_RETURN_ON_ERROR(Trait != nullptr, "FAnimationAnimNextRuntimeTest_GraphExecuteLatent -> Failed to find test trait");
+			// Add a trait
+			const UScriptStruct* CppDecoratorStruct = FRigDecorator_AnimNextCppDecorator::StaticStruct();
+			UE_RETURN_ON_ERROR(CppDecoratorStruct != nullptr, "FAnimationAnimNextRuntimeTest_GraphExecuteLatent -> Failed to get find Cpp trait static struct");
 
-		UScriptStruct* ScriptStruct = Trait->GetTraitSharedDataStruct();
-		UE_RETURN_ON_ERROR(ScriptStruct != nullptr, "FAnimationAnimNextRuntimeTest_GraphExecuteLatent -> Failed to find trait shared data struct");
+			const FTrait* Trait = FTraitRegistry::Get().Find(FTestTrait::TraitUID);
+			UE_RETURN_ON_ERROR(Trait != nullptr, "FAnimationAnimNextRuntimeTest_GraphExecuteLatent -> Failed to find test trait");
 
-		FString DefaultValue;
-		{
-			const FRigDecorator_AnimNextCppDecorator DefaultCppDecoratorStructInstance;
-			FRigDecorator_AnimNextCppDecorator CppDecoratorStructInstance;
-			CppDecoratorStructInstance.DecoratorSharedDataStruct = ScriptStruct;
+			UScriptStruct* ScriptStruct = Trait->GetTraitSharedDataStruct();
+			UE_RETURN_ON_ERROR(ScriptStruct != nullptr, "FAnimationAnimNextRuntimeTest_GraphExecuteLatent -> Failed to find trait shared data struct");
 
-			UE_RETURN_ON_ERROR(CppDecoratorStructInstance.CanBeAddedToNode(DecoratorStackNode, nullptr), "FAnimationAnimNextRuntimeTest_GraphExecuteLatent -> Trait cannot be added to trait stack node");
+			FString DefaultValue;
+			{
+				const FRigDecorator_AnimNextCppDecorator DefaultCppDecoratorStructInstance;
+				FRigDecorator_AnimNextCppDecorator CppDecoratorStructInstance;
+				CppDecoratorStructInstance.DecoratorSharedDataStruct = ScriptStruct;
 
-			FRigDecorator_AnimNextCppDecorator::StaticStruct()->ExportText(DefaultValue, &CppDecoratorStructInstance, &DefaultCppDecoratorStructInstance, nullptr, PPF_None, nullptr);
-		}
+				UE_RETURN_ON_ERROR(CppDecoratorStructInstance.CanBeAddedToNode(DecoratorStackNode, nullptr), "FAnimationAnimNextRuntimeTest_GraphExecuteLatent -> Trait cannot be added to trait stack node");
 
-		FString DisplayNameMetadata;
-		ScriptStruct->GetStringMetaDataHierarchical(FRigVMStruct::DisplayNameMetaName, &DisplayNameMetadata);
-		const FString DisplayName = DisplayNameMetadata.IsEmpty() ? Trait->GetTraitName() : DisplayNameMetadata;
+				FRigDecorator_AnimNextCppDecorator::StaticStruct()->ExportText(DefaultValue, &CppDecoratorStructInstance, &DefaultCppDecoratorStructInstance, nullptr, PPF_None, nullptr);
+			}
 
-		const FName DecoratorName = Controller->AddTrait(
-			DecoratorStackNode->GetFName(),
-			*CppDecoratorStruct->GetPathName(),
-			*DisplayName,
-			DefaultValue, INDEX_NONE, true, true);
-		UE_RETURN_ON_ERROR(DecoratorName == DisplayName, TEXT("FAnimationAnimNextRuntimeTest_GraphExecuteLatent -> Unexpected trait name"));
+			FString DisplayNameMetadata;
+			ScriptStruct->GetStringMetaDataHierarchical(FRigVMStruct::DisplayNameMetaName, &DisplayNameMetadata);
+			DisplayName = DisplayNameMetadata.IsEmpty() ? Trait->GetTraitName() : DisplayNameMetadata;
+
+			const FName DecoratorName = Controller->AddTrait(
+				DecoratorStackNode->GetFName(),
+				*CppDecoratorStruct->GetPathName(),
+				*DisplayName,
+				DefaultValue, INDEX_NONE, true, true);
+			UE_RETURN_ON_ERROR(DecoratorName == DisplayName, TEXT("FAnimationAnimNextRuntimeTest_GraphExecuteLatent -> Unexpected trait name"));
+		}
 
 		// Set some values on our trait
 		URigVMPin* DecoratorPin = DecoratorStackNode->FindPin(*DisplayName);
