AnimNext: Fix injection not reverting to pre-injected state when injecting at the module level

Fallback to the pin's value when transitioning away from an injected (via event) graph. If no pin-value is present, fall back to source as before.
Pin value's Asset member is no longer set via the injeciton component. Event injection uses the graph passed by the event and leaves the pin value as-is.
Split injection data out into its own struct so it can be modified/bypassed as required when dealing with graph variables (allows custom host graphs - not required but was missing from the original Injection changes)

diff --git a/Engine/Plugins/Experimental/AnimNextAnimGraph/Source/AnimNextAnimGraph/Private/Injection/AnimNextModuleInjectionComponent.cpp b/Engine/Plugins/Experimental/AnimNextAnimGraph/Source/AnimNextAnimGraph/Private/Injection/AnimNextModuleInjectionComponent.cpp
--- a/Engine/Plugins/Experimental/AnimNextAnimGraph/Source/AnimNextAnimGraph/Private/Injection/AnimNextModuleInjectionComponent.cpp
+++ b/Engine/Plugins/Experimental/AnimNextAnimGraph/Source/AnimNextAnimGraph/Private/Injection/AnimNextModuleInjectionComponent.cpp
@@ -72,12 +72,11 @@ void FAnimNextModuleInjectionComponent::OnInjectionEvent(UE::AnimNext::FInjectio
 	switch(RequestArgs.Type)
 	{
 	case EAnimNextInjectionType::InjectObject:
-		// Update graph and bump serial number to identify this injection routing
+		// Bump serial number to identify this injection routing
 		ensureAlways(RequestArgs.Object != nullptr);
-		Graph.Asset = RequestArgs.Object;
-		Graph.InjectionSerialNumber = IncrementSerialNumber();
-		InEvent.SerialNumber = Graph.InjectionSerialNumber;
-		InjectionRecord.SerialNumber = Graph.InjectionSerialNumber;
+		Graph.InjectionData.InjectionSerialNumber = IncrementSerialNumber();
+		InEvent.SerialNumber = Graph.InjectionData.InjectionSerialNumber;
+		InjectionRecord.SerialNumber = Graph.InjectionData.InjectionSerialNumber;
 		InjectionRecord.GraphRequest = InEvent.Request;
 
 		// Note we dont consume here, as we want the event to forward to the injection site trait
@@ -85,7 +84,7 @@ void FAnimNextModuleInjectionComponent::OnInjectionEvent(UE::AnimNext::FInjectio
 	case EAnimNextInjectionType::EvaluationModifier:
 		// We dont increment serial number when applying evaluation modifiers as we dont want to trigger a graph instantiation
 		ensureAlways(RequestArgs.EvaluationModifier != nullptr);
-		Graph.EvaluationModifier = RequestArgs.EvaluationModifier;
+		Graph.InjectionData.EvaluationModifier = RequestArgs.EvaluationModifier;
 		InjectionRecord.ModifierRequest = InEvent.Request;
 
 		// Update status
@@ -126,16 +125,15 @@ void FAnimNextModuleInjectionComponent::OnUninjectionEvent(UE::AnimNext::FInject
 	{
 	case EAnimNextInjectionType::InjectObject:
 		// Update graph and bump serial number to identify this injection routing
-		Graph.Asset = nullptr;
-		Graph.InjectionSerialNumber = IncrementSerialNumber();
-		InEvent.SerialNumber = Graph.InjectionSerialNumber;
+		Graph.InjectionData.InjectionSerialNumber = IncrementSerialNumber();
+		InEvent.SerialNumber = Graph.InjectionData.InjectionSerialNumber;
 		InjectionRecord.GraphRequest.Reset();
 
 		// Note we dont consume here, as we want the event to forward to the injection site trait
 		break;
 	case EAnimNextInjectionType::EvaluationModifier:
 		// We dont increment serial number when applying evaluation modifiers as we dont want to trigger a graph instantiation
-		Graph.EvaluationModifier = nullptr;
+		Graph.InjectionData.EvaluationModifier = nullptr;
 		InjectionRecord.ModifierRequest.Reset();
 
 		// Update status
@@ -176,14 +174,12 @@ void FAnimNextModuleInjectionComponent::OnReapplyInjection(const UE::AnimNext::F
 		FAnimNextAnimGraph& Graph = InjectableGraph.Get<FAnimNextAnimGraph>();
 		if(RequestPair.Value.GraphRequest.IsValid())
 		{
-			const FInjectionRequestArgs& GraphRequestArgs = RequestPair.Value.GraphRequest->GetArgs();
-			Graph.Asset = GraphRequestArgs.Object;
-			Graph.InjectionSerialNumber = RequestPair.Value.SerialNumber;
+			Graph.InjectionData.InjectionSerialNumber = RequestPair.Value.SerialNumber;
 		}
 		if(RequestPair.Value.ModifierRequest.IsValid())
 		{
 			const FInjectionRequestArgs& ModifierRequestArgs = RequestPair.Value.ModifierRequest->GetArgs();
-			Graph.EvaluationModifier = ModifierRequestArgs.EvaluationModifier;
+			Graph.InjectionData.EvaluationModifier = ModifierRequestArgs.EvaluationModifier;
 		}
 	}
 }
diff --git a/Engine/Plugins/Experimental/AnimNextAnimGraph/Source/AnimNextAnimGraph/Private/Injection/InjectionSiteTrait.cpp b/Engine/Plugins/Experimental/AnimNextAnimGraph/Source/AnimNextAnimGraph/Private/Injection/InjectionSiteTrait.cpp
--- a/Engine/Plugins/Experimental/AnimNextAnimGraph/Source/AnimNextAnimGraph/Private/Injection/InjectionSiteTrait.cpp
+++ b/Engine/Plugins/Experimental/AnimNextAnimGraph/Source/AnimNextAnimGraph/Private/Injection/InjectionSiteTrait.cpp
@@ -91,7 +91,7 @@ namespace UE::AnimNext
 
 		// Check to see if this injection event was routed to us
 		ensure(Event.SerialNumber != 0);
-		if (Graph.InjectionSerialNumber == Event.SerialNumber)
+		if (Graph.InjectionData.InjectionSerialNumber == Event.SerialNumber)
 		{
 			const FAnimNextInjectionRequestArgs& RequestArgs = Event.Request->GetArgs();
 			ensureMsgf(!InstanceData->PendingRequest.IsValid(), TEXT("Injection site %s already contained a pending request, it will be overwritten"), *RequestArgs.Name.ToString());
@@ -114,7 +114,7 @@ namespace UE::AnimNext
 
 		// Check to see if this uninjection event was routed to us
 		ensure(Event.SerialNumber != 0);
-		if (Graph.InjectionSerialNumber == Event.SerialNumber)
+		if (Graph.InjectionData.InjectionSerialNumber == Event.SerialNumber)
 		{
 			// Reset any pending request we might have, and cancel it
 			InstanceData->PendingRequest.Reset();
@@ -128,6 +128,21 @@ namespace UE::AnimNext
 
 	void FInjectionSiteTrait::PreUpdate(FUpdateTraversalContext& Context, const TTraitBinding<IUpdate>& Binding, const FTraitUpdateState& TraitState) const
 	{
+		auto TranslateBlendMode = [](EAnimNextInjectionBlendMode InBlendMode)
+		{
+			switch(InBlendMode)
+			{
+			case EAnimNextInjectionBlendMode::Standard:
+				return IBlendStack::EBlendMode::Standard;
+			case EAnimNextInjectionBlendMode::Inertialization:
+				return IBlendStack::EBlendMode::Inertialization;
+			default:
+				checkNoEntry();
+				break;
+			}
+			return IBlendStack::EBlendMode::Standard;
+		};
+
 		const FSharedData* SharedData = Binding.GetSharedData<FSharedData>();
 		FInstanceData* InstanceData = Binding.GetInstanceData<FInstanceData>();
 
@@ -152,21 +167,6 @@ namespace UE::AnimNext
 			const FInjectionRequestPtr Request = InstanceData->PendingRequest.Request;
 			InstanceData->PendingRequest.Reset();
 
-			auto TranslateBlendMode = [](EAnimNextInjectionBlendMode InBlendMode)
-			{
-				switch(InBlendMode)
-				{
-				case EAnimNextInjectionBlendMode::Standard:
-					return IBlendStack::EBlendMode::Standard;
-				case EAnimNextInjectionBlendMode::Inertialization:
-					return IBlendStack::EBlendMode::Inertialization;
-				default:
-					checkNoEntry();
-					break;
-				}
-				return IBlendStack::EBlendMode::Standard;
-			};
-			
 			IBlendStack::FGraphRequest BlendStackRequest;
 			EInjectionTrackerState State = EInjectionTrackerState::Inactive;
 			if (Request && Request->GetArgs().Object)
@@ -186,7 +186,7 @@ namespace UE::AnimNext
 				BlendStackRequest.BlendArgs = RequestArgs.BlendInSettings.Blend;
 				BlendStackRequest.BlendMode = TranslateBlendMode(RequestArgs.BlendInSettings.BlendMode);
 			}
-			else if(bInjectionSiteChanged && Graph.Asset)
+			else if(Graph.Asset)
 			{
 				// Push request from a pin-value change
 				State = EInjectionTrackerState::ActivePin;
@@ -317,19 +317,44 @@ namespace UE::AnimNext
 						Context.RaiseOutputTraitEvent(StatusUpdateEvent);
 					}
 
-					// Add a request to blend in our source using blend out time
-					IBlendStack::FGraphRequest BlendStackRequest;
-					BlendStackRequest.Type = IBlendStack::EGraphRequestType::Child;
-					BlendStackRequest.ChildPtr = Context.AllocateNodeInstance(Binding, SharedData->Source);
-					BlendStackRequest.BlendArgs.BlendTime = BlendOutTime;
-
-					InstanceData->PendingTracker.Initialize(FInjectionRequestPtr(), EInjectionTrackerState::ActiveSource);
-					ActiveChildIndex = BlendStackTrait.PushGraph(Context, MoveTemp(BlendStackRequest));
-
-					// Make sure we tracked it after blending
-					ensure(InstanceData->InjectionTrackers.IsValidIndex(ActiveChildIndex) &&
-						InstanceData->InjectionTrackers[ActiveChildIndex].Request == FInjectionRequestPtr() && 
-						InstanceData->InjectionTrackers[ActiveChildIndex].State == EInjectionTrackerState::ActiveSource);
+					// Blend to source child or do we still have an active pin input?
+					if(Graph.Asset != nullptr)
+					{
+						// Add a request to blend in our fallback graph using blend out time
+						IBlendStack::FGraphRequest BlendStackRequest;
+						BlendStackRequest.Type = IBlendStack::EGraphRequestType::Owned;
+						BlendStackRequest.AnimationGraph = IGraphFactory::GetGraphFromObjectWithFallback(Context, Binding, Graph.Asset);
+						if (BlendStackRequest.AnimationGraph != nullptr)
+						{
+							IPayloadFactory::GetPayloadForGraphWithFallback(Context, Binding, Graph.Asset, BlendStackRequest.AnimationGraph, BlendStackRequest.GraphPayload);
+						}
+						BlendStackRequest.BlendArgs.BlendTime = BlendOutTime;
+						BlendStackRequest.BlendMode = TranslateBlendMode(RequestArgs.BlendOutSettings.BlendMode);
+
+						InstanceData->PendingTracker.Initialize(FInjectionRequestPtr(), EInjectionTrackerState::ActivePin);
+						ActiveChildIndex = BlendStackTrait.PushGraph(Context, MoveTemp(BlendStackRequest));
+
+						// Make sure we tracked it after blending
+						ensure(InstanceData->InjectionTrackers.IsValidIndex(ActiveChildIndex) &&
+							InstanceData->InjectionTrackers[ActiveChildIndex].Request == FInjectionRequestPtr() && 
+							InstanceData->InjectionTrackers[ActiveChildIndex].State == EInjectionTrackerState::ActivePin);
+					}
+					else
+					{
+						// Add a request to blend in our source using blend out time
+						IBlendStack::FGraphRequest BlendStackRequest;
+						BlendStackRequest.Type = IBlendStack::EGraphRequestType::Child;
+						BlendStackRequest.ChildPtr = Context.AllocateNodeInstance(Binding, SharedData->Source);
+						BlendStackRequest.BlendArgs.BlendTime = BlendOutTime;
+
+						InstanceData->PendingTracker.Initialize(FInjectionRequestPtr(), EInjectionTrackerState::ActiveSource);
+						ActiveChildIndex = BlendStackTrait.PushGraph(Context, MoveTemp(BlendStackRequest));
+
+						// Make sure we tracked it after blending
+						ensure(InstanceData->InjectionTrackers.IsValidIndex(ActiveChildIndex) &&
+							InstanceData->InjectionTrackers[ActiveChildIndex].Request == FInjectionRequestPtr() && 
+							InstanceData->InjectionTrackers[ActiveChildIndex].State == EInjectionTrackerState::ActiveSource);
+					}
 				}
 			}
 		}
@@ -410,7 +435,7 @@ namespace UE::AnimNext
 		// Apply evaluation injection if present
 		const FSharedData* SharedData = Binding.GetSharedData<FSharedData>();
 		const FAnimNextAnimGraph& Graph = SharedData->GetGraph(Binding);
-		if(TSharedPtr<IEvaluationModifier> EvaluationModifier = Graph.EvaluationModifier.Pin())
+		if(TSharedPtr<IEvaluationModifier> EvaluationModifier = Graph.InjectionData.EvaluationModifier.Pin())
 		{
 			EvaluationModifier->PreEvaluate(Context);
 		}
@@ -424,7 +449,7 @@ namespace UE::AnimNext
 		// Apply evaluation injection if present
 		const FSharedData* SharedData = Binding.GetSharedData<FSharedData>();
 		const FAnimNextAnimGraph& Graph = SharedData->GetGraph(Binding);
-		if(TSharedPtr<IEvaluationModifier> EvaluationModifier = Graph.EvaluationModifier.Pin())
+		if(TSharedPtr<IEvaluationModifier> EvaluationModifier = Graph.InjectionData.EvaluationModifier.Pin())
 		{
 			EvaluationModifier->PostEvaluate(Context);
 		}
diff --git a/Engine/Plugins/Experimental/AnimNextAnimGraph/Source/AnimNextAnimGraph/Public/Graph/AnimNextAnimGraph.h b/Engine/Plugins/Experimental/AnimNextAnimGraph/Source/AnimNextAnimGraph/Public/Graph/AnimNextAnimGraph.h
--- a/Engine/Plugins/Experimental/AnimNextAnimGraph/Source/AnimNextAnimGraph/Public/Graph/AnimNextAnimGraph.h
+++ b/Engine/Plugins/Experimental/AnimNextAnimGraph/Source/AnimNextAnimGraph/Public/Graph/AnimNextAnimGraph.h
@@ -39,13 +39,30 @@ namespace UE::AnimNext::Editor
 	class FAnimNextGraphDetails;
 }
 
+// Injection data for an AnimNext Graph 
+USTRUCT(BlueprintType)
+struct ANIMNEXTANIMGRAPH_API FAnimNextGraphInjectionData
+{
+	GENERATED_BODY()
+
+private:
+	// Evaluation modifier to allow IEvaluate pass to be controlled
+	TWeakPtr<UE::AnimNext::IEvaluationModifier> EvaluationModifier;
+
+	// Serial number for the last injection that was made to this struct
+	uint32 InjectionSerialNumber = 0;
+
+	friend UE::AnimNext::FInjectionSiteTrait;
+	friend FAnimNextModuleInjectionComponent;
+};
+
 // Represents an instance of an AnimNext graph, either an asset or some externally-provided graph instance.
 // The asset can be set to refer directly to a graph asset, or it can be used via a 'factory' to generate an appropriate graph to run.
 // When public, represents an 'injection site' that can be used to parameterize graph execution externally. 
 //~ This struct uses UE reflection because we wish for the GC to keep the graph
 //~ alive while we own a reference to it (specifically when running an animation graph).
 //~ It is not intended to be serialized on disk with a live instance.
-USTRUCT(BlueprintType, DisplayName="AnimNext Graph", meta = (HideSubPins))
+USTRUCT(BlueprintType, DisplayName="AnimNext Graph")
 struct ANIMNEXTANIMGRAPH_API FAnimNextAnimGraph
 {
 	GENERATED_BODY()
@@ -61,18 +78,16 @@ private:
 	UPROPERTY(EditAnywhere, BlueprintReadWrite, Category=Graph, meta=(GetAllowedClasses="/Script/AnimNextAnimGraph.AnimNextAnimGraphSettings:GetAllowedAssetClasses", AllowPrivateAccess))
 	TObjectPtr<UObject> Asset;
 
+	// Injection data used to override the supplied graph
+	UPROPERTY(VisibleAnywhere, BlueprintReadOnly, Category=Graph, meta=(AllowPrivateAccess))
+	FAnimNextGraphInjectionData InjectionData;
+
 	// The host graph to use to run the asset. Applies only to top-level graphs being run in modules.
 	// We use a host graph to be able to blend between graph sources at the top level of a module when a RunGraph's input graph changes.
 	// If this is unset, then the project's default host graph will be used.
 	UPROPERTY(EditAnywhere, BlueprintReadWrite, Category=Graph, AdvancedDisplay,  meta=(AllowPrivateAccess))
 	TObjectPtr<UAnimNextAnimationGraph> HostGraph;
 
-	// Evaluation modifier to allow IEvaluate pass to be controlled
-	TWeakPtr<UE::AnimNext::IEvaluationModifier> EvaluationModifier;
-
-	// Serial number for the last injection that was made to this struct
-	uint32 InjectionSerialNumber = 0;
-
 	friend FRigUnit_AnimNextRunAnimationGraph_v1;	// We evaluate the instance
 	friend FRigUnit_AnimNextRunAnimationGraph_v2;	// We evaluate the instance
 	friend UE::AnimNext::FBlendStackCoreTrait;		// Requires access to take a reference to external graphs
