AnimNext: Remove delta time from FTimelineState constructor

No longer used

diff --git a/Engine/Plugins/Experimental/AnimNextAnimGraph/Source/AnimNextAnimGraph/Private/Traits/SequencePlayer.cpp b/Engine/Plugins/Experimental/AnimNextAnimGraph/Source/AnimNextAnimGraph/Private/Traits/SequencePlayer.cpp
--- a/Engine/Plugins/Experimental/AnimNextAnimGraph/Source/AnimNextAnimGraph/Private/Traits/SequencePlayer.cpp
+++ b/Engine/Plugins/Experimental/AnimNextAnimGraph/Source/AnimNextAnimGraph/Private/Traits/SequencePlayer.cpp
@@ -99,7 +99,6 @@ namespace UE::AnimNext
 				InstanceData->InternalTimeAccumulator,
 				AnimSeq->GetPlayLength(),
 				SharedData->GetPlayRate(Binding),
-				InstanceData->DeltaTimeRecord.Delta,
 				SharedData->GetbLoop(Binding));
 		}
 
diff --git a/Engine/Plugins/Experimental/AnimNextAnimGraph/Source/AnimNextAnimGraph/Public/TraitInterfaces/ITimeline.h b/Engine/Plugins/Experimental/AnimNextAnimGraph/Source/AnimNextAnimGraph/Public/TraitInterfaces/ITimeline.h
--- a/Engine/Plugins/Experimental/AnimNextAnimGraph/Source/AnimNextAnimGraph/Public/TraitInterfaces/ITimeline.h
+++ b/Engine/Plugins/Experimental/AnimNextAnimGraph/Source/AnimNextAnimGraph/Public/TraitInterfaces/ITimeline.h
@@ -51,7 +51,7 @@ namespace UE::AnimNext
 		FTimelineState() = default;
 
 		// Construct with a specific state
-		FTimelineState(float InPosition, float InDuration, float InPlayRate, float InDelta, bool bInIsLooping)
+		FTimelineState(float InPosition, float InDuration, float InPlayRate, bool bInIsLooping)
 			: Position(InPosition)
 			, Duration(InDuration)
 			, PlayRate(InPlayRate)
