AnimNext: Change PlayAnim to use loop from anim sequence

diff --git a/Engine/Plugins/Experimental/AnimNextAnimGraph/Source/AnimNextAnimGraph/Private/Injection/InjectionUtils.cpp b/Engine/Plugins/Experimental/AnimNextAnimGraph/Source/AnimNextAnimGraph/Private/Injection/InjectionUtils.cpp
--- a/Engine/Plugins/Experimental/AnimNextAnimGraph/Source/AnimNextAnimGraph/Private/Injection/InjectionUtils.cpp
+++ b/Engine/Plugins/Experimental/AnimNextAnimGraph/Source/AnimNextAnimGraph/Private/Injection/InjectionUtils.cpp
@@ -77,6 +77,7 @@ namespace UE::AnimNext
 		PlayNativeInterface.AnimSequence = InAnimSequence;
 		PlayNativeInterface.PlayRate = InPlayRate;
 		PlayNativeInterface.StartPosition = InStartPosition;
+		PlayNativeInterface.Loop = InAnimSequence->bLoop;
 
 		FInjectionRequestArgs RequestArgs;
 		RequestArgs.Name = InSiteName;
