AnimNext: Add the ability for RigUnits to supply required module components

Fixes first-frame ensure() where the injection component was not added at the point an injection request was made (when the module is created on the same frame as the injection request is issued).
This system could additionally be used to provide things like actor transforms, although would neeed to be augmented with an explicit user-specified set of component structs to avoid first-frame issues with dynamically supplied anim graphs.

diff --git a/Engine/Plugins/Experimental/AnimNext/Source/AnimNext/Internal/Module/AnimNextModule.h b/Engine/Plugins/Experimental/AnimNext/Source/AnimNext/Internal/Module/AnimNextModule.h
--- a/Engine/Plugins/Experimental/AnimNext/Source/AnimNext/Internal/Module/AnimNextModule.h
+++ b/Engine/Plugins/Experimental/AnimNext/Source/AnimNext/Internal/Module/AnimNextModule.h
@@ -72,7 +72,12 @@ protected:
 	friend UE::AnimNext::FAnimNextModuleImpl;
 	friend class UE::AnimNext::Editor::FVariableCustomization;
 	friend struct UE::AnimNext::FParametersProxy;
+	friend FAnimNextModuleInstance;
 
+	// All components that are required on startup for this module
+	UPROPERTY()
+	TArray<TObjectPtr<UScriptStruct>> RequiredComponents;
+	
 #if WITH_EDITORONLY_DATA
 	UPROPERTY()
 	FAnimNextGraphState DefaultState_DEPRECATED;
diff --git a/Engine/Plugins/Experimental/AnimNext/Source/AnimNext/Internal/Module/AnimNextModuleInstance.h b/Engine/Plugins/Experimental/AnimNext/Source/AnimNext/Internal/Module/AnimNextModuleInstance.h
--- a/Engine/Plugins/Experimental/AnimNext/Source/AnimNext/Internal/Module/AnimNextModuleInstance.h
+++ b/Engine/Plugins/Experimental/AnimNext/Source/AnimNext/Internal/Module/AnimNextModuleInstance.h
@@ -307,7 +307,7 @@ struct TStructOpsTypeTraits<FAnimNextModuleInstance> : public TStructOpsTypeTrai
 template<typename ComponentType>
 ComponentType& FAnimNextModuleInstance::GetComponent()
 {
-	const FName ComponentName = ComponentType::StaticComponentName();
+	const FName ComponentName = ComponentType::StaticStruct()->GetFName();
 	const int32 ComponentNameHash = GetTypeHash(ComponentName);
 
 	if (ComponentType* Component = static_cast<ComponentType*>(TryGetComponent(ComponentNameHash, ComponentName)))
@@ -315,13 +315,15 @@ ComponentType& FAnimNextModuleInstance::GetComponent()
 		return *Component;
 	}
 
-	return static_cast<ComponentType&>(AddComponentInternal(ComponentNameHash, ComponentName, TInstancedStruct<ComponentType>::Make(*this)));
+	TInstancedStruct<ComponentType> ComponentStruct = TInstancedStruct<ComponentType>::Make();
+	ComponentStruct.template GetMutable<ComponentType>().Initialize(*this);
+	return static_cast<ComponentType&>(AddComponentInternal(ComponentNameHash, ComponentName, MoveTemp(ComponentStruct)));
 }
 
 template<class ComponentType>
 ComponentType* FAnimNextModuleInstance::TryGetComponent()
 {
-	const FName ComponentName = ComponentType::StaticComponentName();
+	const FName ComponentName = ComponentType::StaticStruct()->GetFName();
 	const int32 ComponentNameHash = GetTypeHash(ComponentName);
 
 	return static_cast<ComponentType*>(TryGetComponent(ComponentNameHash, ComponentName));
@@ -330,7 +332,7 @@ ComponentType* FAnimNextModuleInstance::TryGetComponent()
 template<class ComponentType>
 const ComponentType* FAnimNextModuleInstance::TryGetComponent() const
 {
-	const FName ComponentName = ComponentType::StaticComponentName();
+	const FName ComponentName = ComponentType::StaticStruct()->GetFName();
 	const int32 ComponentNameHash = GetTypeHash(ComponentName);
 
 	return static_cast<ComponentType*>(TryGetComponent(ComponentNameHash, ComponentName));
diff --git a/Engine/Plugins/Experimental/AnimNext/Source/AnimNext/Internal/Module/ModuleInstanceComponent.h b/Engine/Plugins/Experimental/AnimNext/Source/AnimNext/Internal/Module/ModuleInstanceComponent.h
--- a/Engine/Plugins/Experimental/AnimNext/Source/AnimNext/Internal/Module/ModuleInstanceComponent.h
+++ b/Engine/Plugins/Experimental/AnimNext/Source/AnimNext/Internal/Module/ModuleInstanceComponent.h
@@ -8,11 +8,6 @@
 struct FAnimNextModuleInstance;
 struct FAnimNextTraitEvent;
 
-// Declares a module instance component and implements the necessary boilerplate
-#define ANIMNEXT_DECLARE_MODULE_INSTANCE_COMPONENT(ComponentType) \
-	static FName StaticComponentName() { return FName(TEXT(#ComponentType)); } \
-	virtual FName GetComponentName() const override { return StaticComponentName(); }
-
 /**
  * FAnimNextModuleInstanceComponent
  *
@@ -25,11 +20,9 @@ struct FAnimNextModuleInstanceComponent
 	GENERATED_BODY()
 
 	FAnimNextModuleInstanceComponent() = default;
-	explicit FAnimNextModuleInstanceComponent(FAnimNextModuleInstance& InOwnerInstance) : OwnerInstance(&InOwnerInstance) {}
 	virtual ~FAnimNextModuleInstanceComponent() {}
 
-	static FName StaticComponentName() { return FName(TEXT("FAnimNextModuleInstanceComponent")); }
-	virtual FName GetComponentName() const { return StaticComponentName(); }
+	ANIMNEXT_API void Initialize(FAnimNextModuleInstance& InOwnerInstance);
 
 	// Returns the owning module instance this component lives on
 	FAnimNextModuleInstance& GetModuleInstance() { return *OwnerInstance; }
@@ -37,6 +30,9 @@ struct FAnimNextModuleInstanceComponent
 	// Returns the owning module instance this component lives on
 	const FAnimNextModuleInstance& GetModuleInstance() const { return *OwnerInstance; }
 
+	// Called when the component is first created to initialize it
+	virtual void OnInitialize(FAnimNextModuleInstance& InOwnerInstance) {}
+
 	// Called during module execution for any events to be handled
 	virtual void OnTraitEvent(FAnimNextTraitEvent& Event) {}
 
diff --git a/Engine/Plugins/Experimental/AnimNext/Source/AnimNext/Private/Module/AnimNextModuleInstance.cpp b/Engine/Plugins/Experimental/AnimNext/Source/AnimNext/Private/Module/AnimNextModuleInstance.cpp
--- a/Engine/Plugins/Experimental/AnimNext/Source/AnimNext/Private/Module/AnimNextModuleInstance.cpp
+++ b/Engine/Plugins/Experimental/AnimNext/Source/AnimNext/Private/Module/AnimNextModuleInstance.cpp
@@ -235,6 +235,17 @@ void FAnimNextModuleInstance::Initialize()
 		// Now initialize the 'instance', cache memory handles etc. in the context
 		VM->InitializeInstance(ExtendedExecuteContext);
 
+		// Allocate compiled-in module components
+		for(UScriptStruct* ComponentStruct : Module->RequiredComponents)
+		{
+			const FName ComponentName = ComponentStruct->GetFName();
+			const int32 ComponentNameHash = GetTypeHash(ComponentName);
+
+			TInstancedStruct<FAnimNextModuleInstanceComponent> Component(ComponentStruct);
+			Component.GetMutable<FAnimNextModuleInstanceComponent>().Initialize(*this);
+			AddComponentInternal(ComponentNameHash, ComponentName, MoveTemp(Component));
+		}
+
 		// Just pause now if we arent needing an initial update
 		if(InitMethod == EAnimNextModuleInitMethod::None)
 		{
diff --git a/Engine/Plugins/Experimental/AnimNext/Source/AnimNext/Private/Module/ModuleInstanceComponent.cpp b/Engine/Plugins/Experimental/AnimNext/Source/AnimNext/Private/Module/ModuleInstanceComponent.cpp
new file mode 100644
--- /dev/null
+++ b/Engine/Plugins/Experimental/AnimNext/Source/AnimNext/Private/Module/ModuleInstanceComponent.cpp
@@ -0,0 +1,10 @@
+﻿// Copyright Epic Games, Inc. All Rights Reserved.
+
+#include "Module/ModuleInstanceComponent.h"
+
+void FAnimNextModuleInstanceComponent::Initialize(FAnimNextModuleInstance& InOwnerInstance)
+{
+	OwnerInstance = &InOwnerInstance;
+
+	OnInitialize(InOwnerInstance);
+}
\ No newline at end of file
diff --git a/Engine/Plugins/Experimental/AnimNext/Source/AnimNextUncookedOnly/Internal/AnimNextRigVMAssetEditorData.h b/Engine/Plugins/Experimental/AnimNext/Source/AnimNextUncookedOnly/Internal/AnimNextRigVMAssetEditorData.h
--- a/Engine/Plugins/Experimental/AnimNext/Source/AnimNextUncookedOnly/Internal/AnimNextRigVMAssetEditorData.h
+++ b/Engine/Plugins/Experimental/AnimNext/Source/AnimNextUncookedOnly/Internal/AnimNextRigVMAssetEditorData.h
@@ -307,12 +307,25 @@ protected:
 	// Override to allow assets to prevent certain entries being created
 	virtual bool CanAddNewEntry(TSubclassOf<UAnimNextRigVMAssetEntry> InClass) const { return true; }
 
-	// Allows this asset to generate variables to be injected at compilation time, separate method to allow programmatic graphs to use these vars
+	// Begin Compilation overrides, in order of operation
+
+	// Compilation phase 1: Called before RigVM compilation to setup compiler settings and clean our outer asset of compiler-generated data
+	virtual void OnPreCompileAsset(FRigVMCompileSettings& InSettings) {}
+
+	// Compilation phase 2: Called before RigVM compilation to allow this asset to generate variables to be injected, separate method to allow programmatic graphs to use these vars
 	// These variables will be regenerated each compile, and are not saved between compiles
-	virtual void GetProgrammaticVariables(const FRigVMCompileSettings& InSettings, TArray<FAnimNextProgrammaticVariable>& OutVariables) {}
+	virtual void OnPreCompileGetProgrammaticVariables(const FRigVMCompileSettings& InSettings, TArray<FAnimNextProgrammaticVariable>& OutVariables) {}
+
+	// Compilation phase 3: Called before RigVM compilation to allow this asset to generate graphs to be injected
+	virtual void OnPreCompileGetProgrammaticGraphs(const FRigVMCompileSettings& InSettings, TArray<URigVMGraph*>& OutGraphs) {}
 
-	// Allows this asset to generate graphs to be injected at compilation time
-	virtual void GetProgrammaticGraphs(const FRigVMCompileSettings& InSettings, TArray<URigVMGraph*>& OutGraphs) {}
+	// Compilation phase 4: Called before RigVM compilation to allow this asset to process, transform or replace the graphs that will be compiled
+	virtual void OnPreCompileProcessGraphs(const FRigVMCompileSettings& InSettings, TArray<URigVMGraph*>& InOutGraphs) {}
+
+	// Compilation phase 5: Called after RigVM compilation to clean up/finish the compilation process
+	virtual void OnPostCompileCleanup(const FRigVMCompileSettings& InSettings) {}
+
+	// End Compilation overrides
 
 	// Customization point for derived types to transform new asset entries
 	virtual void CustomizeNewAssetEntry(UAnimNextRigVMAssetEntry* InNewEntry) const {}
@@ -466,6 +479,9 @@ protected:
 	// Remove any programmatic graphs generated during compilation and consign them to the transient package
 	void RemoveProgrammaticGraphs(TArrayView<URigVMGraph*> InGraphs);
 
+	// Remove any transient graphs in the passed-in array (e.g. generated during compilation) and consign them to the transient package
+	void RemoveTransientGraphs(TArrayView<URigVMGraph*> InGraphs);
+
 	// Handle removing a notify
 	static void HandleRemoveNotify(UObject* InAsset, const FString& InFindString, bool bFindWholeWord, ESearchCase::Type InSearchCase);
 
diff --git a/Engine/Plugins/Experimental/AnimNext/Source/AnimNextUncookedOnly/Internal/Module/AnimNextModule_EditorData.h b/Engine/Plugins/Experimental/AnimNext/Source/AnimNextUncookedOnly/Internal/Module/AnimNextModule_EditorData.h
--- a/Engine/Plugins/Experimental/AnimNext/Source/AnimNextUncookedOnly/Internal/Module/AnimNextModule_EditorData.h
+++ b/Engine/Plugins/Experimental/AnimNext/Source/AnimNextUncookedOnly/Internal/Module/AnimNextModule_EditorData.h
@@ -53,7 +53,9 @@ private:
 	virtual UScriptStruct* GetExecuteContextStruct() const override { return FAnimNextExecuteContext::StaticStruct(); }
 	virtual TConstArrayView<TSubclassOf<UAnimNextRigVMAssetEntry>> GetEntryClasses() const override;
 	virtual void RecompileVM() override;
-	virtual void GetProgrammaticGraphs(const FRigVMCompileSettings& InSettings, TArray<URigVMGraph*>& OutGraphs) override;
+	virtual void OnPreCompileAsset(FRigVMCompileSettings& InSettings) override;
+	virtual void OnPreCompileGetProgrammaticGraphs(const FRigVMCompileSettings& InSettings, TArray<URigVMGraph*>& OutGraphs) override;
+	virtual void OnPreCompileProcessGraphs(const FRigVMCompileSettings& InSettings, TArray<URigVMGraph*>& InOutGraphs) override;
 	virtual void CustomizeNewAssetEntry(UAnimNextRigVMAssetEntry* InNewEntry) const override;
 
 private:
diff --git a/Engine/Plugins/Experimental/AnimNext/Source/AnimNextUncookedOnly/Private/AnimNextRigVMAssetEditorData.cpp b/Engine/Plugins/Experimental/AnimNext/Source/AnimNextUncookedOnly/Private/AnimNextRigVMAssetEditorData.cpp
--- a/Engine/Plugins/Experimental/AnimNext/Source/AnimNextUncookedOnly/Private/AnimNextRigVMAssetEditorData.cpp
+++ b/Engine/Plugins/Experimental/AnimNext/Source/AnimNextUncookedOnly/Private/AnimNextRigVMAssetEditorData.cpp
@@ -764,6 +764,15 @@ void UAnimNextRigVMAssetEditorData::RecompileVM()
 	UAnimNextRigVMAsset* Asset = FUtils::GetAsset<UAnimNextRigVMAsset>(this);
 	FScopedCompilerResults CompilerResults(Asset);
 
+	VMCompileSettings.SetExecuteContextStruct(FAnimNextExecuteContext::StaticStruct());
+	FRigVMCompileSettings Settings = (bCompileInDebugMode) ? FRigVMCompileSettings::Fast(VMCompileSettings.GetExecuteContextStruct()) : VMCompileSettings;
+	Settings.SurpressInfoMessages = false;
+	Settings.ASTSettings.ReportDelegate.BindUObject(this, &UAnimNextRigVMAssetEditorData::HandleReportFromCompiler);
+
+	Asset->VMRuntimeSettings = VMRuntimeSettings;
+
+	OnPreCompileAsset(Settings);
+
 	CachedExports.Reset();  // asset variables and other tags will be updated at the end by AssetRegistry->AssetUpdateTags
 
 	bWarningsDuringCompilation = false;
@@ -777,28 +786,23 @@ void UAnimNextRigVMAssetEditorData::RecompileVM()
 		TGuardValue<bool> ReentrantGuardSelf(bSuspendModelNotificationsForSelf, true);
 		TGuardValue<bool> ReentrantGuardOthers(RigVMClient.bSuspendModelNotificationsForOthers, true);
 
-		VMCompileSettings.SetExecuteContextStruct(FAnimNextExecuteContext::StaticStruct());
-		FRigVMCompileSettings Settings = (bCompileInDebugMode) ? FRigVMCompileSettings::Fast(VMCompileSettings.GetExecuteContextStruct()) : VMCompileSettings;
-		Settings.SurpressInfoMessages = false;
-		Settings.ASTSettings.ReportDelegate.BindUObject(this, &UAnimNextRigVMAssetEditorData::HandleReportFromCompiler);
-
 		FUtils::RecreateVM(Asset);
 
 		FUtils::CompileVariables(Settings, Asset);
 
-		GetProgrammaticGraphs(Settings, ProgrammaticGraphs);
+		OnPreCompileGetProgrammaticGraphs(Settings, ProgrammaticGraphs);
 		for(URigVMGraph* ProgrammaticGraph : ProgrammaticGraphs)
 		{
 			check(ProgrammaticGraph != nullptr);
 		}
 
-		Asset->VMRuntimeSettings = VMRuntimeSettings;
-
 		FRigVMClient* VMClient = GetRigVMClient();
 
 		TArray<URigVMGraph*> AllGraphs = VMClient->GetAllModels(false, false);
 		AllGraphs.Append(ProgrammaticGraphs);
 
+		OnPreCompileProcessGraphs(Settings, AllGraphs);
+
 		if(AllGraphs.Num() > 0)
 		{
 			URigVMController* Controller = VMClient->GetOrCreateController(AllGraphs[0]);
@@ -844,6 +848,10 @@ void UAnimNextRigVMAssetEditorData::RecompileVM()
 			RemoveProgrammaticGraphs(ProgrammaticGraphs);
 		}
 
+		RemoveTransientGraphs(AllGraphs);
+
+		OnPostCompileCleanup(Settings);
+
 #if WITH_EDITOR
 		//	RefreshBreakpoints(EditorData);
 #endif
@@ -867,6 +875,20 @@ void UAnimNextRigVMAssetEditorData::RemoveProgrammaticGraphs(TArrayView<URigVMGr
 	}
 }
 
+void UAnimNextRigVMAssetEditorData::RemoveTransientGraphs(TArrayView<URigVMGraph*> InGraphs)
+{
+	FRigVMClient* VMClient = GetRigVMClient();
+	
+	for(URigVMGraph* Graph : InGraphs)
+	{
+		if(Graph->HasAnyFlags(RF_Transient))
+		{
+			VMClient->RemoveController(Graph);
+			Graph->Rename(nullptr, GetTransientPackage(), REN_ForceNoResetLoaders | REN_DoNotDirty | REN_DontCreateRedirectors | REN_NonTransactional);
+		}
+	}
+}
+
 void UAnimNextRigVMAssetEditorData::HandleRemoveNotify(UObject* InAsset, const FString& InFindString, bool bFindWholeWord, ESearchCase::Type InSearchCase)
 {
 	UAnimNextRigVMAsset* Asset = Cast<UAnimNextRigVMAsset>(InAsset);
diff --git a/Engine/Plugins/Experimental/AnimNext/Source/AnimNextUncookedOnly/Private/Module/AnimNextModule_EditorData.cpp b/Engine/Plugins/Experimental/AnimNext/Source/AnimNextUncookedOnly/Private/Module/AnimNextModule_EditorData.cpp
--- a/Engine/Plugins/Experimental/AnimNext/Source/AnimNextUncookedOnly/Private/Module/AnimNextModule_EditorData.cpp
+++ b/Engine/Plugins/Experimental/AnimNext/Source/AnimNextUncookedOnly/Private/Module/AnimNextModule_EditorData.cpp
@@ -8,6 +8,7 @@
 #include "Entries/AnimNextEventGraphEntry.h"
 #include "Entries/AnimNextDataInterfaceEntry.h"
 #include "Entries/AnimNextVariableEntry.h"
+#include "String/ParseTokens.h"
 #include "UObject/AssetRegistryTagsContext.h"
 #include "UObject/LinkerLoad.h"
 #include "Variables/AnimNextUniversalObjectLocatorBindingData.h"
@@ -41,13 +42,65 @@ TConstArrayView<TSubclassOf<UAnimNextRigVMAssetEntry>> UAnimNextModule_EditorDat
 	return Classes;
 }
 
-void UAnimNextModule_EditorData::GetProgrammaticGraphs(const FRigVMCompileSettings& InSettings, TArray<URigVMGraph*>& OutGraphs)
+void UAnimNextModule_EditorData::OnPreCompileAsset(FRigVMCompileSettings& InSettings)
+{
+	using namespace UE::AnimNext::UncookedOnly;
+
+	UAnimNextModule* Module = FUtils::GetAsset<UAnimNextModule>(this);
+
+	Module->RequiredComponents.Empty();
+}
+
+void UAnimNextModule_EditorData::OnPreCompileGetProgrammaticGraphs(const FRigVMCompileSettings& InSettings, TArray<URigVMGraph*>& OutGraphs)
 {
 	using namespace UE::AnimNext::UncookedOnly;
 
 	FUtils::CompileVariableBindings(InSettings, FUtils::GetAsset<UAnimNextModule>(this), OutGraphs);
 }
 
+void UAnimNextModule_EditorData::OnPreCompileProcessGraphs(const FRigVMCompileSettings& InSettings, TArray<URigVMGraph*>& InOutGraphs)
+{
+	using namespace UE::AnimNext::UncookedOnly;
+
+	UAnimNextModule* Module = FUtils::GetAsset<UAnimNextModule>(this);
+
+	// Gather any required components from node metadata
+	for(URigVMGraph* Graph : InOutGraphs)
+	{
+		for(URigVMNode* Node : Graph->GetNodes())
+		{
+			URigVMTemplateNode* TemplateNode = Cast<URigVMTemplateNode>(Node);
+			if(TemplateNode == nullptr)
+			{
+				continue;
+			}
+
+			UScriptStruct* Struct = TemplateNode->GetScriptStruct();
+			if(Struct == nullptr)
+			{
+				continue;
+			}
+
+			static FName Metadata_RequiredComponents("RequiredComponents");
+			FString ComponentsString = Struct->GetMetaData(Metadata_RequiredComponents);
+			if(ComponentsString.Len() == 0)
+			{
+				continue;
+			}
+
+			UE::String::ParseTokens(ComponentsString, TEXT(','), [Module](FStringView InToken)
+			{
+				FString StructName(InToken);
+				UScriptStruct* ComponentStruct = FindFirstObject<UScriptStruct>(*StructName);
+				if(ComponentStruct)
+				{
+					Module->RequiredComponents.Add(ComponentStruct);
+				}
+			});
+		}
+	}
+}
+
 void UAnimNextModule_EditorData::CustomizeNewAssetEntry(UAnimNextRigVMAssetEntry* InNewEntry) const
 {
 	Super::CustomizeNewAssetEntry(InNewEntry);
diff --git a/Engine/Plugins/Experimental/AnimNext/Source/AnimNextUncookedOnly/Private/UncookedOnlyUtils.cpp b/Engine/Plugins/Experimental/AnimNext/Source/AnimNextUncookedOnly/Private/UncookedOnlyUtils.cpp
--- a/Engine/Plugins/Experimental/AnimNext/Source/AnimNextUncookedOnly/Private/UncookedOnlyUtils.cpp
+++ b/Engine/Plugins/Experimental/AnimNext/Source/AnimNextUncookedOnly/Private/UncookedOnlyUtils.cpp
@@ -69,7 +69,7 @@ void FUtils::CompileVariables(const FRigVMCompileSettings& InSettings, UAnimNext
 
 	// Gather programmatic variables regenerated each compile
 	TArray<FAnimNextProgrammaticVariable> ProgrammaticVariables;
-	EditorData->GetProgrammaticVariables(InSettings, ProgrammaticVariables);
+	EditorData->OnPreCompileGetProgrammaticVariables(InSettings, ProgrammaticVariables);
 
 	struct FStructEntryInfo
 	{
diff --git a/Engine/Plugins/Experimental/AnimNextAnimGraph/Source/AnimNextAnimGraph/Internal/Graph/AnimNextModuleAnimGraphComponent.h b/Engine/Plugins/Experimental/AnimNextAnimGraph/Source/AnimNextAnimGraph/Internal/Graph/AnimNextModuleAnimGraphComponent.h
--- a/Engine/Plugins/Experimental/AnimNextAnimGraph/Source/AnimNextAnimGraph/Internal/Graph/AnimNextModuleAnimGraphComponent.h
+++ b/Engine/Plugins/Experimental/AnimNextAnimGraph/Source/AnimNextAnimGraph/Internal/Graph/AnimNextModuleAnimGraphComponent.h
@@ -15,10 +15,7 @@ struct FAnimNextModuleAnimGraphComponent : public FAnimNextModuleInstanceCompone
 {
 	GENERATED_BODY()
 
-	ANIMNEXT_DECLARE_MODULE_INSTANCE_COMPONENT(FAnimNextModuleAnimGraphComponent);
-
 	FAnimNextModuleAnimGraphComponent() = default;
-	explicit FAnimNextModuleAnimGraphComponent(FAnimNextModuleInstance& InOwnerInstance);
 
 	TWeakPtr<FAnimNextGraphInstance> AllocateInstance(const UAnimNextAnimationGraph* InAnimationGraph, FAnimNextGraphInstance* InParentInstance = nullptr, FName InEntryPoint = NAME_None);
 
diff --git a/Engine/Plugins/Experimental/AnimNextAnimGraph/Source/AnimNextAnimGraph/Internal/Graph/RigUnit_AnimNextRunAnimationGraph_v2.h b/Engine/Plugins/Experimental/AnimNextAnimGraph/Source/AnimNextAnimGraph/Internal/Graph/RigUnit_AnimNextRunAnimationGraph_v2.h
--- a/Engine/Plugins/Experimental/AnimNextAnimGraph/Source/AnimNextAnimGraph/Internal/Graph/RigUnit_AnimNextRunAnimationGraph_v2.h
+++ b/Engine/Plugins/Experimental/AnimNextAnimGraph/Source/AnimNextAnimGraph/Internal/Graph/RigUnit_AnimNextRunAnimationGraph_v2.h
@@ -23,7 +23,7 @@ struct FAnimNextRunGraphWorkData_v2
 };
 
 /** Runs an animation graph */
-USTRUCT(meta=(DisplayName="Run Graph", Category="Animation Graph", NodeColor="0, 1, 1", Keywords="Trait,Stack"))
+USTRUCT(meta=(DisplayName="Run Graph", Category="Animation Graph", NodeColor="0, 1, 1", Keywords="Trait,Stack", RequiredComponents="AnimNextModuleInjectionComponent"))
 struct ANIMNEXTANIMGRAPH_API FRigUnit_AnimNextRunAnimationGraph_v2 : public FRigUnit_AnimNextBase
 {
 	GENERATED_BODY()
diff --git a/Engine/Plugins/Experimental/AnimNextAnimGraph/Source/AnimNextAnimGraph/Private/Graph/AnimNextModuleAnimGraphComponent.cpp b/Engine/Plugins/Experimental/AnimNextAnimGraph/Source/AnimNextAnimGraph/Private/Graph/AnimNextModuleAnimGraphComponent.cpp
--- a/Engine/Plugins/Experimental/AnimNextAnimGraph/Source/AnimNextAnimGraph/Private/Graph/AnimNextModuleAnimGraphComponent.cpp
+++ b/Engine/Plugins/Experimental/AnimNextAnimGraph/Source/AnimNextAnimGraph/Private/Graph/AnimNextModuleAnimGraphComponent.cpp
@@ -5,11 +5,6 @@
 #include "Graph/AnimNextGraphInstance.h"
 #include "Graph/AnimNextAnimationGraph.h"
 
-FAnimNextModuleAnimGraphComponent::FAnimNextModuleAnimGraphComponent(FAnimNextModuleInstance& InOwnerInstance)
-	: FAnimNextModuleInstanceComponent(InOwnerInstance)
-{
-}
-
 TWeakPtr<FAnimNextGraphInstance> FAnimNextModuleAnimGraphComponent::AllocateInstance(const UAnimNextAnimationGraph* InAnimationGraph, FAnimNextGraphInstance* InParentInstance, FName InEntryPoint)
 {
 	check(InAnimationGraph);
diff --git a/Engine/Plugins/Experimental/AnimNextAnimGraph/Source/AnimNextAnimGraph/Private/Graph/RigUnit_AnimNextRunAnimationGraph_v2.cpp b/Engine/Plugins/Experimental/AnimNextAnimGraph/Source/AnimNextAnimGraph/Private/Graph/RigUnit_AnimNextRunAnimationGraph_v2.cpp
--- a/Engine/Plugins/Experimental/AnimNextAnimGraph/Source/AnimNextAnimGraph/Private/Graph/RigUnit_AnimNextRunAnimationGraph_v2.cpp
+++ b/Engine/Plugins/Experimental/AnimNextAnimGraph/Source/AnimNextAnimGraph/Private/Graph/RigUnit_AnimNextRunAnimationGraph_v2.cpp
@@ -94,9 +94,6 @@ FRigUnit_AnimNextRunAnimationGraph_v2_Execute()
 			Hosts.Add(&Adapter);
 		}
 
-		// Ensure we have an injection component, as this node acts as a top-level injection site
-		ModuleInstance.GetComponent<FAnimNextModuleInjectionComponent>();
-
 		// Inject the default graph if there is any (_after_ pin-exposed variables so they cant override it)
 		FGraphInstanceInjectionComponent& HostInjectionComponent = HostInstance.GetComponent<FGraphInstanceInjectionComponent>();
 		FName GraphName;
diff --git a/Engine/Plugins/Experimental/AnimNextAnimGraph/Source/AnimNextAnimGraph/Private/Injection/AnimNextModuleInjectionComponent.cpp b/Engine/Plugins/Experimental/AnimNextAnimGraph/Source/AnimNextAnimGraph/Private/Injection/AnimNextModuleInjectionComponent.cpp
--- a/Engine/Plugins/Experimental/AnimNextAnimGraph/Source/AnimNextAnimGraph/Private/Injection/AnimNextModuleInjectionComponent.cpp
+++ b/Engine/Plugins/Experimental/AnimNextAnimGraph/Source/AnimNextAnimGraph/Private/Injection/AnimNextModuleInjectionComponent.cpp
@@ -10,10 +10,10 @@
 #include "Module/ModuleTaskContext.h"
 
 
-FAnimNextModuleInjectionComponent::FAnimNextModuleInjectionComponent(FAnimNextModuleInstance& InOwnerInstance)
-	: FAnimNextModuleInstanceComponent(InOwnerInstance)
-	, InjectionInfo(InOwnerInstance)
+void FAnimNextModuleInjectionComponent::OnInitialize(FAnimNextModuleInstance& InOwnerInstance)
 {
+	InjectionInfo = UE::AnimNext::FInjectionInfo(InOwnerInstance);
+	
 	// Register re-injection for each (user) tick function
 	for(UE::AnimNext::FModuleEventTickFunction& TickFunction : InOwnerInstance.GetTickFunctions())
 	{
diff --git a/Engine/Plugins/Experimental/AnimNextAnimGraph/Source/AnimNextAnimGraph/Private/Injection/AnimNextModuleInjectionComponent.h b/Engine/Plugins/Experimental/AnimNextAnimGraph/Source/AnimNextAnimGraph/Private/Injection/AnimNextModuleInjectionComponent.h
--- a/Engine/Plugins/Experimental/AnimNextAnimGraph/Source/AnimNextAnimGraph/Private/Injection/AnimNextModuleInjectionComponent.h
+++ b/Engine/Plugins/Experimental/AnimNextAnimGraph/Source/AnimNextAnimGraph/Private/Injection/AnimNextModuleInjectionComponent.h
@@ -21,16 +21,14 @@ struct FAnimNextModuleInjectionComponent : public FAnimNextModuleInstanceCompone
 {
 	GENERATED_BODY()
 
-	ANIMNEXT_DECLARE_MODULE_INSTANCE_COMPONENT(FAnimNextModuleInjectionComponent);
-
 	FAnimNextModuleInjectionComponent() = default;
-	explicit FAnimNextModuleInjectionComponent(FAnimNextModuleInstance& InOwnerInstance);
 
 	// Get the cached injection info for our module 
 	const UE::AnimNext::FInjectionInfo& GetInjectionInfo() const { return InjectionInfo; }
 
 private:
 	// FAnimNextModuleInstanceComponent interface
+	virtual void OnInitialize(FAnimNextModuleInstance& InOwnerInstance) override;
 	virtual void OnTraitEvent(FAnimNextTraitEvent& Event) override;
 
 	void OnInjectionEvent(UE::AnimNext::FInjection_InjectEvent& InEvent);
diff --git a/Engine/Plugins/Experimental/AnimNextAnimGraph/Source/AnimNextAnimGraph/Private/Traits/NotifyDispatcher.h b/Engine/Plugins/Experimental/AnimNextAnimGraph/Source/AnimNextAnimGraph/Private/Traits/NotifyDispatcher.h
--- a/Engine/Plugins/Experimental/AnimNextAnimGraph/Source/AnimNextAnimGraph/Private/Traits/NotifyDispatcher.h
+++ b/Engine/Plugins/Experimental/AnimNextAnimGraph/Source/AnimNextAnimGraph/Private/Traits/NotifyDispatcher.h
@@ -39,10 +39,6 @@ struct FAnimNextNotifyDispatcherComponent : public FAnimNextModuleInstanceCompon
 
 	FAnimNextNotifyDispatcherComponent() = default;
 
-	FAnimNextNotifyDispatcherComponent(FAnimNextModuleInstance& InOwnerInstance)
-		: FAnimNextModuleInstanceComponent(InOwnerInstance)
-	{}
-
 	// FAnimNextModuleInstanceComponent interface
 	virtual void OnTraitEvent(FAnimNextTraitEvent& Event) override;
 	virtual void OnEndExecution(float InDeltaTime) override;
diff --git a/Engine/Plugins/Experimental/AnimNextAnimGraph/Source/AnimNextAnimGraphUncookedOnly/Internal/Graph/AnimNextAnimationGraph_EditorData.h b/Engine/Plugins/Experimental/AnimNextAnimGraph/Source/AnimNextAnimGraphUncookedOnly/Internal/Graph/AnimNextAnimationGraph_EditorData.h
--- a/Engine/Plugins/Experimental/AnimNextAnimGraph/Source/AnimNextAnimGraphUncookedOnly/Internal/Graph/AnimNextAnimationGraph_EditorData.h
+++ b/Engine/Plugins/Experimental/AnimNextAnimGraph/Source/AnimNextAnimGraphUncookedOnly/Internal/Graph/AnimNextAnimationGraph_EditorData.h
@@ -64,9 +64,6 @@ public:
 	UAnimNextAnimationGraphEntry* AddAnimationGraph(FName InName, bool bSetupUndoRedo = true, bool bPrintPythonCommand = true);
 
 protected:
-	// IRigVMClientHost interface
-	virtual void RecompileVM() override;
-
 	// UAnimNextRigVMAssetEditorData interface
 	virtual TSubclassOf<URigVMController> GetControllerClass() const override { return UAnimNextController::StaticClass(); }
 	virtual UScriptStruct* GetExecuteContextStruct() const override { return FAnimNextExecuteContext::StaticStruct(); }
@@ -74,7 +71,10 @@ protected:
 	virtual bool CanAddNewEntry(TSubclassOf<UAnimNextRigVMAssetEntry> InClass) const override;
 	virtual TSubclassOf<UAssetUserData> GetAssetUserDataClass() const override;
 	virtual void InitializeAssetUserData() override;
-	virtual void GetProgrammaticGraphs(const FRigVMCompileSettings& InSettings, TArray<URigVMGraph*>& OutGraphs) override;
+	virtual void OnPreCompileAsset(FRigVMCompileSettings& InSettings) override;
+	virtual void OnPreCompileGetProgrammaticGraphs(const FRigVMCompileSettings& InSettings, TArray<URigVMGraph*>& OutGraphs) override;
+	virtual void OnPreCompileProcessGraphs(const FRigVMCompileSettings& InSettings, TArray<URigVMGraph*>& InOutGraphs) override;
+	virtual void OnPostCompileCleanup(const FRigVMCompileSettings& InSettings) override;
 
 	// Super overrides
 	virtual void GetAnimNextAssetRegistryTags(FAssetRegistryTagsContext& Context, FAnimNextAssetRegistryExports& OutExports) const override;
diff --git a/Engine/Plugins/Experimental/AnimNextAnimGraph/Source/AnimNextAnimGraphUncookedOnly/Private/Graph/AnimNextAnimationGraph_EditorData.cpp b/Engine/Plugins/Experimental/AnimNextAnimGraph/Source/AnimNextAnimGraphUncookedOnly/Private/Graph/AnimNextAnimationGraph_EditorData.cpp
--- a/Engine/Plugins/Experimental/AnimNextAnimGraph/Source/AnimNextAnimGraphUncookedOnly/Private/Graph/AnimNextAnimationGraph_EditorData.cpp
+++ b/Engine/Plugins/Experimental/AnimNextAnimGraph/Source/AnimNextAnimGraphUncookedOnly/Private/Graph/AnimNextAnimationGraph_EditorData.cpp
@@ -321,7 +321,7 @@ namespace UE::AnimNext::UncookedOnly::Private
 		}
 	}
 
-	FTraitGraph CollectGraphInfo(FRigVMCompileSettings& InSettings, const UAnimNextAnimationGraph* InAnimationGraph, const URigVMGraph* VMGraph, URigVMController* VMController)
+	FTraitGraph CollectGraphInfo(const FRigVMCompileSettings& InSettings, const UAnimNextAnimationGraph* InAnimationGraph, const URigVMGraph* VMGraph, URigVMController* VMController)
 	{
 		const TArray<URigVMNode*>& VMNodes = VMGraph->GetNodes();
 		URigVMUnitNode* VMRootNode = FindRootNode(VMNodes);
@@ -469,253 +469,171 @@ namespace UE::AnimNext::UncookedOnly::Private
 	}
 }
 
-void UAnimNextAnimationGraph_EditorData::RecompileVM()
+void UAnimNextAnimationGraph_EditorData::OnPreCompileAsset(FRigVMCompileSettings& InSettings)
 {
 	using namespace UE::AnimNext::UncookedOnly;
-	using namespace UE::AnimNext::UncookedOnly::Private;
 
-	if (bIsCompiling)
-	{
-		return;
-	}
-
-	TGuardValue<bool> CompilingGuard(bIsCompiling, true);
+	InSettings.ASTSettings.bSetupTraits = false; // disable the default implementation of decorators for now
 
 	UAnimNextAnimationGraph* AnimationGraph = FUtils::GetAsset<UAnimNextAnimationGraph>(this);
-	FScopedCompilerResults CompilerResults(AnimationGraph);
 
 	// Before we re-compile a graph, we need to release and live instances since we need the metadata we are about to replace
 	// to call trait destructors etc
 	AnimationGraph->FreezeGraphInstances();
 
-	CachedExports.Reset(); // asset variables and other tags will be updated at the end by AssetRegistry->AssetUpdateTags
+	AnimationGraph->EntryPoints.Empty();
+	AnimationGraph->ResolvedRootTraitHandles.Empty();
+	AnimationGraph->ResolvedEntryPoints.Empty();
+	AnimationGraph->ExecuteDefinition = FAnimNextGraphEvaluatorExecuteDefinition();
+	AnimationGraph->SharedDataBuffer.Empty();
+	AnimationGraph->GraphReferencedObjects.Empty();
+	AnimationGraph->GraphReferencedSoftObjects.Empty();
+	AnimationGraph->DefaultEntryPoint = NAME_None;
+}
 
-	bWarningsDuringCompilation = false;
-	bErrorsDuringCompilation = false;
+void UAnimNextAnimationGraph_EditorData::OnPreCompileProcessGraphs(const FRigVMCompileSettings& InSettings, TArray<URigVMGraph*>& InOutGraphs)
+{
+	using namespace UE::AnimNext::UncookedOnly;
+	using namespace UE::AnimNext::UncookedOnly::Private;
 
-	RigGraphDisplaySettings.MinMicroSeconds = RigGraphDisplaySettings.LastMinMicroSeconds = DBL_MAX;
-	RigGraphDisplaySettings.MaxMicroSeconds = RigGraphDisplaySettings.LastMaxMicroSeconds = (double)INDEX_NONE;
+	FRigVMClient* VMClient = GetRigVMClient();
+	UAnimNextAnimationGraph* AnimationGraph = FUtils::GetAsset<UAnimNextAnimationGraph>(this);
 
-	TArray<URigVMGraph*> ProgrammaticGraphs;
+	TArray<URigVMGraph*> AnimGraphs;
+	TArray<URigVMGraph*> NonAnimGraphs;
+	for(URigVMGraph* SourceGraph : InOutGraphs)
 	{
-		TGuardValue<bool> ReentrantGuardSelf(bSuspendModelNotificationsForSelf, true);
-		TGuardValue<bool> ReentrantGuardOthers(RigVMClient.bSuspendModelNotificationsForOthers, true);
-
-		VMCompileSettings.SetExecuteContextStruct(FAnimNextExecuteContext::StaticStruct());
-		FRigVMCompileSettings Settings = (bCompileInDebugMode) ? FRigVMCompileSettings::Fast(VMCompileSettings.GetExecuteContextStruct()) : VMCompileSettings;
-		Settings.SurpressInfoMessages = false;
-		Settings.ASTSettings.bSetupTraits = false; // disable the default implementation of decorators for now
-		Settings.ASTSettings.ReportDelegate.BindUObject(this, &UAnimNextAnimationGraph_EditorData::HandleReportFromCompiler);
-
-		FUtils::RecreateVM(AnimationGraph);
-
-		FUtils::CompileVariables(Settings, AnimationGraph);
-
-		AnimationGraph->VMRuntimeSettings = VMRuntimeSettings;
-		AnimationGraph->EntryPoints.Empty();
-		AnimationGraph->ResolvedRootTraitHandles.Empty();
-		AnimationGraph->ResolvedEntryPoints.Empty();
-		AnimationGraph->ExecuteDefinition = FAnimNextGraphEvaluatorExecuteDefinition();
-		AnimationGraph->SharedDataBuffer.Empty();
-		AnimationGraph->GraphReferencedObjects.Empty();
-		AnimationGraph->GraphReferencedSoftObjects.Empty();
-		AnimationGraph->DefaultEntryPoint = NAME_None;
-
-		FRigVMClient* VMClient = GetRigVMClient();
-
-		GetProgrammaticGraphs(Settings, ProgrammaticGraphs);
-		for(URigVMGraph* ProgrammaticGraph : ProgrammaticGraphs)
+		// We use a temporary graph models to build our final graphs that we'll compile
+		if(SourceGraph->GetSchemaClass() == UAnimNextAnimationGraphSchema::StaticClass())
 		{
-			check(ProgrammaticGraph != nullptr);
-		}
-	
-		TArray<URigVMGraph*> AllGraphs = VMClient->GetAllModels(false, false);
-		AllGraphs.Append(ProgrammaticGraphs);
+			TMap<UObject*, UObject*> CreatedObjects;
+			FObjectDuplicationParameters Parameters = InitStaticDuplicateObjectParams(SourceGraph, this, NAME_None, RF_Transient);
+			Parameters.CreatedObjects = &CreatedObjects;
+			URigVMGraph* TempGraph = CastChecked<URigVMGraph>(StaticDuplicateObjectEx(Parameters));
+			TempGraph->SetExternalPackage(nullptr);
+			for(URigVMNode* SourceNode : SourceGraph->GetNodes())
+			{
+				FScopedCompilerResults::GetLog().NotifyIntermediateObjectCreation(CreatedObjects.FindChecked(SourceNode), SourceNode);
+			}
 
-		if(AllGraphs.Num() == 0)
-		{
-			return;
+			UAnimNextController* TempController = CastChecked<UAnimNextController>(VMClient->GetOrCreateController(TempGraph));
+			TempGraph->SetFlags(RF_Transient);
+			AnimGraphs.Add(TempGraph);
 		}
-
-		TArray<URigVMGraph*> AnimGraphs;
-		TArray<URigVMGraph*> NonAnimGraphs;
-		for(URigVMGraph* SourceGraph : AllGraphs)
+		else
 		{
-			// We use a temporary graph models to build our final graphs that we'll compile
-			if(SourceGraph->GetSchemaClass() == UAnimNextAnimationGraphSchema::StaticClass())
-			{
-				TMap<UObject*, UObject*> CreatedObjects;
-				FObjectDuplicationParameters Parameters = InitStaticDuplicateObjectParams(SourceGraph, this, NAME_None, RF_Transient);
-				Parameters.CreatedObjects = &CreatedObjects;
-				URigVMGraph* TempGraph = CastChecked<URigVMGraph>(StaticDuplicateObjectEx(Parameters));
-				TempGraph->SetExternalPackage(nullptr);
-				for(URigVMNode* SourceNode : SourceGraph->GetNodes())
-				{
-					FScopedCompilerResults::GetLog().NotifyIntermediateObjectCreation(CreatedObjects.FindChecked(SourceNode), SourceNode);
-				}
-
-				UAnimNextController* TempController = CastChecked<UAnimNextController>(VMClient->GetOrCreateController(TempGraph));
-				TempGraph->SetFlags(RF_Transient);
-				AnimGraphs.Add(TempGraph);
-			}
-			else
-			{
-				NonAnimGraphs.Add(SourceGraph);
-			}
+			NonAnimGraphs.Add(SourceGraph);
 		}
+	}
 
-		if(AnimGraphs.Num() > 0)
-		{
-			UAnimNextController* TempController = CastChecked<UAnimNextController>(VMClient->GetOrCreateController(AnimGraphs[0]));
-
-			UE::AnimNext::FTraitWriter TraitWriter;
+	if(AnimGraphs.Num() > 0)
+	{
+		UAnimNextController* TempController = CastChecked<UAnimNextController>(VMClient->GetOrCreateController(AnimGraphs[0]));
 
-			FRigVMPinInfoArray LatentPins;
-			TMap<FName, URigVMPin*> LatentPinMapping;
-			TArray<FTraitGraph> TraitGraphs;
+		UE::AnimNext::FTraitWriter TraitWriter;
 
-			// Build entry points and extract their required latent pins
-			for(const URigVMGraph* AnimGraph : AnimGraphs)
-			{
-				// Gather our trait stacks
-				FTraitGraph& TraitGraph = TraitGraphs.Add_GetRef(CollectGraphInfo(Settings, AnimationGraph, AnimGraph, TempController->GetControllerForGraph(AnimGraph)));
-				check(!TraitGraph.TraitStackNodes.IsEmpty());
+		FRigVMPinInfoArray LatentPins;
+		TMap<FName, URigVMPin*> LatentPinMapping;
+		TArray<FTraitGraph> TraitGraphs;
 
-				FAnimNextGraphEntryPoint& EntryPoint = AnimationGraph->EntryPoints.AddDefaulted_GetRef();
-				EntryPoint.EntryPointName = TraitGraph.EntryPoint;
+		// Build entry points and extract their required latent pins
+		for(const URigVMGraph* AnimGraph : AnimGraphs)
+		{
+			// Gather our trait stacks
+			FTraitGraph& TraitGraph = TraitGraphs.Add_GetRef(CollectGraphInfo(InSettings, AnimationGraph, AnimGraph, TempController->GetControllerForGraph(AnimGraph)));
+			check(!TraitGraph.TraitStackNodes.IsEmpty());
 
-				// Extract latent pins for this graph
-				CollectLatentPins(TraitGraph.TraitStackNodes, LatentPins, LatentPinMapping);
+			FAnimNextGraphEntryPoint& EntryPoint = AnimationGraph->EntryPoints.AddDefaulted_GetRef();
+			EntryPoint.EntryPointName = TraitGraph.EntryPoint;
 
-				// Iterate over every trait stack and register our node templates
-				for (FTraitStackMapping& NodeMapping : TraitGraph.TraitStackNodes)
-				{
-					NodeMapping.TraitStackNodeHandle = RegisterTraitNodeTemplate(TraitWriter, NodeMapping.DecoratorStackNode);
-				}
+			// Extract latent pins for this graph
+			CollectLatentPins(TraitGraph.TraitStackNodes, LatentPins, LatentPinMapping);
 
-				// Find our root node handle, if we have any stack nodes, the first one is our root stack
-				if (TraitGraph.TraitStackNodes.Num() != 0)
-				{
-					EntryPoint.RootTraitHandle = FAnimNextEntryPointHandle(TraitGraph.TraitStackNodes[0].TraitStackNodeHandle);
-				}
-			}
-
-			// Set default entry point
-			if(AnimationGraph->EntryPoints.Num() > 0)
+			// Iterate over every trait stack and register our node templates
+			for (FTraitStackMapping& NodeMapping : TraitGraph.TraitStackNodes)
 			{
-				AnimationGraph->DefaultEntryPoint = AnimationGraph->EntryPoints[0].EntryPointName;
+				NodeMapping.TraitStackNodeHandle = RegisterTraitNodeTemplate(TraitWriter, NodeMapping.DecoratorStackNode);
 			}
 
-			// Remove our old root nodes
-			for (FTraitGraph& TraitGraph : TraitGraphs)
+			// Find our root node handle, if we have any stack nodes, the first one is our root stack
+			if (TraitGraph.TraitStackNodes.Num() != 0)
 			{
-				URigVMController* GraphController = TempController->GetControllerForGraph(TraitGraph.RootNode->GetGraph());
-				GraphController->RemoveNode(TraitGraph.RootNode, false, false);
+				EntryPoint.RootTraitHandle = FAnimNextEntryPointHandle(TraitGraph.TraitStackNodes[0].TraitStackNodeHandle);
 			}
+		}
 
-			if(LatentPins.Num() > 0)
-			{
-				// We need a unique method name to match our unique argument list
-				AnimationGraph->ExecuteDefinition = GetGraphEvaluatorExecuteMethod(LatentPins);
+		// Set default entry point
+		if(AnimationGraph->EntryPoints.Num() > 0)
+		{
+			AnimationGraph->DefaultEntryPoint = AnimationGraph->EntryPoints[0].EntryPointName;
+		}
 
-				// Add our runtime shim root node
-				URigVMUnitNode* TempShimRootNode = TempController->AddUnitNode(FRigUnit_AnimNextShimRoot::StaticStruct(), FRigUnit_AnimNextShimRoot::EventName, FVector2D::ZeroVector, FString(), false);
-				URigVMUnitNode* GraphEvaluatorNode = TempController->AddUnitNodeWithPins(FRigUnit_AnimNextGraphEvaluator::StaticStruct(), LatentPins, *AnimationGraph->ExecuteDefinition.MethodName, FVector2D::ZeroVector, FString(), false);
+		// Remove our old root nodes
+		for (FTraitGraph& TraitGraph : TraitGraphs)
+		{
+			URigVMController* GraphController = TempController->GetControllerForGraph(TraitGraph.RootNode->GetGraph());
+			GraphController->RemoveNode(TraitGraph.RootNode, false, false);
+		}
 
-				// Link our shim and evaluator nodes together using the execution context
-				TempController->AddLink(
-					TempShimRootNode->FindPin(GET_MEMBER_NAME_STRING_CHECKED(FRigUnit_AnimNextShimRoot, ExecuteContext)),
-					GraphEvaluatorNode->FindPin(GET_MEMBER_NAME_STRING_CHECKED(FRigUnit_AnimNextGraphEvaluator, ExecuteContext)),
-					false);
+		if(LatentPins.Num() > 0)
+		{
+			// We need a unique method name to match our unique argument list
+			AnimationGraph->ExecuteDefinition = GetGraphEvaluatorExecuteMethod(LatentPins);
 
-				// Link our latent pins
-				for (const FRigVMPinInfo& LatentPin : LatentPins)
-				{
-					TempController->AddLink(
-						LatentPinMapping[LatentPin.Name],
-						GraphEvaluatorNode->FindPin(LatentPin.Name.ToString()),
-						false);
-				}
-			}
+			// Add our runtime shim root node
+			URigVMUnitNode* TempShimRootNode = TempController->AddUnitNode(FRigUnit_AnimNextShimRoot::StaticStruct(), FRigUnit_AnimNextShimRoot::EventName, FVector2D::ZeroVector, FString(), false);
+			URigVMUnitNode* GraphEvaluatorNode = TempController->AddUnitNodeWithPins(FRigUnit_AnimNextGraphEvaluator::StaticStruct(), LatentPins, *AnimationGraph->ExecuteDefinition.MethodName, FVector2D::ZeroVector, FString(), false);
 
-			// Write our node shared data
-			TraitWriter.BeginNodeWriting();
+			// Link our shim and evaluator nodes together using the execution context
+			TempController->AddLink(
+				TempShimRootNode->FindPin(GET_MEMBER_NAME_STRING_CHECKED(FRigUnit_AnimNextShimRoot, ExecuteContext)),
+				GraphEvaluatorNode->FindPin(GET_MEMBER_NAME_STRING_CHECKED(FRigUnit_AnimNextGraphEvaluator, ExecuteContext)),
+				false);
 
-			for(FTraitGraph& TraitGraph : TraitGraphs)
+			// Link our latent pins
+			for (const FRigVMPinInfo& LatentPin : LatentPins)
 			{
-				for (const FTraitStackMapping& NodeMapping : TraitGraph.TraitStackNodes)
-				{
-					WriteTraitProperties(TraitWriter, NodeMapping, TraitGraph.TraitStackNodes);
-				}
+				TempController->AddLink(
+					LatentPinMapping[LatentPin.Name],
+					GraphEvaluatorNode->FindPin(LatentPin.Name.ToString()),
+					false);
 			}
-
-			TraitWriter.EndNodeWriting();
-
-			// Cache our compiled metadata
-			AnimationGraph->SharedDataArchiveBuffer = TraitWriter.GetGraphSharedData();
-			AnimationGraph->GraphReferencedObjects = TraitWriter.GetGraphReferencedObjects();
-			AnimationGraph->GraphReferencedSoftObjects = TraitWriter.GetGraphReferencedSoftObjects();
-
-			// Populate our runtime metadata
-			AnimationGraph->LoadFromArchiveBuffer(AnimationGraph->SharedDataArchiveBuffer);
 		}
 
-		TArray<URigVMGraph*> GraphsToCompile = MoveTemp(AnimGraphs);
-		GraphsToCompile.Append(NonAnimGraphs);
+		// Write our node shared data
+		TraitWriter.BeginNodeWriting();
 
-		UAnimNextController* Controller = CastChecked<UAnimNextController>(VMClient->GetOrCreateController(GraphsToCompile[0]));
-		URigVMCompiler* Compiler = URigVMCompiler::StaticClass()->GetDefaultObject<URigVMCompiler>();
-		Compiler->Compile(Settings, GraphsToCompile, Controller, AnimationGraph->VM, AnimationGraph->ExtendedExecuteContext, AnimationGraph->GetExternalVariables(), &PinToOperandMap);
-
-		// Initialize right away, in packaged builds we initialize during PostLoad
-		AnimationGraph->VM->Initialize(AnimationGraph->ExtendedExecuteContext);
-		AnimationGraph->GenerateUserDefinedDependenciesData(AnimationGraph->ExtendedExecuteContext);
-
-		// Notable difference with vanilla RigVM host behavior - we init the VM here at the moment as we only have one 'instance'
-		AnimationGraph->InitializeVM(FRigUnit_AnimNextBeginExecution::EventName);
-
-		if (bErrorsDuringCompilation)
+		for(FTraitGraph& TraitGraph : TraitGraphs)
 		{
-			if(Settings.SurpressErrors)
+			for (const FTraitStackMapping& NodeMapping : TraitGraph.TraitStackNodes)
 			{
-				Settings.Reportf(EMessageSeverity::Info, AnimationGraph, TEXT("Compilation Errors may be suppressed for AnimNext asset: %s. See VM Compile Settings for more Details"), *AnimationGraph->GetName());
+				WriteTraitProperties(TraitWriter, NodeMapping, TraitGraph.TraitStackNodes);
 			}
 		}
 
-		bVMRecompilationRequired = false;
-		if(AnimationGraph->VM)
-		{
-			RigVMCompiledEvent.Broadcast(AnimationGraph, AnimationGraph->VM, AnimationGraph->ExtendedExecuteContext);
-		}
+		TraitWriter.EndNodeWriting();
 
-		// Remove the temp anim graphs
-		RemoveProgrammaticGraphs(AnimGraphs);
+		// Cache our compiled metadata
+		AnimationGraph->SharedDataArchiveBuffer = TraitWriter.GetGraphSharedData();
+		AnimationGraph->GraphReferencedObjects = TraitWriter.GetGraphReferencedObjects();
+		AnimationGraph->GraphReferencedSoftObjects = TraitWriter.GetGraphReferencedSoftObjects();
 
-		// Now that the graph has been re-compiled, re-allocate the previous live instances
-		AnimationGraph->ThawGraphInstances();
+		// Populate our runtime metadata
+		AnimationGraph->LoadFromArchiveBuffer(AnimationGraph->SharedDataArchiveBuffer);
+	}
 
-#if WITH_EDITOR
-		// Display programmatic graphs
-		if(CVarDumpProgrammaticGraphs.GetValueOnGameThread())
-		{
-			FUtils::OpenProgrammaticGraphs(this, ProgrammaticGraphs);
-		}
-		else
-#endif
-		{
-			RemoveProgrammaticGraphs(ProgrammaticGraphs);
-		}
+	InOutGraphs = MoveTemp(AnimGraphs);
+	InOutGraphs.Append(NonAnimGraphs);
+}
 
-#if WITH_EDITOR
-		//	RefreshBreakpoints(EditorData);
-#endif
+void UAnimNextAnimationGraph_EditorData::OnPostCompileCleanup(const FRigVMCompileSettings& InSettings)
+{
+	using namespace UE::AnimNext::UncookedOnly;
 
-		// Refresh CachedExports, also updates variables at UAnimNextRigVMAssetEditorData::GetAssetRegistryTags
-		if (IAssetRegistry* AssetRegistry = IAssetRegistry::Get())
-		{
-			AssetRegistry->AssetUpdateTags(GetTypedOuter<UAnimNextAnimationGraph>(), EAssetRegistryTagsCaller::Fast);
-		}
-	}
+	UAnimNextAnimationGraph* AnimationGraph = FUtils::GetAsset<UAnimNextAnimationGraph>(this);
+
+	// Now that the graph has been re-compiled, re-allocate the previous live instances
+	AnimationGraph->ThawGraphInstances();
 }
 
 TConstArrayView<TSubclassOf<UAnimNextRigVMAssetEntry>> UAnimNextAnimationGraph_EditorData::GetEntryClasses() const
@@ -854,11 +772,11 @@ void UAnimNextAnimationGraph_EditorData::InitializeAssetUserData()
 	Super::InitializeAssetUserData();
 }
 
-void UAnimNextAnimationGraph_EditorData::GetProgrammaticGraphs(const FRigVMCompileSettings& InSettings, TArray<URigVMGraph*>& OutGraphs)
+void UAnimNextAnimationGraph_EditorData::OnPreCompileGetProgrammaticGraphs(const FRigVMCompileSettings& InSettings, TArray<URigVMGraph*>& OutGraphs)
 {
 	using namespace UE::AnimNext::UncookedOnly;
 
-	Super::GetProgrammaticGraphs(InSettings, OutGraphs);
+	Super::OnPreCompileGetProgrammaticGraphs(InSettings, OutGraphs);
 
 	// Gather all 'call function' traits and create shim-calls for them.
 	// For the compiler to pick them up if they are not public we need a calling reference to the function from a graph
diff --git a/Engine/Plugins/Experimental/AnimNextStateTree/Source/AnimNextStateTreeUncookedOnly/Internal/AnimNextStateTree_EditorData.cpp b/Engine/Plugins/Experimental/AnimNextStateTree/Source/AnimNextStateTreeUncookedOnly/Internal/AnimNextStateTree_EditorData.cpp
--- a/Engine/Plugins/Experimental/AnimNextStateTree/Source/AnimNextStateTreeUncookedOnly/Internal/AnimNextStateTree_EditorData.cpp
+++ b/Engine/Plugins/Experimental/AnimNextStateTree/Source/AnimNextStateTreeUncookedOnly/Internal/AnimNextStateTree_EditorData.cpp
@@ -55,7 +55,7 @@ TConstArrayView<TSubclassOf<UAnimNextRigVMAssetEntry>> UAnimNextStateTree_Editor
 	return Classes;
 }
 
-void UAnimNextStateTree_EditorData::GetProgrammaticVariables(const FRigVMCompileSettings& InSettings, TArray<FAnimNextProgrammaticVariable>& OutVariables)
+void UAnimNextStateTree_EditorData::OnPreCompileGetProgrammaticVariables(const FRigVMCompileSettings& InSettings, TArray<FAnimNextProgrammaticVariable>& OutVariables)
 {
 	if (UAnimNextStateTree* AnimStateTree = UE::AnimNext::UncookedOnly::FUtils::GetAsset<UAnimNextStateTree>(this))
 	{
@@ -100,7 +100,7 @@ void UAnimNextStateTree_EditorData::GetProgrammaticVariables(const FRigVMCompile
 	}
 }
 
-void UAnimNextStateTree_EditorData::GetProgrammaticGraphs(const FRigVMCompileSettings& InSettings, TArray<URigVMGraph*>& OutGraphs)
+void UAnimNextStateTree_EditorData::OnPreCompileGetProgrammaticGraphs(const FRigVMCompileSettings& InSettings, TArray<URigVMGraph*>& OutGraphs)
 {
 	if (UAnimNextStateTree* AnimStateTree = UE::AnimNext::UncookedOnly::FUtils::GetAsset<UAnimNextStateTree>(this))
 	{
diff --git a/Engine/Plugins/Experimental/AnimNextStateTree/Source/AnimNextStateTreeUncookedOnly/Internal/AnimNextStateTree_EditorData.h b/Engine/Plugins/Experimental/AnimNextStateTree/Source/AnimNextStateTreeUncookedOnly/Internal/AnimNextStateTree_EditorData.h
--- a/Engine/Plugins/Experimental/AnimNextStateTree/Source/AnimNextStateTreeUncookedOnly/Internal/AnimNextStateTree_EditorData.h
+++ b/Engine/Plugins/Experimental/AnimNextStateTree/Source/AnimNextStateTreeUncookedOnly/Internal/AnimNextStateTree_EditorData.h
@@ -22,8 +22,8 @@ protected:
 	virtual TConstArrayView<TSubclassOf<UAnimNextRigVMAssetEntry>> GetEntryClasses() const override;
 
 	// Allows this asset to generate variables to be injected at compilation time, separate method to allow programmatic graphs to use these vars
-	virtual void GetProgrammaticVariables(const FRigVMCompileSettings& InSettings, TArray<FAnimNextProgrammaticVariable>& OutVariables) override;
+	virtual void OnPreCompileGetProgrammaticVariables(const FRigVMCompileSettings& InSettings, TArray<FAnimNextProgrammaticVariable>& OutVariables) override;
 
 	// Allows this asset to generate graphs to be injected at compilation time
-	virtual void GetProgrammaticGraphs(const FRigVMCompileSettings& InSettings, TArray<URigVMGraph*>& OutGraphs) override;
+	virtual void OnPreCompileGetProgrammaticGraphs(const FRigVMCompileSettings& InSettings, TArray<URigVMGraph*>& OutGraphs) override;
 };
\ No newline at end of file
