Added a BlendSpace Trait to AnimNext.  This initial version does not have synchronization.

diff --git a/Engine/Plugins/Experimental/AnimNext/Source/AnimNext/Private/AnimNextModuleImpl.cpp b/Engine/Plugins/Experimental/AnimNext/Source/AnimNext/Private/AnimNextModuleImpl.cpp
--- a/Engine/Plugins/Experimental/AnimNext/Source/AnimNext/Private/AnimNextModuleImpl.cpp
+++ b/Engine/Plugins/Experimental/AnimNext/Source/AnimNext/Private/AnimNextModuleImpl.cpp
@@ -7,6 +7,7 @@
 #include "RigVMRuntimeDataRegistry.h"
 #include "UniversalObjectLocator.h"
 #include "Animation/AnimSequence.h"
+#include "Animation/BlendSpace.h"
 #include "Animation/BlendProfile.h"
 #include "Component/AnimNextComponent.h"
 #include "Curves/CurveFloat.h"
@@ -36,6 +37,7 @@ namespace UE::AnimNext
 		static TPair<UClass*, FRigVMRegistry::ERegisterObjectOperation> const AllowedObjectTypes[] =
 		{
 			{ UAnimSequence::StaticClass(), FRigVMRegistry::ERegisterObjectOperation::Class },
+			{ UBlendSpace::StaticClass(), FRigVMRegistry::ERegisterObjectOperation::Class },
 			{ UScriptStruct::StaticClass(), FRigVMRegistry::ERegisterObjectOperation::Class },
 			{ UBlendProfile::StaticClass(), FRigVMRegistry::ERegisterObjectOperation::Class },
 			{ UCurveFloat::StaticClass(), FRigVMRegistry::ERegisterObjectOperation::Class },
diff --git a/Engine/Plugins/Experimental/AnimNextAnimGraph/Source/AnimNextAnimGraph/Private/Traits/BlendSpacePlayer.cpp b/Engine/Plugins/Experimental/AnimNextAnimGraph/Source/AnimNextAnimGraph/Private/Traits/BlendSpacePlayer.cpp
new file mode 100644
--- /dev/null
+++ b/Engine/Plugins/Experimental/AnimNextAnimGraph/Source/AnimNextAnimGraph/Private/Traits/BlendSpacePlayer.cpp
@@ -0,0 +1,316 @@
+// Copyright Epic Games, Inc. All Rights Reserved.
+
+#include "Traits/BlendSpacePlayer.h"
+#include "AnimationRuntime.h"
+#include "TraitCore/ExecutionContext.h"
+#include "EvaluationVM/Tasks/BlendKeyframes.h"
+#include "Graph/AnimNextAnimationGraph.h"
+#include "AnimNextAnimGraphSettings.h"
+#include "AnimNextDataInterfacePayload.h"
+#include "DataInterface/DataInterfaceStructAdapter.h"
+#include "GraphInterfaces/AnimNextNativeDataInterface_AnimSequencePlayer.h"
+#include "Graph/AnimNextGraphInstance.h"
+#include "TraitCore/NodeInstance.h"
+#include "TraitInterfaces/IGraphFactory.h"
+#include "TraitInterfaces/IPayloadFactory.h"
+
+
+namespace UE::AnimNext
+{
+	AUTO_REGISTER_ANIM_TRAIT(FBlendSpacePlayerTrait)
+
+	// Trait implementation boilerplate
+	#define TRAIT_INTERFACE_ENUMERATOR(GeneratorMacro) \
+		GeneratorMacro(IContinuousBlend) \
+		GeneratorMacro(IEvaluate) \
+		GeneratorMacro(IHierarchy) \
+		GeneratorMacro(ITimeline) \
+		GeneratorMacro(IUpdate) \
+		GeneratorMacro(IUpdateTraversal) \
+		GeneratorMacro(IGarbageCollection) \
+
+	GENERATE_ANIM_TRAIT_IMPLEMENTATION(FBlendSpacePlayerTrait, TRAIT_INTERFACE_ENUMERATOR, NULL_ANIM_TRAIT_INTERFACE_ENUMERATOR, NULL_ANIM_TRAIT_EVENT_ENUMERATOR)
+	#undef TRAIT_INTERFACE_ENUMERATOR
+
+	void FBlendSpacePlayerTrait::FInstanceData::Construct(const FExecutionContext& Context, const FTraitBinding& Binding)
+	{
+		FTrait::FInstanceData::Construct(Context, Binding);
+
+		IGarbageCollection::RegisterWithGC(Context, Binding);
+
+	}
+
+	void FBlendSpacePlayerTrait::FInstanceData::Destruct(const FExecutionContext& Context, const FTraitBinding& Binding)
+	{
+		FTrait::FInstanceData::Destruct(Context, Binding);
+
+		IGarbageCollection::UnregisterWithGC(Context, Binding);
+	}
+
+	void FBlendSpacePlayerTrait::PostEvaluate(FEvaluateTraversalContext& Context, const TTraitBinding<IEvaluate>& Binding) const
+	{
+		const FInstanceData* InstanceData = Binding.GetInstanceData<FInstanceData>();
+
+		if (InstanceData->NumChildren == 3)
+		{
+			// We have three children, interpolate them
+
+			TTraitBinding<IContinuousBlend> ContinuousBlendTrait;
+			Binding.GetStackInterface(ContinuousBlendTrait);
+
+			const float BlendWeight2 = ContinuousBlendTrait.GetBlendWeight(Context, 2);
+			const float BlendWeight1 = ContinuousBlendTrait.GetBlendWeight(Context, 1);
+			const float BlendWeight0 = ContinuousBlendTrait.GetBlendWeight(Context, 0);
+			
+			if (FAnimWeight::IsRelevant(BlendWeight2 + BlendWeight1))
+			{
+				const float FirstBlendWeight = BlendWeight2 / (BlendWeight2 + BlendWeight1);
+				Context.AppendTask(FAnimNextBlendTwoKeyframesTask::Make(FirstBlendWeight));
+			}
+			else
+			{
+				// we have no weight but still need to pop, so still make a task that will get ignored in the next step 
+				Context.AppendTask(FAnimNextBlendTwoKeyframesTask::Make(0.0));
+			}
+			Context.AppendTask(FAnimNextBlendTwoKeyframesTask::Make(BlendWeight0));
+		}
+		else if (InstanceData->NumChildren == 2)
+		{
+			// We have two children, interpolate them
+
+			TTraitBinding<IContinuousBlend> ContinuousBlendTrait;
+			Binding.GetStackInterface(ContinuousBlendTrait);
+
+			const float BlendWeight = ContinuousBlendTrait.GetBlendWeight(Context, 1);
+			Context.AppendTask(FAnimNextBlendTwoKeyframesTask::Make(BlendWeight));
+		}
+		else
+		{
+			// We have only one child that is active, do nothing
+		}
+	}
+
+	void FBlendSpacePlayerTrait::OnBecomeRelevant(FUpdateTraversalContext& Context, const TTraitBinding<IUpdate>& Binding, const FTraitUpdateState& TraitState) const
+	{
+		const FSharedData* SharedData = Binding.GetSharedData<FSharedData>();
+		FInstanceData* InstanceData = Binding.GetInstanceData<FInstanceData>();
+
+		// Cache the blend space we'll play during construction, we don't allow it to change afterwards
+		InstanceData->BlendSpace = SharedData->GetBlendSpace(Binding);
+		if (InstanceData->BlendSpace)
+		{
+			InstanceData->BlendSpace->InitializeFilter(&InstanceData->BlendFilter);
+
+			FVector BlendParameters(SharedData->XAxisSamplePoint, SharedData->YAxisSamplePoint, 0.0f);
+
+			// reset all the blend samples for the first frame
+			if (InstanceData->BlendSpace->GetSamplesFromBlendInput(BlendParameters, InstanceData->BlendSamplesData, InstanceData->CachedTriangulationIndex, true))
+			{
+				// Get starting time for all samples.
+				InstanceData->BlendSpace->ResetBlendSamples(InstanceData->BlendSamplesData, 0.0f, true, true);
+			}
+
+			// create graphs for all the blend samples
+			const TArray<FBlendSample>& BlendSamples = InstanceData->BlendSpace->GetBlendSamples();
+			if (InstanceData->SampleGraphs.Num() != BlendSamples.Num())
+			{
+				InstanceData->SampleGraphs.SetNum(BlendSamples.Num(), EAllowShrinking::No); 
+				for (int32 CurSampleIndex = 0; CurSampleIndex < BlendSamples.Num(); CurSampleIndex++)
+				{
+					const UAnimSequence* CurSequence = BlendSamples[CurSampleIndex].Animation;
+					const UAnimNextAnimationGraph* AnimationGraph = IGraphFactory::GetGraphFromObjectWithFallback(Context, Binding, CurSequence);
+					if (AnimationGraph != nullptr)
+					{
+						FSampleGraphState& TargetGraphState = InstanceData->SampleGraphs[CurSampleIndex];
+						// Create a default 'native data interface' payload via which we expect to communicate with the graph 
+						FInstancedStruct NewPayload;
+						NewPayload.InitializeAs<FAnimNextNativeDataInterface_AnimSequencePlayer>();
+						FAnimNextNativeDataInterface_AnimSequencePlayer& NewNativeInterface = NewPayload.GetMutable<FAnimNextNativeDataInterface_AnimSequencePlayer>();
+						NewNativeInterface.AnimSequence = CurSequence;
+						NewNativeInterface.PlayRate = 1.0f;
+						NewNativeInterface.StartPosition = 0.0f;
+						NewNativeInterface.Loop = true;
+
+						TargetGraphState.GraphPayload.AddNative(MoveTemp(NewPayload));
+
+						IPayloadFactory::GetPayloadForGraphWithFallback(Context, Binding, CurSequence, AnimationGraph, TargetGraphState.GraphPayload);
+
+						const FName EntryPoint = NAME_None;
+						TargetGraphState.Instance = AnimationGraph->AllocateInstance(nullptr, &Binding.GetTraitPtr().GetNodeInstance()->GetOwner(), EntryPoint);
+						TargetGraphState.ChildPtr = TargetGraphState.Instance->GetGraphRootPtr();
+
+						TargetGraphState.Instance->BindPublicVariables(TargetGraphState.GraphPayload.Get());
+					}
+				}
+			}
+		}
+	}
+
+	uint32 FBlendSpacePlayerTrait::GetNumChildren(const FExecutionContext& Context, const TTraitBinding<IHierarchy>& Binding) const
+	{
+		const FInstanceData* InstanceData = Binding.GetInstanceData<FInstanceData>();
+		return InstanceData->NumChildren; 
+
+	}
+
+	void FBlendSpacePlayerTrait::GetChildren(const FExecutionContext& Context, const TTraitBinding<IHierarchy>& Binding, FChildrenArray& Children) const
+	{
+		const FInstanceData* InstanceData = Binding.GetInstanceData<FInstanceData>();
+		if (InstanceData->NumChildren > 0)
+		{
+			for (const FSampleGraphState& SampleState : InstanceData->SampleGraphs)
+			{
+				if (SampleState.bSampledThisFrame)
+				{
+					Children.Add(SampleState.ChildPtr);
+				}
+			}
+		}
+		check(Children.Num() == InstanceData->NumChildren); 
+	}
+
+	void FBlendSpacePlayerTrait::QueueChildrenForTraversal(FUpdateTraversalContext& Context, const TTraitBinding<IUpdateTraversal>& Binding, const FTraitUpdateState& TraitState, FUpdateTraversalQueue& TraversalQueue) const
+	{
+		const FInstanceData* InstanceData = Binding.GetInstanceData<FInstanceData>();
+
+		TTraitBinding<IContinuousBlend> ContinuousBlendTrait;
+		Binding.GetStackInterface(ContinuousBlendTrait);
+
+		if (InstanceData->NumChildren > 0)
+		{
+			for (const FSampleGraphState& SampleState : InstanceData->SampleGraphs)
+			{
+				if (SampleState.bSampledThisFrame)
+				{
+					bool bNewlyRelevant = SampleState.bSampledThisFrame && !SampleState.bSampledLastFrame;
+					TraversalQueue.Push(SampleState.ChildPtr, TraitState.WithWeight(SampleState.Weight).AsNewlyRelevant(bNewlyRelevant));
+				}
+			}
+		}
+	}
+
+	float FBlendSpacePlayerTrait::GetBlendWeight(const FExecutionContext& Context, const TTraitBinding<IContinuousBlend>& Binding, int32 ChildIndex) const
+	{
+		const FInstanceData* InstanceData = Binding.GetInstanceData<FInstanceData>();
+		if (ChildIndex < InstanceData->NumChildren)
+		{
+			int CurChildIndex = 0; 
+			for (const FSampleGraphState& SampleState : InstanceData->SampleGraphs)
+			{
+				if (SampleState.bSampledThisFrame)
+				{
+					if (ChildIndex == CurChildIndex)
+					{
+						return SampleState.Weight;
+					}
+					CurChildIndex++;
+				}
+			}
+		}
+		return -1.0f;
+	}
+
+	void FBlendSpacePlayerTrait::PreUpdate(FUpdateTraversalContext& Context, const TTraitBinding<IUpdate>& Binding, const FTraitUpdateState& TraitState) const
+	{
+		IUpdate::PreUpdate(Context, Binding, TraitState);
+
+		const FSharedData* SharedData = Binding.GetSharedData<FSharedData>();
+		check(SharedData);
+
+		FInstanceData* InstanceData = Binding.GetInstanceData<FInstanceData>();
+		check(InstanceData);
+
+		const float DeltaTime = TraitState.GetDeltaTime(); 
+
+		// input data 
+		FVector BlendParameters(SharedData->XAxisSamplePoint, SharedData->YAxisSamplePoint, 0.0f);
+		for (FSampleGraphState& SampleState : InstanceData->SampleGraphs)
+		{
+			SampleState.bSampledLastFrame = SampleState.bSampledThisFrame;
+			SampleState.bSampledThisFrame = false; 
+		}
+
+		InstanceData->NumChildren = 0;
+
+		if (const UBlendSpace* BlendSpacePtr = InstanceData->BlendSpace.Get())
+		{
+			if (BlendSpacePtr != nullptr && BlendSpacePtr->GetSkeleton() != nullptr)
+			{
+				const FVector FilteredBlendParams = BlendSpacePtr->FilterInput(&InstanceData->BlendFilter, BlendParameters, DeltaTime);
+				if (BlendSpacePtr->UpdateBlendSamples(FilteredBlendParams, DeltaTime, InstanceData->BlendSamplesData, InstanceData->CachedTriangulationIndex))
+				{
+					// using the blend samples set the weights and activity of the 
+					for (const FBlendSampleData& CurSample : InstanceData->BlendSamplesData)
+					{
+						if (CurSample.SampleDataIndex < InstanceData->SampleGraphs.Num())
+						{
+							FSampleGraphState& SampleState = InstanceData->SampleGraphs[CurSample.SampleDataIndex];
+							SampleState.Weight = CurSample.TotalWeight;
+							SampleState.bSampledThisFrame = true;
+							InstanceData->NumChildren++;
+						}
+					}
+				}
+			}
+		}
+	}
+
+	FTimelineState FBlendSpacePlayerTrait::GetState(FExecutionContext& Context, const TTraitBinding<ITimeline>& Binding) const
+	{
+		const FInstanceData* InstanceData = Binding.GetInstanceData<FInstanceData>();
+
+		// get the highest weighted timeline for looping, but average the duration, position ratios, and play rate based on weighting 
+		if (InstanceData->NumChildren > 0)
+		{
+			int32 HighestWeightIndex = 0;
+			float HighestWeight = 0.0f; 
+			bool bLooping = false; 
+			double WeightedDuration = 0.0; 
+			double WeightedPlayRate = 0.0; 
+			double WeightedPositionRatio = 0.0; 
+
+			// using the blend samples set the weights and activity 
+			for (const FBlendSampleData& CurSample : InstanceData->BlendSamplesData)
+			{	
+				const int32 CurSampleDataIndex = CurSample.SampleDataIndex;
+				if (CurSampleDataIndex < InstanceData->SampleGraphs.Num())
+				{
+					const FSampleGraphState& SampleState = InstanceData->SampleGraphs[CurSampleDataIndex];
+					FTraitStackBinding ChildTraitStack;
+					ensure(Context.GetStack(SampleState.ChildPtr, ChildTraitStack));
+
+					UE::AnimNext::TTraitBinding<UE::AnimNext::ITimeline> Timeline;
+					if (ChildTraitStack.GetInterface<UE::AnimNext::ITimeline>(Timeline))
+					{
+						FTimelineState CurTimelineState =  Timeline.GetState(Context);
+						if (CurSample.TotalWeight > 0.0f)
+						{
+							WeightedDuration += static_cast<double>(CurSample.TotalWeight) * static_cast<double> (CurTimelineState.GetDuration()); 
+							WeightedPlayRate += static_cast<double>(CurSample.TotalWeight) * static_cast<double> (CurTimelineState.GetPlayRate());
+							WeightedPositionRatio += static_cast<double>(CurSample.TotalWeight) * static_cast<double> (CurTimelineState.GetPositionRatio()); 
+						}
+						if (CurSample.TotalWeight > HighestWeight)
+						{
+							HighestWeight = CurSample.TotalWeight;
+							bLooping = CurTimelineState.IsLooping(); 
+						}
+					}
+				}
+			}
+
+			double Position = WeightedPositionRatio * WeightedDuration; 
+			return FTimelineState(Position, WeightedDuration, WeightedPlayRate, bLooping);
+		}
+		return FTimelineState();
+	}
+
+
+	void FBlendSpacePlayerTrait::AddReferencedObjects(const FExecutionContext& Context, const TTraitBinding<IGarbageCollection>& Binding, FReferenceCollector& Collector) const
+	{
+		IGarbageCollection::AddReferencedObjects(Context, Binding, Collector);
+
+		FInstanceData* InstanceData = Binding.GetInstanceData<FInstanceData>();
+
+		Collector.AddReferencedObject(InstanceData->BlendSpace);
+	}
+}
diff --git a/Engine/Plugins/Experimental/AnimNextAnimGraph/Source/AnimNextAnimGraph/Private/Traits/BlendSpacePlayer.h b/Engine/Plugins/Experimental/AnimNextAnimGraph/Source/AnimNextAnimGraph/Private/Traits/BlendSpacePlayer.h
new file mode 100644
--- /dev/null
+++ b/Engine/Plugins/Experimental/AnimNextAnimGraph/Source/AnimNextAnimGraph/Private/Traits/BlendSpacePlayer.h
@@ -0,0 +1,126 @@
+// Copyright Epic Games, Inc. All Rights Reserved.
+
+#pragma once
+
+#include "CoreMinimal.h"
+#include "TraitCore/Trait.h"
+#include "TraitInterfaces/IEvaluate.h"
+#include "TraitInterfaces/IGarbageCollection.h"
+#include "TraitInterfaces/ITimeline.h"
+#include "TraitInterfaces/IHierarchy.h"
+#include "TraitInterfaces/IUpdate.h"
+#include "TraitInterfaces/IContinuousBlend.h"
+#include "Animation/AnimSequence.h"
+#include "Animation/BlendSpace.h"
+#include "AnimNextDataInterfacePayload.h"
+
+#include "BlendSpacePlayer.generated.h"
+
+USTRUCT(meta = (DisplayName = "Blend Space Player"))
+struct FAnimNextBlendSpacePlayerTraitSharedData : public FAnimNextTraitSharedData
+{
+	GENERATED_BODY()
+
+	/** The blend space to play. */
+	UPROPERTY(EditAnywhere, Category = "Default", meta = (ExportAsReference = "true"))
+	TObjectPtr<UBlendSpace> BlendSpace;
+
+	/** The location on the x-axis to sample. */
+	UPROPERTY(EditAnywhere, Category = "Default")
+	float XAxisSamplePoint = 0.0f;
+
+	/** The location on the y-axis to sample. */
+	UPROPERTY(EditAnywhere, Category = "Default")
+	float YAxisSamplePoint = 0.0f;
+
+	// Latent pin support boilerplate
+	#define TRAIT_LATENT_PROPERTIES_ENUMERATOR(GeneratorMacro) \
+		GeneratorMacro(BlendSpace) \
+		GeneratorMacro(XAxisSamplePoint) \
+		GeneratorMacro(YAxisSamplePoint) \
+
+
+	GENERATE_TRAIT_LATENT_PROPERTIES(FAnimNextBlendSpacePlayerTraitSharedData, TRAIT_LATENT_PROPERTIES_ENUMERATOR)
+	#undef TRAIT_LATENT_PROPERTIES_ENUMERATOR
+};
+
+namespace UE::AnimNext
+{
+	/**
+	 * FBlendSpacePlayerTrait
+	 * 
+	 * A trait that can play a blend space
+	 */
+	struct FBlendSpacePlayerTrait : FBaseTrait, IEvaluate,  IUpdate, IUpdateTraversal, IHierarchy, ITimeline, IContinuousBlend, IGarbageCollection
+	{
+		DECLARE_ANIM_TRAIT(FBlendSpacePlayerTrait, 0x11a0f4e3, FBaseTrait)
+
+		using FSharedData = FAnimNextBlendSpacePlayerTraitSharedData;
+
+		struct FSampleGraphState
+		{
+			FAnimNextDataInterfacePayload GraphPayload;
+
+			// The graph instance
+			TSharedPtr<FAnimNextGraphInstance> Instance;
+
+			// Our child handle
+			// Weak handle to the graph instance's root
+			FTraitPtr ChildPtr;
+
+			float Weight = 0.0f; 
+
+			bool bSampledLastFrame = false; 
+
+			bool bSampledThisFrame = false; 
+
+		};
+
+		struct FInstanceData : FTrait::FInstanceData
+		{
+			// List of anim next graphs that represent each sample in the blend space
+			TArray<FSampleGraphState> SampleGraphs;
+
+			// Cached value of the blend space we are playing
+			TObjectPtr<UBlendSpace> BlendSpace;
+
+			// Cached blend samples updated by task and used to interpolate between points over time
+			TArray<FBlendSampleData> BlendSamplesData;
+			
+			/** Previous position in the triangulation/segmentation */
+			int32 CachedTriangulationIndex = -1;
+
+			// Filter used to dampen coordinate changes
+			FBlendFilter BlendFilter;
+
+			int32 NumChildren = 0;
+
+			void Construct(const FExecutionContext& Context, const FTraitBinding& Binding);
+			void Destruct(const FExecutionContext& Context, const FTraitBinding& Binding);
+		};
+
+		// IEvaluate impl
+		virtual void PostEvaluate(FEvaluateTraversalContext& Context, const TTraitBinding<IEvaluate>& Binding) const override;
+
+		// IHierarchy impl
+		virtual uint32 GetNumChildren(const FExecutionContext& Context, const TTraitBinding<IHierarchy>& Binding) const override;
+		virtual void GetChildren(const FExecutionContext& Context, const TTraitBinding<IHierarchy>& Binding, FChildrenArray& Children) const override;
+
+		// IContinuousBlend impl
+		virtual float GetBlendWeight(const FExecutionContext& Context, const TTraitBinding<IContinuousBlend>& Binding, int32 ChildIndex) const override;
+
+		// IUpdate impl
+		virtual void OnBecomeRelevant(FUpdateTraversalContext& Context, const TTraitBinding<IUpdate>& Binding, const FTraitUpdateState& TraitState) const override;
+		virtual void PreUpdate(FUpdateTraversalContext& Context, const TTraitBinding<IUpdate>& Binding, const FTraitUpdateState& TraitState) const override;
+
+		// IUpdateTraversal impl
+		virtual void QueueChildrenForTraversal(FUpdateTraversalContext& Context, const TTraitBinding<IUpdateTraversal>& Binding, const FTraitUpdateState& TraitState, FUpdateTraversalQueue& TraversalQueue) const override;
+
+		// ITimeline impl
+		virtual FTimelineState GetState(FExecutionContext& Context, const TTraitBinding<ITimeline>& Binding) const override;
+
+		// IGarbageCollection impl
+		virtual void AddReferencedObjects(const FExecutionContext& Context, const TTraitBinding<IGarbageCollection>& Binding, FReferenceCollector& Collector) const override;
+
+	};
+}
diff --git a/Engine/Plugins/Experimental/AnimNextAnimGraph/Source/AnimNextAnimGraph/Public/Graph/AnimNextGraphInstance.h b/Engine/Plugins/Experimental/AnimNextAnimGraph/Source/AnimNextAnimGraph/Public/Graph/AnimNextGraphInstance.h
--- a/Engine/Plugins/Experimental/AnimNextAnimGraph/Source/AnimNextAnimGraph/Public/Graph/AnimNextGraphInstance.h
+++ b/Engine/Plugins/Experimental/AnimNextAnimGraph/Source/AnimNextAnimGraph/Public/Graph/AnimNextGraphInstance.h
@@ -40,6 +40,7 @@ namespace UE::AnimNext
 	struct FStateTreeTrait;
 	struct FCachedBindingInfo;
 	struct FInjectionUtils;
+	struct FBlendSpacePlayerTrait;
 }
 
 using GraphInstanceComponentMapType = TMap<FName, TSharedPtr<UE::AnimNext::FGraphInstanceComponent>>;
@@ -245,6 +246,7 @@ private:
 	friend FRigUnit_AnimNextRunAnimationGraph_v2;
 	friend UE::AnimNext::FInjectionUtils;
 	friend FAnimNextStateTreeRigVMConditionBase;	// Temp - remove: Needs to access mutable variables to copy-in data for function shim
+	friend UE::AnimNext::FBlendSpacePlayerTrait;
 };
 
 template<>
