Change normalised time computation in anim next blend space players to use the highest weighted child time.  This method is not technically correct but maintains compatibility with ABP and current motion matching implementations to ease implementation

diff --git a/Engine/Plugins/Experimental/AnimNextAnimGraph/Source/AnimNextAnimGraph/Private/Traits/BlendSpacePlayer.cpp b/Engine/Plugins/Experimental/AnimNextAnimGraph/Source/AnimNextAnimGraph/Private/Traits/BlendSpacePlayer.cpp
--- a/Engine/Plugins/Experimental/AnimNextAnimGraph/Source/AnimNextAnimGraph/Private/Traits/BlendSpacePlayer.cpp
+++ b/Engine/Plugins/Experimental/AnimNextAnimGraph/Source/AnimNextAnimGraph/Private/Traits/BlendSpacePlayer.cpp
@@ -259,52 +259,34 @@ namespace UE::AnimNext
 	{
 		const FInstanceData* InstanceData = Binding.GetInstanceData<FInstanceData>();
 
-		// get the highest weighted timeline for looping, but average the duration, position ratios, and play rate based on weighting 
 		if (InstanceData->NumChildren > 0)
 		{
 			int32 HighestWeightIndex = 0;
-			float HighestWeight = 0.0f; 
-			bool bLooping = false; 
-			double WeightedDuration = 0.0; 
-			double WeightedPlayRate = 0.0; 
-			double WeightedPositionRatio = 0.0; 
-
+			float HighestWeight = 0.0f;
 			// using the blend samples set the weights and activity 
 			for (const FBlendSampleData& CurSample : InstanceData->BlendSamplesData)
-			{	
-				const int32 CurSampleDataIndex = CurSample.SampleDataIndex;
-				if (CurSampleDataIndex < InstanceData->SampleGraphs.Num())
+			{
+				if (CurSample.SampleDataIndex < InstanceData->SampleGraphs.Num() && CurSample.TotalWeight > HighestWeight)
 				{
-					const FSampleGraphState& SampleState = InstanceData->SampleGraphs[CurSampleDataIndex];
-					FTraitStackBinding ChildTraitStack;
-					ensure(Context.GetStack(SampleState.ChildPtr, ChildTraitStack));
-
-					UE::AnimNext::TTraitBinding<UE::AnimNext::ITimeline> Timeline;
-					if (ChildTraitStack.GetInterface<UE::AnimNext::ITimeline>(Timeline))
-					{
-						FTimelineState CurTimelineState =  Timeline.GetState(Context);
-						if (CurSample.TotalWeight > 0.0f)
-						{
-							WeightedDuration += static_cast<double>(CurSample.TotalWeight) * static_cast<double> (CurTimelineState.GetDuration()); 
-							WeightedPlayRate += static_cast<double>(CurSample.TotalWeight) * static_cast<double> (CurTimelineState.GetPlayRate());
-							WeightedPositionRatio += static_cast<double>(CurSample.TotalWeight) * static_cast<double> (CurTimelineState.GetPositionRatio()); 
-						}
-						if (CurSample.TotalWeight > HighestWeight)
-						{
-							HighestWeight = CurSample.TotalWeight;
-							bLooping = CurTimelineState.IsLooping(); 
-						}
-					}
+					HighestWeight = CurSample.TotalWeight;
+					HighestWeightIndex = CurSample.SampleDataIndex;
+				}
+			}
+			if (HighestWeightIndex < InstanceData->SampleGraphs.Num())
+			{
+				const FSampleGraphState& SampleState = InstanceData->SampleGraphs[HighestWeightIndex];
+				FTraitStackBinding ChildTraitStack;
+				ensure(Context.GetStack(SampleState.ChildPtr, ChildTraitStack));
+				UE::AnimNext::TTraitBinding<UE::AnimNext::ITimeline> Timeline;
+				if (ChildTraitStack.GetInterface<UE::AnimNext::ITimeline>(Timeline))
+				{
+					return Timeline.GetState(Context);
 				}
 			}
-
-			double Position = WeightedPositionRatio * WeightedDuration; 
-			return FTimelineState(Position, WeightedDuration, WeightedPlayRate, bLooping);
 		}
 		return FTimelineState();
 	}
 
-
 	void FBlendSpacePlayerTrait::AddReferencedObjects(const FExecutionContext& Context, const TTraitBinding<IGarbageCollection>& Binding, FReferenceCollector& Collector) const
 	{
 		IGarbageCollection::AddReferencedObjects(Context, Binding, Collector);
