AnimNext: Fix blendspace blending when smoothing is enabled

diff --git a/Engine/Plugins/Experimental/AnimNextAnimGraph/Source/AnimNextAnimGraph/Private/Traits/BlendSpacePlayer.cpp b/Engine/Plugins/Experimental/AnimNextAnimGraph/Source/AnimNextAnimGraph/Private/Traits/BlendSpacePlayer.cpp
--- a/Engine/Plugins/Experimental/AnimNextAnimGraph/Source/AnimNextAnimGraph/Private/Traits/BlendSpacePlayer.cpp
+++ b/Engine/Plugins/Experimental/AnimNextAnimGraph/Source/AnimNextAnimGraph/Private/Traits/BlendSpacePlayer.cpp
@@ -4,6 +4,7 @@
 #include "AnimationRuntime.h"
 #include "TraitCore/ExecutionContext.h"
 #include "EvaluationVM/Tasks/BlendKeyframes.h"
+#include "EvaluationVM/Tasks/NormalizeRotations.h"
 #include "Graph/AnimNextAnimationGraph.h"
 #include "AnimNextAnimGraphSettings.h"
 #include "AnimNextDataInterfacePayload.h"
@@ -51,28 +52,29 @@ namespace UE::AnimNext
 	{
 		const FInstanceData* InstanceData = Binding.GetInstanceData<FInstanceData>();
 
-		if (InstanceData->NumChildren == 3)
+		if (InstanceData->NumChildren >= 3)
 		{
-			// We have three children, interpolate them
-
 			TTraitBinding<IContinuousBlend> ContinuousBlendTrait;
 			Binding.GetStackInterface(ContinuousBlendTrait);
 
-			const float BlendWeight2 = ContinuousBlendTrait.GetBlendWeight(Context, 2);
-			const float BlendWeight1 = ContinuousBlendTrait.GetBlendWeight(Context, 1);
-			const float BlendWeight0 = ContinuousBlendTrait.GetBlendWeight(Context, 0);
-			
-			if (FAnimWeight::IsRelevant(BlendWeight2 + BlendWeight1))
-			{
-				const float FirstBlendWeight = BlendWeight2 / (BlendWeight2 + BlendWeight1);
-				Context.AppendTask(FAnimNextBlendTwoKeyframesTask::Make(FirstBlendWeight));
-			}
-			else
+			int32 ChildIndex = InstanceData->NumChildren - 1;
+			float ChildBlendWeight = ContinuousBlendTrait.GetBlendWeight(Context, ChildIndex);
+
+			// The last child overrides the top keyframe and scales it
+			Context.AppendTask(FAnimNextBlendOverwriteKeyframeWithScaleTask::Make(ChildBlendWeight));
+
+			// Other children accumulate with scale
+			ChildIndex--;
+			for (; ChildIndex >= 0; --ChildIndex)
 			{
-				// we have no weight but still need to pop, so still make a task that will get ignored in the next step 
-				Context.AppendTask(FAnimNextBlendTwoKeyframesTask::Make(0.0));
+				ChildBlendWeight = ContinuousBlendTrait.GetBlendWeight(Context, ChildIndex);
+
+				// This trait controls the blend weight and owns it
+				Context.AppendTask(FAnimNextBlendAddKeyframeWithScaleTask::Make(ChildBlendWeight));
 			}
-			Context.AppendTask(FAnimNextBlendTwoKeyframesTask::Make(1.0 - BlendWeight0));
+
+			// Once we are done, we normalize rotations
+			Context.AppendTask(FAnimNextNormalizeKeyframeRotationsTask());
 		}
 		else if (InstanceData->NumChildren == 2)
 		{
