AnimNext: Fix for interface variable index mapping issue

diff --git a/Engine/Plugins/Experimental/AnimNextAnimGraph/Source/AnimNextAnimGraph/Private/Graph/AnimNextGraphInstance.cpp b/Engine/Plugins/Experimental/AnimNextAnimGraph/Source/AnimNextAnimGraph/Private/Graph/AnimNextGraphInstance.cpp
--- a/Engine/Plugins/Experimental/AnimNextAnimGraph/Source/AnimNextAnimGraph/Private/Graph/AnimNextGraphInstance.cpp
+++ b/Engine/Plugins/Experimental/AnimNextAnimGraph/Source/AnimNextAnimGraph/Private/Graph/AnimNextGraphInstance.cpp
@@ -301,7 +301,7 @@ bool FAnimNextGraphInstance::BindToCachedBindings()
 			}
 
 			if( Desc.Name == CachedBindingVariable.VariableName &&
-				InterfaceVariableIndex == CachedBindingVariable.VariableIndex &&
+				VariableIndex == CachedBindingVariable.VariableIndex &&
 				Desc.CachedProperty->GetClass() == CachedBindingVariable.Property->GetClass())
 			{
 				ExtendedExecuteContext.ExternalVariableRuntimeData[VariableIndex].Memory = CachedBindingVariable.Memory;
