Added SCOPED_NAMED_EVENT for some high level animnext profilng events instead of using stats. This is done so that we can have some high level timing data in test builds that have the stat system disabled

diff --git a/Engine/Plugins/Experimental/AnimNext/Source/AnimNext/Private/Module/ModuleTickFunction.cpp b/Engine/Plugins/Experimental/AnimNext/Source/AnimNext/Private/Module/ModuleTickFunction.cpp
--- a/Engine/Plugins/Experimental/AnimNext/Source/AnimNext/Private/Module/ModuleTickFunction.cpp
+++ b/Engine/Plugins/Experimental/AnimNext/Source/AnimNext/Private/Module/ModuleTickFunction.cpp
@@ -15,8 +15,6 @@
 #include "TraitCore/TraitEventList.h"
 #include "Module/ModuleGuard.h"
 
-DEFINE_STAT(STAT_AnimNext_Module_EventTick);
-DEFINE_STAT(STAT_AnimNext_Module_EndTick);
 
 namespace UE::AnimNext
 {
@@ -28,7 +26,7 @@ void FModuleEndTickFunction::ExecuteTick(float DeltaTime, ELevelTick TickType, E
 
 void FModuleEndTickFunction::Run(float DeltaTime)
 {
-	SCOPE_CYCLE_COUNTER(STAT_AnimNext_Module_EndTick);
+	SCOPED_NAMED_EVENT(AnimNext_Module_EndTick, FColor::Orange);
 
 	FModuleWriteGuard Guard(ModuleInstance);
 
@@ -117,7 +115,7 @@ void FModuleEventTickFunction::ExecuteTick(float DeltaTime, ELevelTick TickType,
 
 void FModuleEventTickFunction::Run(float InDeltaTime)
 {
-	SCOPE_CYCLE_COUNTER(STAT_AnimNext_Module_EventTick);
+	SCOPED_NAMED_EVENT(AnimNext_Module_EventTick, FColor::Orange);
 
 	FModuleWriteGuard Guard(ModuleInstance);
 
diff --git a/Engine/Plugins/Experimental/AnimNextAnimGraph/Source/AnimNextAnimGraph/Private/Graph/RigUnit_AnimNextRunAnimationGraph_v1.cpp b/Engine/Plugins/Experimental/AnimNextAnimGraph/Source/AnimNextAnimGraph/Private/Graph/RigUnit_AnimNextRunAnimationGraph_v1.cpp
--- a/Engine/Plugins/Experimental/AnimNextAnimGraph/Source/AnimNextAnimGraph/Private/Graph/RigUnit_AnimNextRunAnimationGraph_v1.cpp
+++ b/Engine/Plugins/Experimental/AnimNextAnimGraph/Source/AnimNextAnimGraph/Private/Graph/RigUnit_AnimNextRunAnimationGraph_v1.cpp
@@ -14,11 +14,10 @@
 
 #include UE_INLINE_GENERATED_CPP_BY_NAME(RigUnit_AnimNextRunAnimationGraph_v1)
 
-DEFINE_STAT(STAT_AnimNext_Run_Graph);
 
 FRigUnit_AnimNextRunAnimationGraph_v1_Execute()
 {
-	SCOPE_CYCLE_COUNTER(STAT_AnimNext_Run_Graph);
+	SCOPED_NAMED_EVENT(AnimNext_Run_Graph_V1, FColor::Orange);
 
 	using namespace UE::AnimNext;
 
diff --git a/Engine/Plugins/Experimental/AnimNextAnimGraph/Source/AnimNextAnimGraph/Private/Graph/RigUnit_AnimNextRunAnimationGraph_v2.cpp b/Engine/Plugins/Experimental/AnimNextAnimGraph/Source/AnimNextAnimGraph/Private/Graph/RigUnit_AnimNextRunAnimationGraph_v2.cpp
--- a/Engine/Plugins/Experimental/AnimNextAnimGraph/Source/AnimNextAnimGraph/Private/Graph/RigUnit_AnimNextRunAnimationGraph_v2.cpp
+++ b/Engine/Plugins/Experimental/AnimNextAnimGraph/Source/AnimNextAnimGraph/Private/Graph/RigUnit_AnimNextRunAnimationGraph_v2.cpp
@@ -20,7 +20,7 @@
 
 FRigUnit_AnimNextRunAnimationGraph_v2_Execute()
 {
-	SCOPE_CYCLE_COUNTER(STAT_AnimNext_Run_Graph);
+	SCOPED_NAMED_EVENT(AnimNext_Run_Graph_V2, FColor::Orange);
 
 	using namespace UE::AnimNext;
 
diff --git a/Engine/Plugins/Experimental/AnimNextAnimGraph/Source/AnimNextAnimGraph/Private/TraitInterfaces/IEvaluate.cpp b/Engine/Plugins/Experimental/AnimNextAnimGraph/Source/AnimNextAnimGraph/Private/TraitInterfaces/IEvaluate.cpp
--- a/Engine/Plugins/Experimental/AnimNextAnimGraph/Source/AnimNextAnimGraph/Private/TraitInterfaces/IEvaluate.cpp
+++ b/Engine/Plugins/Experimental/AnimNextAnimGraph/Source/AnimNextAnimGraph/Private/TraitInterfaces/IEvaluate.cpp
@@ -10,8 +10,6 @@
 #include "Graph/AnimNextGraphInstance.h"
 #include "Graph/AnimNext_LODPose.h"
 
-DEFINE_STAT(STAT_AnimNext_EvaluateGraph);
-
 namespace UE::AnimNext
 {
 	FEvaluateTraversalContext::FEvaluateTraversalContext(FEvaluationProgram& InEvaluationProgram)
@@ -99,7 +97,7 @@ namespace UE::AnimNext
 
 	FEvaluationProgram EvaluateGraph(const FWeakTraitPtr& GraphRootPtr)
 	{
-		SCOPE_CYCLE_COUNTER(STAT_AnimNext_EvaluateGraph);
+		SCOPED_NAMED_EVENT(AnimNext_EvaluateGraph, FColor::Orange);
 		
 		FEvaluationProgram EvaluationProgram;
 
diff --git a/Engine/Plugins/Experimental/AnimNextAnimGraph/Source/AnimNextAnimGraph/Private/TraitInterfaces/IUpdate.cpp b/Engine/Plugins/Experimental/AnimNextAnimGraph/Source/AnimNextAnimGraph/Private/TraitInterfaces/IUpdate.cpp
--- a/Engine/Plugins/Experimental/AnimNextAnimGraph/Source/AnimNextAnimGraph/Private/TraitInterfaces/IUpdate.cpp
+++ b/Engine/Plugins/Experimental/AnimNextAnimGraph/Source/AnimNextAnimGraph/Private/TraitInterfaces/IUpdate.cpp
@@ -8,8 +8,6 @@
 #include "Graph/AnimNextGraphInstance.h"
 #include "TraitCore/TraitEventRaising.h"
 
-DEFINE_STAT(STAT_AnimNext_UpdateGraph);
-
 namespace UE::AnimNext
 {
 	namespace Private
@@ -476,7 +474,7 @@ namespace UE::AnimNext
 
 	void UpdateGraph(FAnimNextGraphInstance& GraphInstance, float DeltaTime, FTraitEventList& InputEventList, FTraitEventList& OutputEventList)
 	{
-		SCOPE_CYCLE_COUNTER(STAT_AnimNext_UpdateGraph);
+		SCOPED_NAMED_EVENT(AnimNext_UpdateGraph, FColor::Orange);
 		
 		if (!GraphInstance.IsValid())
 		{
diff --git a/Engine/Plugins/Experimental/AnimNextAnimGraph/Source/AnimNextAnimGraph/Public/AnimNextAnimGraphStats.h b/Engine/Plugins/Experimental/AnimNextAnimGraph/Source/AnimNextAnimGraph/Public/AnimNextAnimGraphStats.h
--- a/Engine/Plugins/Experimental/AnimNextAnimGraph/Source/AnimNextAnimGraph/Public/AnimNextAnimGraphStats.h
+++ b/Engine/Plugins/Experimental/AnimNextAnimGraph/Source/AnimNextAnimGraph/Public/AnimNextAnimGraphStats.h
@@ -5,11 +5,6 @@
 #include "Stats/Stats.h"
 #include "AnimNextStats.h"
 
-DECLARE_CYCLE_STAT_EXTERN(TEXT("AnimNext: Update Graph"), STAT_AnimNext_UpdateGraph, STATGROUP_AnimNext, ANIMNEXTANIMGRAPH_API);
-DECLARE_CYCLE_STAT_EXTERN(TEXT("AnimNext: Evaluate Graph"), STAT_AnimNext_EvaluateGraph, STATGROUP_AnimNext, ANIMNEXTANIMGRAPH_API);
-
-DECLARE_CYCLE_STAT_EXTERN(TEXT("AnimNext: Run Graph"), STAT_AnimNext_Run_Graph, STATGROUP_AnimNext, ANIMNEXTANIMGRAPH_API);
-
 DECLARE_CYCLE_STAT_EXTERN(TEXT("AnimNext: Allocate Graph Instance"), STAT_AnimNext_Graph_AllocateInstance, STATGROUP_AnimNext, ANIMNEXTANIMGRAPH_API);
 DECLARE_CYCLE_STAT_EXTERN(TEXT("AnimNext: Graph RigVM"), STAT_AnimNext_Graph_RigVM, STATGROUP_AnimNext, ANIMNEXTANIMGRAPH_API);
 DECLARE_CYCLE_STAT_EXTERN(TEXT("AnimNext: Execute Evaluation Program"), STAT_AnimNext_EvaluationProgram_Execute, STATGROUP_AnimNext, ANIMNEXTANIMGRAPH_API);
