Added new AnimNext AlphaBlend Trait and DuplicateTopKeyframe Task

diff --git a/Engine/Plugins/Experimental/AnimNextAnimGraph/Source/AnimNextAnimGraph/Private/EvaluationVM/Tasks/BlendKeyframes.cpp b/Engine/Plugins/Experimental/AnimNextAnimGraph/Source/AnimNextAnimGraph/Private/EvaluationVM/Tasks/BlendKeyframes.cpp
--- a/Engine/Plugins/Experimental/AnimNextAnimGraph/Source/AnimNextAnimGraph/Private/EvaluationVM/Tasks/BlendKeyframes.cpp
+++ b/Engine/Plugins/Experimental/AnimNextAnimGraph/Source/AnimNextAnimGraph/Private/EvaluationVM/Tasks/BlendKeyframes.cpp
@@ -5,6 +5,7 @@
 #include "EvaluationVM/EvaluationVM.h"
 #include "EvaluationVM/KeyframeState.h"
 #include "TransformArrayOperations.h"
+#include "Animation/InputScaleBias.h"
 
 FAnimNextBlendTwoKeyframesTask FAnimNextBlendTwoKeyframesTask::Make(float InterpolationAlpha)
 {
@@ -14,6 +15,15 @@ FAnimNextBlendTwoKeyframesTask FAnimNextBlendTwoKeyframesTask::Make(float Interp
 	return Task;
 }
 
+FAnimNextBlendTwoKeyframesTask FAnimNextBlendTwoKeyframesTask::Make(const FName& AlphaSourceCurveName, const int8 AlphaCurveInputIndex, TFunction<float(float)> InputScaleBiasClampFn)
+{
+	FAnimNextBlendTwoKeyframesTask Task;
+	Task.AlphaSourceCurveName = AlphaSourceCurveName;
+	Task.AlphaCurveInputIndex = AlphaCurveInputIndex;
+	Task.InputScaleBiasClampFn = MoveTemp(InputScaleBiasClampFn);
+	return Task;
+}
+
 void FAnimNextBlendTwoKeyframesTask::Execute(UE::AnimNext::FEvaluationVM& VM) const
 {
 	using namespace UE::AnimNext;
@@ -35,40 +45,74 @@ void FAnimNextBlendTwoKeyframesTask::Execute(UE::AnimNext::FEvaluationVM& VM) co
 		return;
 	}
 
-	const float Alpha = FMath::Clamp(InterpolationAlpha, 0.0f, 1.0f);
-	const float WeightOfPoseA = 1.0f - Alpha;
-	const float WeightOfPoseB = Alpha;
+	const float WeightOfPoseB = GetInterpolationAlpha(KeyframeA.Get(), KeyframeB.Get());
 
-	if (EnumHasAnyFlags(VM.GetFlags(), EEvaluationFlags::Bones))
+	if (!FAnimWeight::IsRelevant(WeightOfPoseB))
 	{
-		check(KeyframeA->Pose.GetNumBones() == KeyframeB->Pose.GetNumBones());
+		VM.PushValue(KEYFRAME_STACK_NAME, MoveTemp(KeyframeA));
+	}
+	else if (FAnimWeight::IsFullWeight(WeightOfPoseB))
+	{
+		VM.PushValue(KEYFRAME_STACK_NAME, MoveTemp(KeyframeB));
+	}
+	else
+	{
+		const float WeightOfPoseA = 1.0f - WeightOfPoseB;
 
-		const FTransformArrayView KeyframeBTransformsView = KeyframeB->Pose.LocalTransforms.GetView();
+		if (EnumHasAnyFlags(VM.GetFlags(), EEvaluationFlags::Bones))
+		{
+			check(KeyframeA->Pose.GetNumBones() == KeyframeB->Pose.GetNumBones());
 
-		BlendOverwriteWithScale(KeyframeBTransformsView, KeyframeB->Pose.LocalTransforms.GetConstView(), WeightOfPoseB);
-		BlendAddWithScale(KeyframeBTransformsView, KeyframeA->Pose.LocalTransforms.GetConstView(), WeightOfPoseA);
+			const FTransformArrayView KeyframeBTransformsView = KeyframeB->Pose.LocalTransforms.GetView();
 
-		// Ensure that all of the resulting rotations are normalized
-		NormalizeRotations(KeyframeBTransformsView);
-	}
+			BlendOverwriteWithScale(KeyframeBTransformsView, KeyframeB->Pose.LocalTransforms.GetConstView(), WeightOfPoseB);
+			BlendAddWithScale(KeyframeBTransformsView, KeyframeA->Pose.LocalTransforms.GetConstView(), WeightOfPoseA);
 
-	if (EnumHasAnyFlags(VM.GetFlags(), EEvaluationFlags::Curves))
-	{
-		// Curves cannot blend in place
-		FBlendedCurve Result;
-		Result.Lerp(KeyframeA->Curves, KeyframeB->Curves, WeightOfPoseB);
+			// Ensure that all of the resulting rotations are normalized
+			NormalizeRotations(KeyframeBTransformsView);
+		}
+
+		if (EnumHasAnyFlags(VM.GetFlags(), EEvaluationFlags::Curves))
+		{
+			// Curves cannot blend in place
+			FBlendedCurve Result;
+			Result.Lerp(KeyframeA->Curves, KeyframeB->Curves, WeightOfPoseB);
+
+			KeyframeB->Curves = MoveTemp(Result);
+		}
 
-		KeyframeB->Curves = MoveTemp(Result);
+		if (EnumHasAnyFlags(VM.GetFlags(), EEvaluationFlags::Attributes))
+		{
+			UE::Anim::Attributes::BlendAttributes({ KeyframeA->Attributes, KeyframeB->Attributes }, { WeightOfPoseA, WeightOfPoseB }, { 0, 1 }, KeyframeB->Attributes);
+		}
+
+		VM.PushValue(KEYFRAME_STACK_NAME, MoveTemp(KeyframeB));
 	}
+}
 
-	if (EnumHasAnyFlags(VM.GetFlags(), EEvaluationFlags::Attributes))
+float FAnimNextBlendTwoKeyframesTask::GetInterpolationAlpha(const UE::AnimNext::FKeyframeState* KeyframeA, const UE::AnimNext::FKeyframeState* KeyframeB) const
+{
+	float Alpha = InterpolationAlpha;
+
+	if (AlphaSourceCurveName != NAME_None && AlphaCurveInputIndex != INDEX_NONE)
 	{
-		UE::Anim::Attributes::BlendAttributes({ KeyframeA->Attributes, KeyframeB->Attributes }, { WeightOfPoseA, WeightOfPoseB }, { 0, 1 }, KeyframeB->Attributes);
+		if (ensure(KeyframeA != nullptr && KeyframeB != nullptr))
+		{
+			const FBlendedCurve& Curves = AlphaCurveInputIndex == 0 ? KeyframeA->Curves : KeyframeB->Curves;
+			Alpha = Curves.Get(AlphaSourceCurveName); // if the curve does not exist, it returns 0.f
+
+			if (InputScaleBiasClampFn.IsSet())
+			{
+				Alpha = InputScaleBiasClampFn(Alpha);
+			}
+		}
 	}
 
-	VM.PushValue(KEYFRAME_STACK_NAME, MoveTemp(KeyframeB));
+	return FMath::Clamp(Alpha, 0.0f, 1.0f);
 }
 
+// --- FAnimNextBlendOverwriteKeyframeWithScaleTask ---
+
 FAnimNextBlendOverwriteKeyframeWithScaleTask FAnimNextBlendOverwriteKeyframeWithScaleTask::Make(float ScaleFactor)
 {
 	FAnimNextBlendOverwriteKeyframeWithScaleTask Task;
@@ -110,6 +154,8 @@ void FAnimNextBlendOverwriteKeyframeWithScaleTask::Execute(UE::AnimNext::FEvalua
 	VM.PushValue(KEYFRAME_STACK_NAME, MoveTemp(Keyframe));
 }
 
+// --- FAnimNextBlendAddKeyframeWithScaleTask  ---
+
 FAnimNextBlendAddKeyframeWithScaleTask FAnimNextBlendAddKeyframeWithScaleTask::Make(float ScaleFactor)
 {
 	FAnimNextBlendAddKeyframeWithScaleTask Task;
diff --git a/Engine/Plugins/Experimental/AnimNextAnimGraph/Source/AnimNextAnimGraph/Private/EvaluationVM/Tasks/StoreKeyframe.cpp b/Engine/Plugins/Experimental/AnimNextAnimGraph/Source/AnimNextAnimGraph/Private/EvaluationVM/Tasks/StoreKeyframe.cpp
--- a/Engine/Plugins/Experimental/AnimNextAnimGraph/Source/AnimNextAnimGraph/Private/EvaluationVM/Tasks/StoreKeyframe.cpp
+++ b/Engine/Plugins/Experimental/AnimNextAnimGraph/Source/AnimNextAnimGraph/Private/EvaluationVM/Tasks/StoreKeyframe.cpp
@@ -39,3 +39,37 @@ void FAnimNextStoreKeyframeTransformsTask::Execute(UE::AnimNext::FEvaluationVM&
 		}
 	}
 }
+
+// --- FAnimNextDuplicateTopKeyframeTask  ---
+
+FAnimNextDuplicateTopKeyframeTask FAnimNextDuplicateTopKeyframeTask::Make()
+{
+	FAnimNextDuplicateTopKeyframeTask Task;
+	return Task;
+}
+
+// Task entry point
+void FAnimNextDuplicateTopKeyframeTask::Execute(UE::AnimNext::FEvaluationVM& VM) const
+{
+	using namespace UE::AnimNext;
+
+	TUniquePtr<FKeyframeState> KeyframeOut;
+
+	// Try to get a Keyframe from the stack without consuming it
+	// If there is a Keyframe, create a copy
+	if (const TUniquePtr<FKeyframeState>* KeyframeIn = VM.PeekValue<TUniquePtr<FKeyframeState>>(KEYFRAME_STACK_NAME, 0))
+	{
+		KeyframeOut = MakeUnique<FKeyframeState>(VM.MakeUninitializedKeyframe(false));
+		KeyframeOut->Pose.CopyFrom((*KeyframeIn)->Pose);
+		KeyframeOut->Curves.CopyFrom((*KeyframeIn)->Curves);
+		KeyframeOut->Attributes.CopyFrom((*KeyframeIn)->Attributes);
+	}
+	else
+	{
+		// If there is no Keyframe, generate a reference key frame
+		KeyframeOut = MakeUnique<FKeyframeState>(VM.MakeReferenceKeyframe(false));
+	}
+
+	// Push the copy or reference keyframe to the stack
+	VM.PushValue(KEYFRAME_STACK_NAME, MoveTemp(KeyframeOut));
+}
diff --git a/Engine/Plugins/Experimental/AnimNextAnimGraph/Source/AnimNextAnimGraph/Private/Traits/AlphaBlendTrait.cpp b/Engine/Plugins/Experimental/AnimNextAnimGraph/Source/AnimNextAnimGraph/Private/Traits/AlphaBlendTrait.cpp
new file mode 100644
--- /dev/null
+++ b/Engine/Plugins/Experimental/AnimNextAnimGraph/Source/AnimNextAnimGraph/Private/Traits/AlphaBlendTrait.cpp
@@ -0,0 +1,122 @@
+// Copyright Epic Games, Inc. All Rights Reserved.
+
+#include "Traits/AlphaBlendTrait.h"
+
+#include "Animation/AnimTypes.h"
+#include "TraitCore/ExecutionContext.h"
+#include "EvaluationVM/Tasks/BlendKeyframes.h"
+#include "EvaluationVM/Tasks/StoreKeyframe.h"
+
+namespace UE::AnimNext
+{
+
+AUTO_REGISTER_ANIM_TRAIT(FAlphaBlendTrait)
+
+#define TRAIT_INTERFACE_ENUMERATOR(GeneratorMacro) \
+	GeneratorMacro(IUpdate) \
+	GeneratorMacro(IEvaluate) \
+
+// Trait implementation boilerplate
+GENERATE_ANIM_TRAIT_IMPLEMENTATION(FAlphaBlendTrait, TRAIT_INTERFACE_ENUMERATOR, NULL_ANIM_TRAIT_INTERFACE_ENUMERATOR, NULL_ANIM_TRAIT_EVENT_ENUMERATOR)
+#undef TRAIT_INTERFACE_ENUMERATOR
+
+
+void FAlphaBlendTrait::FInstanceData::Construct(const FExecutionContext& Context, const FTraitBinding& Binding)
+{
+	const FSharedData* SharedData = Binding.GetSharedData<FSharedData>();
+
+	InternalAlpha = SharedData->GetAlpha(Binding);
+	ComputedAlphaValue = InternalAlpha;
+	AlphaBoolBlend = SharedData->GetAlphaBoolBlend(Binding);
+	AlphaScaleBiasClamp = SharedData->AlphaScaleBiasClamp;
+}
+
+void FAlphaBlendTrait::OnBecomeRelevant(FUpdateTraversalContext& Context, const TTraitBinding<IUpdate>& Binding, const FTraitUpdateState& TraitState) const
+{
+	IUpdate::OnBecomeRelevant(Context, Binding, TraitState);
+
+	const FSharedData* SharedData = Binding.GetSharedData<FSharedData>();
+	FInstanceData* InstanceData = Binding.GetInstanceData<FInstanceData>();
+
+	InstanceData->InternalAlpha = SharedData->GetAlpha(Binding);
+	InstanceData->ComputedAlphaValue = ComputeAlphaValue(SharedData->GetAlphaInputType(Binding), SharedData, InstanceData, TraitState.GetDeltaTime());
+
+	InstanceData->AlphaBoolBlend.Reinitialize();
+	InstanceData->AlphaScaleBiasClamp.Reinitialize();
+}
+
+void FAlphaBlendTrait::PreUpdate(FUpdateTraversalContext& Context, const TTraitBinding<IUpdate>& Binding, const FTraitUpdateState& TraitState) const
+{
+	IUpdate::PreUpdate(Context, Binding, TraitState);
+
+	const FSharedData* SharedData = Binding.GetSharedData<FSharedData>();
+	FInstanceData* InstanceData = Binding.GetInstanceData<FInstanceData>();
+
+	const float DeltaTime = TraitState.GetDeltaTime();
+	InstanceData->DeltaTime = DeltaTime;
+	InstanceData->ComputedAlphaValue = ComputeAlphaValue(SharedData->GetAlphaInputType(Binding), SharedData, InstanceData, DeltaTime);
+}
+
+// IEvaluate impl
+void FAlphaBlendTrait::PostEvaluate(FEvaluateTraversalContext& Context, const TTraitBinding<IEvaluate>& Binding) const
+{
+	// Duplicate existing pose in the stack
+	Context.AppendTask(FAnimNextDuplicateTopKeyframeTask::Make());
+
+	// Let children perform their tasks on the duplicated pose
+	IEvaluate::PostEvaluate(Context, Binding);
+
+	const FSharedData* SharedData = Binding.GetSharedData<FSharedData>();
+	FInstanceData* InstanceData = Binding.GetInstanceData<FInstanceData>();
+
+	const EAnimAlphaInputType AlphaInputType = SharedData->GetAlphaInputType(Binding);
+	if (AlphaInputType == EAnimAlphaInputType::Curve)
+	{
+		// Add a BlendTwoKey with the original pose and the children modified copy, using alpha curve data (value will be calculated / clamped at the task)
+		const float DeltaTime = InstanceData->DeltaTime;
+		const FName AlphaCurveName = SharedData->GetAlphaCurveName(Binding);
+		FInputScaleBiasClamp* InputScaleBiasClamp = &InstanceData->AlphaScaleBiasClamp;
+		static constexpr int32 AlphaSourceInputKeyframeIndex = 0;
+		Context.AppendTask(FAnimNextBlendTwoKeyframesTask::Make(AlphaCurveName, AlphaSourceInputKeyframeIndex, [DeltaTime, InputScaleBiasClamp](float Alpha) -> float
+			{
+				return InputScaleBiasClamp->ApplyTo(Alpha, DeltaTime);
+			}));
+	}
+	else
+	{
+		// Add a BlendTwoKey with the original pose and the children modified copy, using alpha value as float blend weight
+		const float BlendWeight = InstanceData->ComputedAlphaValue;
+		Context.AppendTask(FAnimNextBlendTwoKeyframesTask::Make(BlendWeight));
+	}
+}
+
+
+float FAlphaBlendTrait::ComputeAlphaValue(const EAnimAlphaInputType AlphaInputType, const FAlphaBlendTraitSharedData* SharedData, FInstanceData* InstanceData, float DeltaTime)
+{
+	// alpha handlers
+	float CurrentAlpha = 0.f;
+
+	switch (AlphaInputType)
+	{
+		case EAnimAlphaInputType::Float:
+		{
+			CurrentAlpha = SharedData->AlphaScaleBias.ApplyTo(InstanceData->AlphaScaleBiasClamp.ApplyTo(InstanceData->InternalAlpha, DeltaTime));
+			break;
+		}
+		case EAnimAlphaInputType::Bool:
+		{
+			CurrentAlpha = InstanceData->AlphaBoolBlend.ApplyTo(SharedData->bAlphaBoolEnabled, DeltaTime);
+			break;
+		}
+		case EAnimAlphaInputType::Curve:
+		{
+			CurrentAlpha = 1.0; // This will be calculated on the task when reading the curve, but we need both branches at full weight
+			break;
+		}
+	};
+
+	const float ClampedAlpha = FMath::Clamp(CurrentAlpha, 0.0f, 1.0f);
+	return ClampedAlpha;
+}
+
+} // namespace UE::AnimNext
diff --git a/Engine/Plugins/Experimental/AnimNextAnimGraph/Source/AnimNextAnimGraph/Public/EvaluationVM/Tasks/BlendKeyframes.h b/Engine/Plugins/Experimental/AnimNextAnimGraph/Source/AnimNextAnimGraph/Public/EvaluationVM/Tasks/BlendKeyframes.h
--- a/Engine/Plugins/Experimental/AnimNextAnimGraph/Source/AnimNextAnimGraph/Public/EvaluationVM/Tasks/BlendKeyframes.h
+++ b/Engine/Plugins/Experimental/AnimNextAnimGraph/Source/AnimNextAnimGraph/Public/EvaluationVM/Tasks/BlendKeyframes.h
@@ -7,6 +7,13 @@
 
 #include "BlendKeyframes.generated.h"
 
+struct FInputScaleBiasClamp;
+
+namespace UE::AnimNext
+{
+	struct FKeyframeState;
+}
+
 /*
  * Blend Two Keyframes Task
  *
@@ -26,6 +33,7 @@
  * In order to blend more than two keyframes together, simply add more FAnimNextBlendAddKeyframeWithScaleTask
  * tasks each with its corresponding weight before normalizing rotations.
  */
+
 USTRUCT()
 struct ANIMNEXTANIMGRAPH_API FAnimNextBlendTwoKeyframesTask : public FAnimNextEvaluationTask
 {
@@ -34,13 +42,25 @@ struct ANIMNEXTANIMGRAPH_API FAnimNextBlendTwoKeyframesTask : public FAnimNextEv
 	DECLARE_ANIM_EVALUATION_TASK(FAnimNextBlendTwoKeyframesTask)
 
 	static FAnimNextBlendTwoKeyframesTask Make(float InterpolationAlpha);
+	static FAnimNextBlendTwoKeyframesTask Make(const FName& AlphaSourceCurveName, const int8 AlphaCurveInputIndex, TFunction<float(float)> InputScaleBiasClampFn);
+	
 
 	// Task entry point
 	virtual void Execute(UE::AnimNext::FEvaluationVM& VM) const override;
 
+	// Compute the current alpha value for curves or return the interpolation alpha otherwise
+	float GetInterpolationAlpha(const UE::AnimNext::FKeyframeState* KeyframeA, const UE::AnimNext::FKeyframeState* KeyframeB) const;
+
 	// The interpolation alpha between the two input keyframes
 	UPROPERTY()
 	float InterpolationAlpha = 0.0f;
+
+	// The curve to evaluate and extract the interpolation alpha between the two input keyframes
+	FName AlphaSourceCurveName = NAME_None;
+
+	int8 AlphaCurveInputIndex = INDEX_NONE;
+
+	TFunction<float(float)> InputScaleBiasClampFn = nullptr;
 };
 
 /*
diff --git a/Engine/Plugins/Experimental/AnimNextAnimGraph/Source/AnimNextAnimGraph/Public/EvaluationVM/Tasks/StoreKeyframe.h b/Engine/Plugins/Experimental/AnimNextAnimGraph/Source/AnimNextAnimGraph/Public/EvaluationVM/Tasks/StoreKeyframe.h
--- a/Engine/Plugins/Experimental/AnimNextAnimGraph/Source/AnimNextAnimGraph/Public/EvaluationVM/Tasks/StoreKeyframe.h
+++ b/Engine/Plugins/Experimental/AnimNextAnimGraph/Source/AnimNextAnimGraph/Public/EvaluationVM/Tasks/StoreKeyframe.h
@@ -44,3 +44,19 @@ struct ANIMNEXTANIMGRAPH_API FAnimNextStoreKeyframeTransformsTask : public FAnim
 
 	UE::AnimNext::FTransformArraySoAHeap* Dest = nullptr;
 };
+
+/*
+ * Duplicates the pose on top of the stack and pushes it as new top
+ */
+USTRUCT()
+struct ANIMNEXTANIMGRAPH_API FAnimNextDuplicateTopKeyframeTask : public FAnimNextEvaluationTask
+{
+	GENERATED_BODY()
+
+	DECLARE_ANIM_EVALUATION_TASK(FAnimNextDuplicateTopKeyframeTask)
+
+	static FAnimNextDuplicateTopKeyframeTask Make();
+
+	// Task entry point
+	virtual void Execute(UE::AnimNext::FEvaluationVM& VM) const override;
+};
diff --git a/Engine/Plugins/Experimental/AnimNextAnimGraph/Source/AnimNextAnimGraph/Public/Traits/AlphaBlendTrait.h b/Engine/Plugins/Experimental/AnimNextAnimGraph/Source/AnimNextAnimGraph/Public/Traits/AlphaBlendTrait.h
new file mode 100644
--- /dev/null
+++ b/Engine/Plugins/Experimental/AnimNextAnimGraph/Source/AnimNextAnimGraph/Public/Traits/AlphaBlendTrait.h
@@ -0,0 +1,87 @@
+// Copyright Epic Games, Inc. All Rights Reserved.
+
+#pragma once
+
+#include "CoreMinimal.h"
+#include "TraitCore/Trait.h"
+#include "TraitCore/TraitSharedData.h"
+#include "TraitInterfaces/IUpdate.h"
+#include "TraitInterfaces/IEvaluate.h"
+#include "Animation/InputScaleBias.h"
+// --- ---
+#include "AlphaBlendTrait.generated.h"
+
+USTRUCT(meta = (DisplayName = "Alpha Blend"))
+struct FAlphaBlendTraitSharedData : public FAnimNextTraitSharedData
+{
+	GENERATED_BODY()
+
+	// alpha value handler
+	UPROPERTY(EditAnywhere, Category = Settings, meta = (EditConditionHides, EditCondition = "AlphaInputType == EAnimAlphaInputType::Float"))
+	float Alpha = 1.f;
+
+	UPROPERTY(EditAnywhere, Category=Settings, meta = (Hidden, EditConditionHides, EditCondition = "AlphaInputType == EAnimAlphaInputType::Float"))
+	FInputScaleBias AlphaScaleBias;
+
+	UPROPERTY(EditAnywhere, Category = Settings, meta = (Hidden, DisplayAfter = "AlphaCurveName", EditConditionHides, EditCondition = "AlphaInputType == EAnimAlphaInputType::Float || AlphaInputType == EAnimAlphaInputType::Curve"))
+	FInputScaleBiasClamp AlphaScaleBiasClamp;
+
+	UPROPERTY(EditAnywhere, Category = Settings, meta = (DisplayName = "bEnabled", EditConditionHides, EditCondition = "AlphaInputType == EAnimAlphaInputType::Bool"))
+	bool bAlphaBoolEnabled = true;
+
+	UPROPERTY(EditAnywhere, Category = Settings, meta = (Hidden, DisplayName = "Blend Settings", EditConditionHides, EditCondition = "AlphaInputType == EAnimAlphaInputType::Bool"))
+	FInputAlphaBoolBlend AlphaBoolBlend;
+
+	UPROPERTY(EditAnywhere, Category = Settings, meta = (EditConditionHides, EditCondition = "AlphaInputType == EAnimAlphaInputType::Curve"))
+	FName AlphaCurveName = NAME_None;
+
+	UPROPERTY(EditAnywhere, Category = Settings)
+	EAnimAlphaInputType AlphaInputType = EAnimAlphaInputType::Float;
+
+	// Latent pin support boilerplate
+#define TRAIT_LATENT_PROPERTIES_ENUMERATOR(GeneratorMacro) \
+		GeneratorMacro(Alpha) \
+		GeneratorMacro(bAlphaBoolEnabled) \
+		GeneratorMacro(AlphaBoolBlend) \
+		GeneratorMacro(AlphaCurveName) \
+		GeneratorMacro(AlphaInputType) \
+
+	GENERATE_TRAIT_LATENT_PROPERTIES(FAlphaBlendTraitSharedData, TRAIT_LATENT_PROPERTIES_ENUMERATOR)
+#undef TRAIT_LATENT_PROPERTIES_ENUMERATOR
+};
+
+namespace UE::AnimNext
+{
+
+USTRUCT(meta = (DisplayName = "Alpha Blend"))
+struct FAlphaBlendTrait : FAdditiveTrait, IUpdate, IEvaluate
+{
+	typedef FBaseTrait Super;
+
+	DECLARE_ANIM_TRAIT(FAlphaBlendTrait, 0xa7c496ca, FAdditiveTrait)
+
+	using FSharedData = FAlphaBlendTraitSharedData;
+		
+	struct FInstanceData : FTrait::FInstanceData
+	{
+		FInputAlphaBoolBlend AlphaBoolBlend;
+		FInputScaleBiasClamp AlphaScaleBiasClamp;
+
+		float DeltaTime = 0.f;
+		float InternalAlpha = 1.f;
+		float ComputedAlphaValue = 0.f;
+
+		void Construct(const FExecutionContext& Context, const FTraitBinding& Binding);
+	};
+
+	// IUpdate
+	virtual void OnBecomeRelevant(FUpdateTraversalContext& Context, const TTraitBinding<IUpdate>& Binding, const FTraitUpdateState& TraitState) const override;
+	virtual void PreUpdate(FUpdateTraversalContext& Context, const TTraitBinding<IUpdate>& Binding, const FTraitUpdateState& TraitState) const override;
+
+	// IEvaluate impl
+	virtual void PostEvaluate(FEvaluateTraversalContext& Context, const TTraitBinding<IEvaluate>& Binding) const override;
+
+	static float ComputeAlphaValue(const EAnimAlphaInputType AlphaInputType, const FAlphaBlendTraitSharedData* SharedData, FInstanceData* InstanceData, float DeltaTime);
+};
+
+} // namespace UE::AnimNext
