Use constexpr hash generation for AnimNext trait UIDs to reduce boilerplate

diff --git a/Engine/Plugins/Experimental/AnimNextAnimGraph/Source/AnimNextAnimGraph/Internal/Traits/BlendStackTrait.h b/Engine/Plugins/Experimental/AnimNextAnimGraph/Source/AnimNextAnimGraph/Internal/Traits/BlendStackTrait.h
--- a/Engine/Plugins/Experimental/AnimNextAnimGraph/Source/AnimNextAnimGraph/Internal/Traits/BlendStackTrait.h
+++ b/Engine/Plugins/Experimental/AnimNextAnimGraph/Source/AnimNextAnimGraph/Internal/Traits/BlendStackTrait.h
@@ -88,7 +88,7 @@ namespace UE::AnimNext
 	 */
 	struct FBlendStackCoreTrait : FBaseTrait, IUpdate, IUpdateTraversal, IHierarchy, IDiscreteBlend, ISmoothBlend, IInertializerBlend, IAttributeProvider, ITimeline, IGarbageCollection, IBlendStack
 	{
-		DECLARE_ANIM_TRAIT(FBlendStackCoreTrait, 0xda7b7f8d, FBaseTrait)
+		DECLARE_ANIM_TRAIT(FBlendStackCoreTrait, FBaseTrait)
 
 		enum class EGraphState : uint8
 		{
@@ -187,7 +187,7 @@ namespace UE::AnimNext
 	 */	
 	struct FBlendStackTrait : FBlendStackCoreTrait
 	{
-		DECLARE_ANIM_TRAIT(FBlendStackTrait, 0x46678513, FBlendStackCoreTrait)
+		DECLARE_ANIM_TRAIT(FBlendStackTrait, FBlendStackCoreTrait)
 
 		using FSharedData = FAnimNextBlendStackTraitSharedData;
 
@@ -207,7 +207,7 @@ namespace UE::AnimNext
 	 */	
 	struct FBlendStackRequesterTrait : FAdditiveTrait, IUpdate
 	{
-		DECLARE_ANIM_TRAIT(FBlendStackRequesterTrait, 0x3076bf31, FAdditiveTrait)
+		DECLARE_ANIM_TRAIT(FBlendStackRequesterTrait, FAdditiveTrait)
 
 		using FSharedData = FAnimNextBlendStackRequesterTraitSharedData;
 
diff --git a/Engine/Plugins/Experimental/AnimNextAnimGraph/Source/AnimNextAnimGraph/Internal/Traits/CallFunction.h b/Engine/Plugins/Experimental/AnimNextAnimGraph/Source/AnimNextAnimGraph/Internal/Traits/CallFunction.h
--- a/Engine/Plugins/Experimental/AnimNextAnimGraph/Source/AnimNextAnimGraph/Internal/Traits/CallFunction.h
+++ b/Engine/Plugins/Experimental/AnimNextAnimGraph/Source/AnimNextAnimGraph/Internal/Traits/CallFunction.h
@@ -1,4 +1,4 @@
-﻿// Copyright Epic Games, Inc. All Rights Reserved.
+// Copyright Epic Games, Inc. All Rights Reserved.
 
 #pragma once
 
@@ -60,7 +60,7 @@ namespace UE::AnimNext
  */
 struct FCallFunctionTrait : FAdditiveTrait, IUpdate
 {
-	DECLARE_ANIM_TRAIT(FCallFunctionTrait, 0x76ef7c0f, FAdditiveTrait)
+	DECLARE_ANIM_TRAIT(FCallFunctionTrait, FAdditiveTrait)
 
 	using FSharedData = FAnimNextCallFunctionSharedData;
 
diff --git a/Engine/Plugins/Experimental/AnimNextAnimGraph/Source/AnimNextAnimGraph/Private/Injection/InjectionSiteTrait.h b/Engine/Plugins/Experimental/AnimNextAnimGraph/Source/AnimNextAnimGraph/Private/Injection/InjectionSiteTrait.h
--- a/Engine/Plugins/Experimental/AnimNextAnimGraph/Source/AnimNextAnimGraph/Private/Injection/InjectionSiteTrait.h
+++ b/Engine/Plugins/Experimental/AnimNextAnimGraph/Source/AnimNextAnimGraph/Private/Injection/InjectionSiteTrait.h
@@ -59,7 +59,7 @@ namespace UE::AnimNext
 	 */
 	struct FInjectionSiteTrait : FAdditiveTrait, IUpdate, IDiscreteBlend, IGarbageCollection, IEvaluate
 	{
-		DECLARE_ANIM_TRAIT(FInjectionSiteTrait, 0xf59a5e4b, FAdditiveTrait)
+		DECLARE_ANIM_TRAIT(FInjectionSiteTrait, FAdditiveTrait)
 
 		// Injection tracker state
 		enum class EInjectionTrackerState : uint8
diff --git a/Engine/Plugins/Experimental/AnimNextAnimGraph/Source/AnimNextAnimGraph/Private/PlayAnim/PlayAnimSlotTrait.h b/Engine/Plugins/Experimental/AnimNextAnimGraph/Source/AnimNextAnimGraph/Private/PlayAnim/PlayAnimSlotTrait.h
--- a/Engine/Plugins/Experimental/AnimNextAnimGraph/Source/AnimNextAnimGraph/Private/PlayAnim/PlayAnimSlotTrait.h
+++ b/Engine/Plugins/Experimental/AnimNextAnimGraph/Source/AnimNextAnimGraph/Private/PlayAnim/PlayAnimSlotTrait.h
@@ -49,7 +49,7 @@ namespace UE::AnimNext
 	 */
 	struct FPlayAnimSlotTrait : FBaseTrait, IUpdate, IUpdateTraversal, IHierarchy, IDiscreteBlend, ISmoothBlend, IInertializerBlend, IGarbageCollection
 	{
-		DECLARE_ANIM_TRAIT(FPlayAnimSlotTrait, 0x5667d3ee, FBaseTrait)
+		DECLARE_ANIM_TRAIT(FPlayAnimSlotTrait, FBaseTrait)
 
 		// Slot request state
 		enum class EPlayAnimRequestState : uint8
diff --git a/Engine/Plugins/Experimental/AnimNextAnimGraph/Source/AnimNextAnimGraph/Private/Traits/BlendByBool.h b/Engine/Plugins/Experimental/AnimNextAnimGraph/Source/AnimNextAnimGraph/Private/Traits/BlendByBool.h
--- a/Engine/Plugins/Experimental/AnimNextAnimGraph/Source/AnimNextAnimGraph/Private/Traits/BlendByBool.h
+++ b/Engine/Plugins/Experimental/AnimNextAnimGraph/Source/AnimNextAnimGraph/Private/Traits/BlendByBool.h
@@ -45,7 +45,7 @@ namespace UE::AnimNext
 	 */
 	struct FBlendByBoolTrait : FBaseTrait, IUpdate, IUpdateTraversal, IHierarchy, IDiscreteBlend
 	{
-		DECLARE_ANIM_TRAIT(FBlendByBoolTrait, 0x452572ff, FBaseTrait)
+		DECLARE_ANIM_TRAIT(FBlendByBoolTrait, FBaseTrait)
 
 		using FSharedData = FAnimNextBlendByBoolTraitSharedData;
 
diff --git a/Engine/Plugins/Experimental/AnimNextAnimGraph/Source/AnimNextAnimGraph/Private/Traits/BlendInertializer.h b/Engine/Plugins/Experimental/AnimNextAnimGraph/Source/AnimNextAnimGraph/Private/Traits/BlendInertializer.h
--- a/Engine/Plugins/Experimental/AnimNextAnimGraph/Source/AnimNextAnimGraph/Private/Traits/BlendInertializer.h
+++ b/Engine/Plugins/Experimental/AnimNextAnimGraph/Source/AnimNextAnimGraph/Private/Traits/BlendInertializer.h
@@ -39,7 +39,7 @@ namespace UE::AnimNext
 	 */
 	struct FBlendInertializerCoreTrait : FAdditiveTrait, IDiscreteBlend, ISmoothBlend
 	{
-		DECLARE_ANIM_TRAIT(FBlendInertializerCoreTrait, 0x137a22ab, FAdditiveTrait)
+		DECLARE_ANIM_TRAIT(FBlendInertializerCoreTrait, FAdditiveTrait)
 
 		using FSharedData = FAnimNextBlendInertializerCoreTraitSharedData;
 
@@ -62,7 +62,7 @@ namespace UE::AnimNext
 	 */
 	struct FBlendInertializerTrait : FBlendInertializerCoreTrait, IInertializerBlend
 	{
-		DECLARE_ANIM_TRAIT(FBlendInertializerTrait, 0x7ea0bdee, FBlendInertializerCoreTrait)
+		DECLARE_ANIM_TRAIT(FBlendInertializerTrait, FBlendInertializerCoreTrait)
 
 		using FSharedData = FAnimNextBlendInertializerTraitSharedData;
 
diff --git a/Engine/Plugins/Experimental/AnimNextAnimGraph/Source/AnimNextAnimGraph/Private/Traits/BlendLayer.h b/Engine/Plugins/Experimental/AnimNextAnimGraph/Source/AnimNextAnimGraph/Private/Traits/BlendLayer.h
--- a/Engine/Plugins/Experimental/AnimNextAnimGraph/Source/AnimNextAnimGraph/Private/Traits/BlendLayer.h
+++ b/Engine/Plugins/Experimental/AnimNextAnimGraph/Source/AnimNextAnimGraph/Private/Traits/BlendLayer.h
@@ -49,7 +49,7 @@ namespace UE::AnimNext
 	 */
 	 struct FBlendLayerTrait : FBaseTrait, IEvaluate, IUpdate, IUpdateTraversal, IHierarchy, IContinuousBlend
 	 {
-		DECLARE_ANIM_TRAIT(FBlendLayerTrait, 0x57d56b7f, FBaseTrait)
+		DECLARE_ANIM_TRAIT(FBlendLayerTrait, FBaseTrait)
 
 
 		using FSharedData = FAnimNextBlendLayerTraitSharedData;
diff --git a/Engine/Plugins/Experimental/AnimNextAnimGraph/Source/AnimNextAnimGraph/Private/Traits/BlendSmoother.h b/Engine/Plugins/Experimental/AnimNextAnimGraph/Source/AnimNextAnimGraph/Private/Traits/BlendSmoother.h
--- a/Engine/Plugins/Experimental/AnimNextAnimGraph/Source/AnimNextAnimGraph/Private/Traits/BlendSmoother.h
+++ b/Engine/Plugins/Experimental/AnimNextAnimGraph/Source/AnimNextAnimGraph/Private/Traits/BlendSmoother.h
@@ -51,7 +51,7 @@ namespace UE::AnimNext
 	 */
 	struct FBlendSmootherCoreTrait : FAdditiveTrait, IEvaluate, IUpdate, IDiscreteBlend
 	{
-		DECLARE_ANIM_TRAIT(FBlendSmootherCoreTrait, 0x440a43d2, FAdditiveTrait)
+		DECLARE_ANIM_TRAIT(FBlendSmootherCoreTrait, FAdditiveTrait)
 
 		using FSharedData = FAnimNextBlendSmootherCoreTraitSharedData;
 
@@ -101,7 +101,7 @@ namespace UE::AnimNext
 	 */
 	struct FBlendSmootherTrait : FBlendSmootherCoreTrait, ISmoothBlend
 	{
-		DECLARE_ANIM_TRAIT(FBlendSmootherTrait, 0xbaeb537b, FBlendSmootherCoreTrait)
+		DECLARE_ANIM_TRAIT(FBlendSmootherTrait, FBlendSmootherCoreTrait)
 
 		using FSharedData = FAnimNextBlendSmootherTraitSharedData;
 
diff --git a/Engine/Plugins/Experimental/AnimNextAnimGraph/Source/AnimNextAnimGraph/Private/Traits/BlendSmootherPerBone.h b/Engine/Plugins/Experimental/AnimNextAnimGraph/Source/AnimNextAnimGraph/Private/Traits/BlendSmootherPerBone.h
--- a/Engine/Plugins/Experimental/AnimNextAnimGraph/Source/AnimNextAnimGraph/Private/Traits/BlendSmootherPerBone.h
+++ b/Engine/Plugins/Experimental/AnimNextAnimGraph/Source/AnimNextAnimGraph/Private/Traits/BlendSmootherPerBone.h
@@ -35,7 +35,7 @@ namespace UE::AnimNext
 	 */
 	struct FBlendSmootherPerBoneTrait : FAdditiveTrait, IEvaluate, IUpdate, IDiscreteBlend
 	{
-		DECLARE_ANIM_TRAIT(FBlendSmootherPerBoneTrait, 0xb97ffc16, FAdditiveTrait)
+		DECLARE_ANIM_TRAIT(FBlendSmootherPerBoneTrait, FAdditiveTrait)
 
 		using FSharedData = FAnimNextBlendSmootherPerBoneTraitSharedData;
 
diff --git a/Engine/Plugins/Experimental/AnimNextAnimGraph/Source/AnimNextAnimGraph/Private/Traits/BlendSpacePlayer.h b/Engine/Plugins/Experimental/AnimNextAnimGraph/Source/AnimNextAnimGraph/Private/Traits/BlendSpacePlayer.h
--- a/Engine/Plugins/Experimental/AnimNextAnimGraph/Source/AnimNextAnimGraph/Private/Traits/BlendSpacePlayer.h
+++ b/Engine/Plugins/Experimental/AnimNextAnimGraph/Source/AnimNextAnimGraph/Private/Traits/BlendSpacePlayer.h
@@ -68,7 +68,7 @@ namespace UE::AnimNext
 	 */
 	struct FBlendSpacePlayerTrait : FBaseTrait, IEvaluate,  IUpdate, IUpdateTraversal, IHierarchy, ITimeline, IContinuousBlend, IGarbageCollection
 	{
-		DECLARE_ANIM_TRAIT(FBlendSpacePlayerTrait, 0x11a0f4e3, FBaseTrait)
+		DECLARE_ANIM_TRAIT(FBlendSpacePlayerTrait, FBaseTrait)
 
 		using FSharedData = FAnimNextBlendSpacePlayerTraitSharedData;
 
diff --git a/Engine/Plugins/Experimental/AnimNextAnimGraph/Source/AnimNextAnimGraph/Private/Traits/BlendTwoWay.h b/Engine/Plugins/Experimental/AnimNextAnimGraph/Source/AnimNextAnimGraph/Private/Traits/BlendTwoWay.h
--- a/Engine/Plugins/Experimental/AnimNextAnimGraph/Source/AnimNextAnimGraph/Private/Traits/BlendTwoWay.h
+++ b/Engine/Plugins/Experimental/AnimNextAnimGraph/Source/AnimNextAnimGraph/Private/Traits/BlendTwoWay.h
@@ -46,7 +46,7 @@ namespace UE::AnimNext
 	 */
 	struct FBlendTwoWayTrait : FBaseTrait, IEvaluate, IUpdate, IUpdateTraversal, IHierarchy, IContinuousBlend
 	{
-		DECLARE_ANIM_TRAIT(FBlendTwoWayTrait, 0x62f40e2b, FBaseTrait)
+		DECLARE_ANIM_TRAIT(FBlendTwoWayTrait, FBaseTrait)
 
 		using FSharedData = FAnimNextBlendTwoWayTraitSharedData;
 
diff --git a/Engine/Plugins/Experimental/AnimNextAnimGraph/Source/AnimNextAnimGraph/Private/Traits/InertializationDeadBlend.h b/Engine/Plugins/Experimental/AnimNextAnimGraph/Source/AnimNextAnimGraph/Private/Traits/InertializationDeadBlend.h
--- a/Engine/Plugins/Experimental/AnimNextAnimGraph/Source/AnimNextAnimGraph/Private/Traits/InertializationDeadBlend.h
+++ b/Engine/Plugins/Experimental/AnimNextAnimGraph/Source/AnimNextAnimGraph/Private/Traits/InertializationDeadBlend.h
@@ -95,7 +95,7 @@ namespace UE::AnimNext
 	 */
 	struct FInertializationDeadBlendTrait : FAdditiveTrait, IUpdate, IEvaluate
 	{
-		DECLARE_ANIM_TRAIT(FInertializationDeadBlendTrait, 0xf11f2b26, FAdditiveTrait)
+		DECLARE_ANIM_TRAIT(FInertializationDeadBlendTrait, FAdditiveTrait)
 
 		using FSharedData = FAnimNextInertializationDeadBlendTraitSharedData;
 
diff --git a/Engine/Plugins/Experimental/AnimNextAnimGraph/Source/AnimNextAnimGraph/Private/Traits/InputPoseTrait.h b/Engine/Plugins/Experimental/AnimNextAnimGraph/Source/AnimNextAnimGraph/Private/Traits/InputPoseTrait.h
--- a/Engine/Plugins/Experimental/AnimNextAnimGraph/Source/AnimNextAnimGraph/Private/Traits/InputPoseTrait.h
+++ b/Engine/Plugins/Experimental/AnimNextAnimGraph/Source/AnimNextAnimGraph/Private/Traits/InputPoseTrait.h
@@ -30,7 +30,7 @@ namespace UE::AnimNext
 {
 	struct FInputPoseTrait : FBaseTrait, IEvaluate
 	{
-		DECLARE_ANIM_TRAIT(FInputPoseTrait, 0xb2974c37, FBaseTrait)
+		DECLARE_ANIM_TRAIT(FInputPoseTrait, FBaseTrait)
 
 		using FSharedData = FAnimNextInputPoseTraitSharedData;
 
diff --git a/Engine/Plugins/Experimental/AnimNextAnimGraph/Source/AnimNextAnimGraph/Private/Traits/MakeGraphFromAsset.h b/Engine/Plugins/Experimental/AnimNextAnimGraph/Source/AnimNextAnimGraph/Private/Traits/MakeGraphFromAsset.h
--- a/Engine/Plugins/Experimental/AnimNextAnimGraph/Source/AnimNextAnimGraph/Private/Traits/MakeGraphFromAsset.h
+++ b/Engine/Plugins/Experimental/AnimNextAnimGraph/Source/AnimNextAnimGraph/Private/Traits/MakeGraphFromAsset.h
@@ -1,4 +1,4 @@
-﻿// Copyright Epic Games, Inc. All Rights Reserved.
+// Copyright Epic Games, Inc. All Rights Reserved.
 
 #pragma once
 
@@ -23,7 +23,7 @@ namespace UE::AnimNext
 	 */
 	struct FMakeGraphFromAssetTrait : FAdditiveTrait, IGraphFactory
 	{
-		DECLARE_ANIM_TRAIT(FMakeGraphFromAssetTrait, 0x31b4fc8d, FAdditiveTrait)
+		DECLARE_ANIM_TRAIT(FMakeGraphFromAssetTrait, FAdditiveTrait)
 
 		using FSharedData = FAnimNextMakeGraphFromAssetTraitSharedData;
 
diff --git a/Engine/Plugins/Experimental/AnimNextAnimGraph/Source/AnimNextAnimGraph/Private/Traits/NotifyDispatcher.h b/Engine/Plugins/Experimental/AnimNextAnimGraph/Source/AnimNextAnimGraph/Private/Traits/NotifyDispatcher.h
--- a/Engine/Plugins/Experimental/AnimNextAnimGraph/Source/AnimNextAnimGraph/Private/Traits/NotifyDispatcher.h
+++ b/Engine/Plugins/Experimental/AnimNextAnimGraph/Source/AnimNextAnimGraph/Private/Traits/NotifyDispatcher.h
@@ -1,4 +1,4 @@
-﻿// Copyright Epic Games, Inc. All Rights Reserved.
+// Copyright Epic Games, Inc. All Rights Reserved.
 
 #pragma once
 
@@ -69,7 +69,7 @@ namespace UE::AnimNext
 	 */
 	struct FNotifyDispatcherTrait : FAdditiveTrait, ITimeline
 	{
-		DECLARE_ANIM_TRAIT(FNotifyDispatcherTrait, 0x33a2435b, FAdditiveTrait)
+		DECLARE_ANIM_TRAIT(FNotifyDispatcherTrait, FAdditiveTrait)
 
 		using FSharedData = FAnimNextNotifyDispatcherTraitSharedData;
 
diff --git a/Engine/Plugins/Experimental/AnimNextAnimGraph/Source/AnimNextAnimGraph/Private/Traits/Passthrough.h b/Engine/Plugins/Experimental/AnimNextAnimGraph/Source/AnimNextAnimGraph/Private/Traits/Passthrough.h
--- a/Engine/Plugins/Experimental/AnimNextAnimGraph/Source/AnimNextAnimGraph/Private/Traits/Passthrough.h
+++ b/Engine/Plugins/Experimental/AnimNextAnimGraph/Source/AnimNextAnimGraph/Private/Traits/Passthrough.h
@@ -29,7 +29,7 @@ namespace UE::AnimNext
 	 */
 	struct FPassthroughTrait : FBaseTrait, IHierarchy
 	{
-		DECLARE_ANIM_TRAIT(FPassthroughTrait, 0xa4982ed3, FBaseTrait)
+		DECLARE_ANIM_TRAIT(FPassthroughTrait, FBaseTrait)
 
 		using FSharedData = FAnimNextPassthroughSharedData;
 
diff --git a/Engine/Plugins/Experimental/AnimNextAnimGraph/Source/AnimNextAnimGraph/Private/Traits/ReferencePoseTrait.h b/Engine/Plugins/Experimental/AnimNextAnimGraph/Source/AnimNextAnimGraph/Private/Traits/ReferencePoseTrait.h
--- a/Engine/Plugins/Experimental/AnimNextAnimGraph/Source/AnimNextAnimGraph/Private/Traits/ReferencePoseTrait.h
+++ b/Engine/Plugins/Experimental/AnimNextAnimGraph/Source/AnimNextAnimGraph/Private/Traits/ReferencePoseTrait.h
@@ -40,7 +40,7 @@ namespace UE::AnimNext
 	 */
 	struct FReferencePoseTrait : FBaseTrait, IEvaluate
 	{
-		DECLARE_ANIM_TRAIT(FReferencePoseTrait, 0x7508ab89, FBaseTrait)
+		DECLARE_ANIM_TRAIT(FReferencePoseTrait, FBaseTrait)
 
 		using FSharedData = FAnimNextReferencePoseTraitSharedData;
 
diff --git a/Engine/Plugins/Experimental/AnimNextAnimGraph/Source/AnimNextAnimGraph/Private/Traits/SequencePlayer.h b/Engine/Plugins/Experimental/AnimNextAnimGraph/Source/AnimNextAnimGraph/Private/Traits/SequencePlayer.h
--- a/Engine/Plugins/Experimental/AnimNextAnimGraph/Source/AnimNextAnimGraph/Private/Traits/SequencePlayer.h
+++ b/Engine/Plugins/Experimental/AnimNextAnimGraph/Source/AnimNextAnimGraph/Private/Traits/SequencePlayer.h
@@ -55,7 +55,7 @@ namespace UE::AnimNext
 	 */
 	struct FSequencePlayerTrait : FBaseTrait, IEvaluate, IAttributeProvider, ITimeline, IUpdate, IGarbageCollection, INotifySource
 	{
-		DECLARE_ANIM_TRAIT(FSequencePlayerTrait, 0x7a0dc157, FBaseTrait)
+		DECLARE_ANIM_TRAIT(FSequencePlayerTrait, FBaseTrait)
 
 		using FSharedData = FAnimNextSequencePlayerTraitSharedData;
 
diff --git a/Engine/Plugins/Experimental/AnimNextAnimGraph/Source/AnimNextAnimGraph/Private/Traits/SubGraphHost.h b/Engine/Plugins/Experimental/AnimNextAnimGraph/Source/AnimNextAnimGraph/Private/Traits/SubGraphHost.h
--- a/Engine/Plugins/Experimental/AnimNextAnimGraph/Source/AnimNextAnimGraph/Private/Traits/SubGraphHost.h
+++ b/Engine/Plugins/Experimental/AnimNextAnimGraph/Source/AnimNextAnimGraph/Private/Traits/SubGraphHost.h
@@ -49,7 +49,7 @@ namespace UE::AnimNext
 	 */
 	struct FSubGraphHostTrait : FBaseTrait, IUpdate, IUpdateTraversal, IHierarchy, IDiscreteBlend, IGarbageCollection
 	{
-		DECLARE_ANIM_TRAIT(FSubGraphHostTrait, 0xb1085d7b, FBaseTrait)
+		DECLARE_ANIM_TRAIT(FSubGraphHostTrait, FBaseTrait)
 
 		enum class ESlotState : uint8
 		{
diff --git a/Engine/Plugins/Experimental/AnimNextAnimGraph/Source/AnimNextAnimGraph/Private/Traits/SynchronizeUsingGroups.h b/Engine/Plugins/Experimental/AnimNextAnimGraph/Source/AnimNextAnimGraph/Private/Traits/SynchronizeUsingGroups.h
--- a/Engine/Plugins/Experimental/AnimNextAnimGraph/Source/AnimNextAnimGraph/Private/Traits/SynchronizeUsingGroups.h
+++ b/Engine/Plugins/Experimental/AnimNextAnimGraph/Source/AnimNextAnimGraph/Private/Traits/SynchronizeUsingGroups.h
@@ -40,7 +40,7 @@ namespace UE::AnimNext
 	 */
 	struct FSynchronizeUsingGroupsTrait : FAdditiveTrait, IUpdate, IGroupSynchronization, ITimeline
 	{
-		DECLARE_ANIM_TRAIT(FSynchronizeUsingGroupsTrait, 0x6d318931, FAdditiveTrait)
+		DECLARE_ANIM_TRAIT(FSynchronizeUsingGroupsTrait, FAdditiveTrait)
 
 		using FSharedData = FAnimNextSynchronizeUsingGroupsTraitSharedData;
 
diff --git a/Engine/Plugins/Experimental/AnimNextAnimGraph/Source/AnimNextAnimGraph/Public/TraitCore/Trait.h b/Engine/Plugins/Experimental/AnimNextAnimGraph/Source/AnimNextAnimGraph/Public/TraitCore/Trait.h
--- a/Engine/Plugins/Experimental/AnimNextAnimGraph/Source/AnimNextAnimGraph/Public/TraitCore/Trait.h
+++ b/Engine/Plugins/Experimental/AnimNextAnimGraph/Source/AnimNextAnimGraph/Public/TraitCore/Trait.h
@@ -20,10 +20,10 @@ class FArchive;
 class URigVMController;
 struct FRigVMPinInfoArray;
 
-#define ANIM_NEXT_IMPL_DECLARE_ANIM_TRAIT_BASIC(TraitName, TraitNameHash, SuperTraitName) \
+#define ANIM_NEXT_IMPL_DECLARE_ANIM_TRAIT_BASIC(TraitName, SuperTraitName) \
 	using TraitSuper = SuperTraitName; \
 	/* FTrait impl */ \
-	static constexpr UE::AnimNext::FTraitUID TraitUID = UE::AnimNext::FTraitUID(TraitNameHash, TEXT(#TraitName)); \
+	static constexpr UE::AnimNext::FTraitUID TraitUID = UE::AnimNext::FTraitUID::MakeUID(TEXT(#TraitName)); \
 	virtual UE::AnimNext::FTraitUID GetTraitUID() const override { return TraitUID; } \
 	virtual FString GetTraitName() const override { return TEXT(#TraitName); } \
 	static_assert(std::is_base_of<TraitSuper::FSharedData, FSharedData>::value, "Trait shared data must derive from its Super's FSharedData"); \
@@ -52,15 +52,15 @@ struct FRigVMPinInfoArray;
 // Helper macros
 // In the trait class declaration, this macro declares the Super alias and base functions we override
 // We use use FNV1A32 to hash entire TraitName, including 'F' part (Currently must be manually done). 
-#define DECLARE_ANIM_TRAIT(TraitName, TraitNameHash, SuperTraitName) \
-	ANIM_NEXT_IMPL_DECLARE_ANIM_TRAIT_BASIC(TraitName, TraitNameHash, SuperTraitName) \
+#define DECLARE_ANIM_TRAIT(TraitName, SuperTraitName) \
+	ANIM_NEXT_IMPL_DECLARE_ANIM_TRAIT_BASIC(TraitName, SuperTraitName) \
 	ANIM_NEXT_IMPL_DECLARE_ANIM_TRAIT_INSTANCING_SUPPORT() \
 	ANIM_NEXT_IMPL_DECLARE_ANIM_TRAIT_INTERFACE_SUPPORT() \
 	ANIM_NEXT_IMPL_DECLARE_ANIM_TRAIT_EVENT_SUPPORT() \
 	ANIM_NEXT_IMPL_DECLARE_ANIM_TRAIT_LATENT_PROPERTY_SUPPORT()
 
-#define DECLARE_ABSTRACT_ANIM_TRAIT(TraitName, TraitNameHash, SuperTraitName) \
-	ANIM_NEXT_IMPL_DECLARE_ANIM_TRAIT_BASIC(TraitName, TraitNameHash, SuperTraitName)
+#define DECLARE_ABSTRACT_ANIM_TRAIT(TraitName, SuperTraitName) \
+	ANIM_NEXT_IMPL_DECLARE_ANIM_TRAIT_BASIC(TraitName, SuperTraitName)
 
 // In the trait cpp, these three macros implement the base functionality
 // 
@@ -346,7 +346,7 @@ namespace UE::AnimNext
 		// The globally unique UID for this trait
 		// Derived types will have their own TraitUID member that hides/aliases/shadows this one
 		// @see DECLARE_ANIM_TRAIT
-		static constexpr FTraitUID TraitUID = FTraitUID(0x4d97bae7, TEXT("FTrait"));
+		static constexpr FTraitUID TraitUID = FTraitUID::MakeUID(TEXT("FTrait"));
 
 		// Returns the globally unique UID for this trait
 		virtual FTraitUID GetTraitUID() const { return TraitUID; };
@@ -460,7 +460,7 @@ namespace UE::AnimNext
 	// Base class for base traits that are standalone
 	struct ANIMNEXTANIMGRAPH_API FBaseTrait : FTrait
 	{
-		DECLARE_ABSTRACT_ANIM_TRAIT(FBaseTrait, 0xbb4b35d0, FTrait)
+		DECLARE_ABSTRACT_ANIM_TRAIT(FBaseTrait, FTrait)
 
 		virtual ETraitMode GetTraitMode() const override { return ETraitMode::Base; }
 	};
@@ -468,7 +468,7 @@ namespace UE::AnimNext
 	// Base class for additive traits that override behavior of other traits
 	struct ANIMNEXTANIMGRAPH_API FAdditiveTrait : FTrait
 	{
-		DECLARE_ABSTRACT_ANIM_TRAIT(FAdditiveTrait, 0x0dbd1b3f, FTrait)
+		DECLARE_ABSTRACT_ANIM_TRAIT(FAdditiveTrait, FTrait)
 
 		virtual ETraitMode GetTraitMode() const override { return ETraitMode::Additive; }
 
diff --git a/Engine/Plugins/Experimental/AnimNextAnimGraph/Source/AnimNextAnimGraph/Public/TraitCore/TraitUID.h b/Engine/Plugins/Experimental/AnimNextAnimGraph/Source/AnimNextAnimGraph/Public/TraitCore/TraitUID.h
--- a/Engine/Plugins/Experimental/AnimNextAnimGraph/Source/AnimNextAnimGraph/Public/TraitCore/TraitUID.h
+++ b/Engine/Plugins/Experimental/AnimNextAnimGraph/Source/AnimNextAnimGraph/Public/TraitCore/TraitUID.h
@@ -3,6 +3,7 @@
 #pragma once
 
 #include "CoreMinimal.h"
+#include "Hash/Fnv.h"
 
 namespace UE::AnimNext
 {
@@ -12,11 +13,9 @@ namespace UE::AnimNext
 	/**
 	 * FTraitUID
 	 *
-	 * Encapsulates an trait global UID.
-	 * The string is exposed in non-shipping builds for logging and debugging purposes.
-	 * The UID should be generated from the provided string using FNV1a with 32 bits.
+	 * Encapsulates a trait global UID.
 	 *
-	 * The whole struct is meant to be 'constexpr' to allow inlining.
+	 * The whole struct is meant to be 'constexpr' to allow inlining within the assembly as constants.
 	 */
 	struct FTraitUID final
 	{
@@ -29,17 +28,24 @@ namespace UE::AnimNext
 		{
 		}
 
-		// Constructs a trait UID
-		explicit constexpr FTraitUID(FTraitUIDRaw InUID, const TCHAR* InTraitName = TEXT("<Unknown Trait Name>")) noexcept
+		// Constructs a trait UID from its raw value
+		explicit constexpr FTraitUID(FTraitUIDRaw InUID) noexcept
 			: UID(InUID)
 #if !UE_BUILD_SHIPPING
-			, TraitName(InTraitName)
+			, TraitName(TEXT("<Unknown Trait Name>"))
 #endif
 		{
 		}
 
+		// Constructs a trait UID from a string literal
+		template<typename CharType, int32 Len>
+		static UE_CONSTEVAL FTraitUID MakeUID(const CharType(&InTraitName)[Len]) noexcept
+		{
+			return FTraitUID(UE::HashStringFNV1a<FTraitUIDRaw>(MakeStringView(InTraitName, Len - 1)), InTraitName);
+		}
+
 #if !UE_BUILD_SHIPPING
-		// Returns a literal string to the interface name
+		// Returns a string literal to the trait name
 		constexpr const TCHAR* GetTraitName() const noexcept { return TraitName; }
 #endif
 
@@ -52,6 +58,15 @@ namespace UE::AnimNext
 	private:
 		static constexpr FTraitUIDRaw INVALID_UID = 0;
 
+		// Constructs a trait UID
+		explicit constexpr FTraitUID(FTraitUIDRaw InUID, const TCHAR* InTraitName) noexcept
+			: UID(InUID)
+#if !UE_BUILD_SHIPPING
+			, TraitName(InTraitName)
+#endif
+		{
+		}
+
 		FTraitUIDRaw	UID;
 
 #if !UE_BUILD_SHIPPING
diff --git a/Engine/Plugins/Experimental/AnimNextAnimGraph/Source/AnimNextAnimGraph/Public/Traits/AlphaBlendTrait.h b/Engine/Plugins/Experimental/AnimNextAnimGraph/Source/AnimNextAnimGraph/Public/Traits/AlphaBlendTrait.h
--- a/Engine/Plugins/Experimental/AnimNextAnimGraph/Source/AnimNextAnimGraph/Public/Traits/AlphaBlendTrait.h
+++ b/Engine/Plugins/Experimental/AnimNextAnimGraph/Source/AnimNextAnimGraph/Public/Traits/AlphaBlendTrait.h
@@ -58,7 +58,7 @@ struct FAlphaBlendTrait : FAdditiveTrait, IUpdate, IEvaluate
 {
 	typedef FBaseTrait Super;
 
-	DECLARE_ANIM_TRAIT(FAlphaBlendTrait, 0xa7c496ca, FAdditiveTrait)
+	DECLARE_ANIM_TRAIT(FAlphaBlendTrait, FAdditiveTrait)
 
 	using FSharedData = FAlphaBlendTraitSharedData;
 		
diff --git a/Engine/Plugins/Experimental/AnimNextAnimGraph/Source/AnimNextAnimGraphTestSuite/Private/AnimNextAnimGraphTraitEventTest.cpp b/Engine/Plugins/Experimental/AnimNextAnimGraph/Source/AnimNextAnimGraphTestSuite/Private/AnimNextAnimGraphTraitEventTest.cpp
--- a/Engine/Plugins/Experimental/AnimNextAnimGraph/Source/AnimNextAnimGraphTestSuite/Private/AnimNextAnimGraphTraitEventTest.cpp
+++ b/Engine/Plugins/Experimental/AnimNextAnimGraph/Source/AnimNextAnimGraphTestSuite/Private/AnimNextAnimGraphTraitEventTest.cpp
@@ -28,7 +28,7 @@ namespace UE::AnimNext
 {
 	struct FTraitGraphTest_EventAB_NoChildren : FBaseTrait
 	{
-		DECLARE_ANIM_TRAIT(FTraitGraphTest_EventAB_NoChildren, 0xfe325e9e, FBaseTrait)
+		DECLARE_ANIM_TRAIT(FTraitGraphTest_EventAB_NoChildren, FBaseTrait)
 
 		ETraitStackPropagation OnEventA(const FExecutionContext& Context, FTraitBinding& Binding, FTraitAnimGraphTest_EventA& Event) const
 		{
@@ -53,7 +53,7 @@ namespace UE::AnimNext
 
 	struct FTraitGraphTest_EventAB_OneChild : FBaseTrait, IHierarchy
 	{
-		DECLARE_ANIM_TRAIT(FTraitGraphTest_EventAB_OneChild, 0xfc122a80, FBaseTrait)
+		DECLARE_ANIM_TRAIT(FTraitGraphTest_EventAB_OneChild, FBaseTrait)
 
 		using FSharedData = FTraitWithOneChildSharedData;
 
@@ -128,7 +128,7 @@ namespace UE::AnimNext
 
 	struct FTraitGraphTest_EventAB_TwoChildren : FBaseTrait, IHierarchy
 	{
-		DECLARE_ANIM_TRAIT(FTraitGraphTest_EventAB_TwoChildren, 0xa440e33f, FBaseTrait)
+		DECLARE_ANIM_TRAIT(FTraitGraphTest_EventAB_TwoChildren, FBaseTrait)
 
 		using FSharedData = FTraitWithChildrenSharedData;
 
diff --git a/Engine/Plugins/Experimental/AnimNextAnimGraph/Source/AnimNextAnimGraphTestSuite/Private/AnimNextAnimGraphTraitGraphTest.cpp b/Engine/Plugins/Experimental/AnimNextAnimGraph/Source/AnimNextAnimGraphTestSuite/Private/AnimNextAnimGraphTraitGraphTest.cpp
--- a/Engine/Plugins/Experimental/AnimNextAnimGraph/Source/AnimNextAnimGraphTestSuite/Private/AnimNextAnimGraphTraitGraphTest.cpp
+++ b/Engine/Plugins/Experimental/AnimNextAnimGraph/Source/AnimNextAnimGraphTestSuite/Private/AnimNextAnimGraphTraitGraphTest.cpp
@@ -36,7 +36,7 @@ namespace UE::AnimNext
 {
 	struct FTestTrait final : FBaseTrait, IEvaluate, IUpdate
 	{
-		DECLARE_ANIM_TRAIT(FTestTrait, 0x41cecb7f, FBaseTrait)
+		DECLARE_ANIM_TRAIT(FTestTrait, FBaseTrait)
 
 		using FSharedData = FTestTraitSharedData;
 
@@ -86,7 +86,7 @@ namespace UE::AnimNext
 
 	struct FTestBasicTrait final : FBaseTrait
 	{
-		DECLARE_ANIM_TRAIT(FTestBasicTrait, 0x24ce4372, FBaseTrait)
+		DECLARE_ANIM_TRAIT(FTestBasicTrait, FBaseTrait)
 
 		using FSharedData = FTestTraitSharedData;
 	};
diff --git a/Engine/Plugins/Experimental/AnimNextAnimGraph/Source/AnimNextAnimGraphTestSuite/Private/AnimNextAnimGraphTraitInterfacesTest.cpp b/Engine/Plugins/Experimental/AnimNextAnimGraph/Source/AnimNextAnimGraphTestSuite/Private/AnimNextAnimGraphTraitInterfacesTest.cpp
--- a/Engine/Plugins/Experimental/AnimNextAnimGraph/Source/AnimNextAnimGraphTestSuite/Private/AnimNextAnimGraphTraitInterfacesTest.cpp
+++ b/Engine/Plugins/Experimental/AnimNextAnimGraph/Source/AnimNextAnimGraphTestSuite/Private/AnimNextAnimGraphTraitInterfacesTest.cpp
@@ -45,7 +45,7 @@ namespace UE::AnimNext
 
 	struct FTraitWithNoChildren : FBaseTrait, IUpdate, IEvaluate
 	{
-		DECLARE_ANIM_TRAIT(FTraitWithNoChildren, 0x79080a53, FBaseTrait)
+		DECLARE_ANIM_TRAIT(FTraitWithNoChildren, FBaseTrait)
 
 		// IUpdate impl
 		virtual void PreUpdate(FUpdateTraversalContext& Context, const TTraitBinding<IUpdate>& Binding, const FTraitUpdateState& TraitState) const override
@@ -101,7 +101,7 @@ namespace UE::AnimNext
 	// This trait does not update or evaluate
 	struct FTraitWithOneChild : FBaseTrait, IHierarchy
 	{
-		DECLARE_ANIM_TRAIT(FTraitWithOneChild, 0xab680b35, FBaseTrait)
+		DECLARE_ANIM_TRAIT(FTraitWithOneChild, FBaseTrait)
 
 		using FSharedData = FTraitWithOneChildSharedData;
 
@@ -138,7 +138,7 @@ namespace UE::AnimNext
 
 	struct FTraitWithChildren : FBaseTrait, IHierarchy, IUpdate, IUpdateTraversal, IEvaluate
 	{
-		DECLARE_ANIM_TRAIT(FTraitWithChildren, 0x4b296948, FBaseTrait)
+		DECLARE_ANIM_TRAIT(FTraitWithChildren, FBaseTrait)
 
 		using FSharedData = FTraitWithChildrenSharedData;
 
@@ -284,7 +284,7 @@ namespace UE::AnimNext
 	// Adds a scoped tag
 	struct FScopedTagTrait : FAdditiveTrait, IScopedTagInterface, IUpdate
 	{
-		DECLARE_ANIM_TRAIT(FScopedTagTrait, 0x4b296948, FAdditiveTrait)
+		DECLARE_ANIM_TRAIT(FScopedTagTrait, FAdditiveTrait)
 
 		// IScopedTagInterface impl
 		virtual FName GetTag(const FExecutionContext& Context, const TTraitBinding<IScopedTagInterface>& Binding) const override
@@ -322,7 +322,7 @@ namespace UE::AnimNext
 	// Tests if we have a scoped tag
 	struct FTestScopedTagTrait : FAdditiveTrait, IUpdate
 	{
-		DECLARE_ANIM_TRAIT(FTestScopedTagTrait, 0x922bd23d, FAdditiveTrait)
+		DECLARE_ANIM_TRAIT(FTestScopedTagTrait, FAdditiveTrait)
 
 		// IUpdate impl
 		virtual void PreUpdate(FUpdateTraversalContext& Context, const TTraitBinding<IUpdate>& Binding, const FTraitUpdateState& TraitState) const override
diff --git a/Engine/Plugins/Experimental/AnimNextAnimGraph/Source/AnimNextAnimGraphTestSuite/Private/AnimNextTraitBaseTest.cpp b/Engine/Plugins/Experimental/AnimNextAnimGraph/Source/AnimNextAnimGraphTestSuite/Private/AnimNextTraitBaseTest.cpp
--- a/Engine/Plugins/Experimental/AnimNextAnimGraph/Source/AnimNextAnimGraphTestSuite/Private/AnimNextTraitBaseTest.cpp
+++ b/Engine/Plugins/Experimental/AnimNextAnimGraph/Source/AnimNextAnimGraphTestSuite/Private/AnimNextTraitBaseTest.cpp
@@ -126,7 +126,7 @@ namespace UE::AnimNext
 
 	struct FTraitA_Base : FBaseTrait, IInterfaceA
 	{
-		DECLARE_ANIM_TRAIT(FTraitA_Base, 0xd0f2ad3a, FBaseTrait)
+		DECLARE_ANIM_TRAIT(FTraitA_Base, FBaseTrait)
 
 		using FSharedData = FTraitA_BaseSharedData;
 
@@ -168,7 +168,7 @@ namespace UE::AnimNext
 
 	struct FTraitAB_Add : FAdditiveTrait, IInterfaceA, IInterfaceB
 	{
-		DECLARE_ANIM_TRAIT(FTraitAB_Add, 0x0996ac7c, FAdditiveTrait)
+		DECLARE_ANIM_TRAIT(FTraitAB_Add, FAdditiveTrait)
 
 		using FSharedData = FTraitAB_AddSharedData;
 
@@ -216,7 +216,7 @@ namespace UE::AnimNext
 
 	struct FTraitAC_Add : FAdditiveTrait, IInterfaceA, IInterfaceC
 	{
-		DECLARE_ANIM_TRAIT(FTraitAC_Add, 0xdbcc694f, FAdditiveTrait)
+		DECLARE_ANIM_TRAIT(FTraitAC_Add, FAdditiveTrait)
 
 		using FSharedData = FTraitAC_AddSharedData;
 
@@ -265,7 +265,7 @@ namespace UE::AnimNext
 
 	struct FTraitSerialization_Base : FBaseTrait, IInterfaceA
 	{
-		DECLARE_ANIM_TRAIT(FTraitSerialization_Base, 0x2ee50577, FBaseTrait)
+		DECLARE_ANIM_TRAIT(FTraitSerialization_Base, FBaseTrait)
 
 		using FSharedData = FTraitSerialization_BaseSharedData;
 	};
@@ -281,7 +281,7 @@ namespace UE::AnimNext
 
 	struct FTraitSerialization_Add : FAdditiveTrait, IInterfaceB
 	{
-		DECLARE_ANIM_TRAIT(FTraitSerialization_Add, 0x90996001, FAdditiveTrait)
+		DECLARE_ANIM_TRAIT(FTraitSerialization_Add, FAdditiveTrait)
 
 		using FSharedData = FTraitSerialization_AddSharedData;
 	};
@@ -297,7 +297,7 @@ namespace UE::AnimNext
 
 	struct FTraitNativeSerialization_Add : FAdditiveTrait, IInterfaceC
 	{
-		DECLARE_ANIM_TRAIT(FTraitNativeSerialization_Add, 0xe76d44dc, FAdditiveTrait)
+		DECLARE_ANIM_TRAIT(FTraitNativeSerialization_Add, FAdditiveTrait)
 
 		using FSharedData = FTraitNativeSerialization_AddSharedData;
 	};
diff --git a/Engine/Plugins/Experimental/AnimNextAnimGraph/Source/AnimNextAnimGraphTestSuite/Private/AnimNextTraitEventTest.cpp b/Engine/Plugins/Experimental/AnimNextAnimGraph/Source/AnimNextAnimGraphTestSuite/Private/AnimNextTraitEventTest.cpp
--- a/Engine/Plugins/Experimental/AnimNextAnimGraph/Source/AnimNextAnimGraphTestSuite/Private/AnimNextTraitEventTest.cpp
+++ b/Engine/Plugins/Experimental/AnimNextAnimGraph/Source/AnimNextAnimGraphTestSuite/Private/AnimNextTraitEventTest.cpp
@@ -29,7 +29,7 @@ namespace UE::AnimNext
 {
 	struct FTraitCoreTest_EventAB_Base : FBaseTrait
 	{
-		DECLARE_ANIM_TRAIT(FTraitCoreTest_EventAB_Base, 0x6bab1b28, FBaseTrait)
+		DECLARE_ANIM_TRAIT(FTraitCoreTest_EventAB_Base, FBaseTrait)
 
 		// We support static functions
 		static ETraitStackPropagation OnEventA(const FExecutionContext& Context, FTraitBinding& Binding, FTraitCoreTest_EventA& Event)
@@ -56,7 +56,7 @@ namespace UE::AnimNext
 
 	struct FTraitCoreTest_EventAB_Add : FAdditiveTrait
 	{
-		DECLARE_ANIM_TRAIT(FTraitCoreTest_EventAB_Add, 0x193709c4, FAdditiveTrait)
+		DECLARE_ANIM_TRAIT(FTraitCoreTest_EventAB_Add, FAdditiveTrait)
 
 		ETraitStackPropagation OnEventA(const FExecutionContext& Context, FTraitBinding& Binding, FTraitCoreTest_EventA& Event) const
 		{
@@ -81,7 +81,7 @@ namespace UE::AnimNext
 
 	struct FTraitCoreTest_EventA_Add : FAdditiveTrait
 	{
-		DECLARE_ANIM_TRAIT(FTraitCoreTest_EventA_Add, 0xcf0402da, FAdditiveTrait)
+		DECLARE_ANIM_TRAIT(FTraitCoreTest_EventA_Add, FAdditiveTrait)
 
 		ETraitStackPropagation OnEventA(const FExecutionContext& Context, FTraitBinding& Binding, FTraitCoreTest_EventA& Event) const
 		{
@@ -99,7 +99,7 @@ namespace UE::AnimNext
 
 	struct FTraitCoreTest_EventB_Add : FAdditiveTrait
 	{
-		DECLARE_ANIM_TRAIT(FTraitCoreTest_EventB_Add, 0x24ba435b, FAdditiveTrait)
+		DECLARE_ANIM_TRAIT(FTraitCoreTest_EventB_Add, FAdditiveTrait)
 
 		ETraitStackPropagation OnEventB(const FExecutionContext& Context, FTraitBinding& Binding, FTraitCoreTest_EventB& Event) const
 		{
diff --git a/Engine/Plugins/Experimental/AnimNextAnimGraph/Source/AnimNextAnimGraphUncookedOnly/Private/Graph/AnimNextAnimationGraph_EditorData.cpp b/Engine/Plugins/Experimental/AnimNextAnimGraph/Source/AnimNextAnimGraphUncookedOnly/Private/Graph/AnimNextAnimationGraph_EditorData.cpp
--- a/Engine/Plugins/Experimental/AnimNextAnimGraph/Source/AnimNextAnimGraphUncookedOnly/Private/Graph/AnimNextAnimationGraph_EditorData.cpp
+++ b/Engine/Plugins/Experimental/AnimNextAnimGraph/Source/AnimNextAnimGraphUncookedOnly/Private/Graph/AnimNextAnimationGraph_EditorData.cpp
@@ -278,7 +278,7 @@ namespace UE::AnimNext::UncookedOnly::Private
 
 			FString DefaultValue;
 			{
-				const UE::AnimNext::FTraitUID ReferencePoseTraitUID(0x7508ab89);	// Trait header is private, reference by UID directly
+				constexpr UE::AnimNext::FTraitUID ReferencePoseTraitUID = UE::AnimNext::FTraitUID::MakeUID(TEXT("FReferencePoseTrait"));	// Trait header is private, reference by UID directly
 				const FTrait* Trait = FTraitRegistry::Get().Find(ReferencePoseTraitUID);
 				check(Trait != nullptr);
 
@@ -382,7 +382,7 @@ namespace UE::AnimNext::UncookedOnly::Private
 			URigVMUnitNode* VMNode = VMController->AddUnitNode(FRigUnit_AnimNextTraitStack::StaticStruct(), FRigVMStruct::ExecuteName, FVector2D(0.0f, 0.0f), FString(), false);
 
 			UAnimNextController* AnimNextController = CastChecked<UAnimNextController>(VMController);
-			constexpr UE::AnimNext::FTraitUID ReferencePoseTraitUID(0x7508ab89); // Trait header is private, reference by UID directly
+			constexpr UE::AnimNext::FTraitUID ReferencePoseTraitUID = UE::AnimNext::FTraitUID::MakeUID(TEXT("FReferencePoseTrait"));	// Trait header is private, reference by UID directly
 			const FName RigVMTraitName =  AnimNextController->AddTraitByName(VMNode->GetFName(), *UE::AnimNext::FTraitRegistry::Get().Find(ReferencePoseTraitUID)->GetTraitName(), INDEX_NONE, TEXT(""), false);
 		
 			check(RigVMTraitName != NAME_None);
diff --git a/Engine/Plugins/Experimental/AnimNextChooser/Source/AnimNextChooser/Private/MakeGraphFromChooser.h b/Engine/Plugins/Experimental/AnimNextChooser/Source/AnimNextChooser/Private/MakeGraphFromChooser.h
--- a/Engine/Plugins/Experimental/AnimNextChooser/Source/AnimNextChooser/Private/MakeGraphFromChooser.h
+++ b/Engine/Plugins/Experimental/AnimNextChooser/Source/AnimNextChooser/Private/MakeGraphFromChooser.h
@@ -1,4 +1,4 @@
-﻿// Copyright Epic Games, Inc. All Rights Reserved.
+// Copyright Epic Games, Inc. All Rights Reserved.
 
 #pragma once
 
@@ -36,7 +36,7 @@ namespace UE::AnimNext
 	 */
 	struct FMakeGraphFromChooserTrait : FAdditiveTrait, IGraphFactory
 	{
-		DECLARE_ANIM_TRAIT(FMakeGraphFromChooserTrait, 0xe8c28546, FAdditiveTrait)
+		DECLARE_ANIM_TRAIT(FMakeGraphFromChooserTrait, FAdditiveTrait)
 
 		using FSharedData = FAnimNextMakeGraphFromChooserTraitSharedData;
 
diff --git a/Engine/Plugins/Experimental/AnimNextControlRig/Source/AnimNextControlRig/Internal/ControlRigTrait.h b/Engine/Plugins/Experimental/AnimNextControlRig/Source/AnimNextControlRig/Internal/ControlRigTrait.h
--- a/Engine/Plugins/Experimental/AnimNextControlRig/Source/AnimNextControlRig/Internal/ControlRigTrait.h
+++ b/Engine/Plugins/Experimental/AnimNextControlRig/Source/AnimNextControlRig/Internal/ControlRigTrait.h
@@ -173,7 +173,7 @@ struct FControlRigTrait : FBaseTrait, IEvaluate, IUpdate, IHierarchy, IGarbageCo
 {
 	typedef FBaseTrait Super;
 
-	ANIM_NEXT_IMPL_DECLARE_ANIM_TRAIT_BASIC(FControlRigTrait, 0xf0ebdcb0, FBaseTrait)
+	ANIM_NEXT_IMPL_DECLARE_ANIM_TRAIT_BASIC(FControlRigTrait, FBaseTrait)
 	ANIM_NEXT_IMPL_DECLARE_ANIM_TRAIT_INSTANCING_SUPPORT()
 	ANIM_NEXT_IMPL_DECLARE_ANIM_TRAIT_INTERFACE_SUPPORT()
 	ANIM_NEXT_IMPL_DECLARE_ANIM_TRAIT_EVENT_SUPPORT()
diff --git a/Engine/Plugins/Experimental/AnimNextPoseSearch/Source/AnimNextPoseSearch/Private/HistoryCollectorTrait.h b/Engine/Plugins/Experimental/AnimNextPoseSearch/Source/AnimNextPoseSearch/Private/HistoryCollectorTrait.h
--- a/Engine/Plugins/Experimental/AnimNextPoseSearch/Source/AnimNextPoseSearch/Private/HistoryCollectorTrait.h
+++ b/Engine/Plugins/Experimental/AnimNextPoseSearch/Source/AnimNextPoseSearch/Private/HistoryCollectorTrait.h
@@ -130,7 +130,7 @@ namespace UE::AnimNext
 {
 	struct FHistoryCollectorTrait : FAdditiveTrait, IUpdate, IEvaluate, IPoseHistory
 	{
-		DECLARE_ANIM_TRAIT(FHistoryCollectorTrait, 0xf1572b26, FAdditiveTrait)
+		DECLARE_ANIM_TRAIT(FHistoryCollectorTrait, FAdditiveTrait)
 
 		using FSharedData = FAnimNextHistoryCollectorTraitSharedData;
 
diff --git a/Engine/Plugins/Experimental/AnimNextPoseSearch/Source/AnimNextPoseSearch/Private/MotionMatchingTrait.h b/Engine/Plugins/Experimental/AnimNextPoseSearch/Source/AnimNextPoseSearch/Private/MotionMatchingTrait.h
--- a/Engine/Plugins/Experimental/AnimNextPoseSearch/Source/AnimNextPoseSearch/Private/MotionMatchingTrait.h
+++ b/Engine/Plugins/Experimental/AnimNextPoseSearch/Source/AnimNextPoseSearch/Private/MotionMatchingTrait.h
@@ -96,7 +96,7 @@ namespace UE::AnimNext
 
 struct FMotionMatchingTrait : FAdditiveTrait, IUpdate
 {
-	DECLARE_ANIM_TRAIT(FMotionMatchingTrait, 0xf1572ff4, FAdditiveTrait)
+	DECLARE_ANIM_TRAIT(FMotionMatchingTrait, FAdditiveTrait)
 
 	using FSharedData = FMotionMatchingTraitSharedData;
 		
diff --git a/Engine/Plugins/Experimental/AnimNextStateTree/Source/AnimNextStateTree/Internal/AnimStateTreeTrait.h b/Engine/Plugins/Experimental/AnimNextStateTree/Source/AnimNextStateTree/Internal/AnimStateTreeTrait.h
--- a/Engine/Plugins/Experimental/AnimNextStateTree/Source/AnimNextStateTree/Internal/AnimStateTreeTrait.h
+++ b/Engine/Plugins/Experimental/AnimNextStateTree/Source/AnimNextStateTree/Internal/AnimStateTreeTrait.h
@@ -38,7 +38,7 @@ namespace UE::AnimNext
 {
 struct FStateTreeTrait : FAdditiveTrait, IUpdate, IGarbageCollection
 {
-	DECLARE_ANIM_TRAIT(FStateTreeTrait, 0x7a0dc697, FAdditiveTrait)
+	DECLARE_ANIM_TRAIT(FStateTreeTrait, FAdditiveTrait)
 
 	using FSharedData = FAnimNextStateTreeTraitSharedData;
 
diff --git a/Engine/Plugins/Experimental/AnimNextWarping/Source/AnimNextWarping/Private/SteeringTrait.h b/Engine/Plugins/Experimental/AnimNextWarping/Source/AnimNextWarping/Private/SteeringTrait.h
--- a/Engine/Plugins/Experimental/AnimNextWarping/Source/AnimNextWarping/Private/SteeringTrait.h
+++ b/Engine/Plugins/Experimental/AnimNextWarping/Source/AnimNextWarping/Private/SteeringTrait.h
@@ -89,7 +89,7 @@ namespace UE::AnimNext
 
 struct FSteeringTrait : FAdditiveTrait, IUpdate, IEvaluate
 {
-	DECLARE_ANIM_TRAIT(FSteeringTrait, 0x387fdf2a, FAdditiveTrait)
+	DECLARE_ANIM_TRAIT(FSteeringTrait, FAdditiveTrait)
 
 	using FSharedData = FSteeringTraitSharedData;
 
diff --git a/Engine/Plugins/Experimental/EvaluationNotifies/Source/Runtime/Private/AnimNext/EvaluationNotifiesTrait.h b/Engine/Plugins/Experimental/EvaluationNotifies/Source/Runtime/Private/AnimNext/EvaluationNotifiesTrait.h
--- a/Engine/Plugins/Experimental/EvaluationNotifies/Source/Runtime/Private/AnimNext/EvaluationNotifiesTrait.h
+++ b/Engine/Plugins/Experimental/EvaluationNotifies/Source/Runtime/Private/AnimNext/EvaluationNotifiesTrait.h
@@ -41,8 +41,7 @@ namespace UE::AnimNext
  */
 struct EVALUATIONNOTIFIESRUNTIME_API FEvaluationNotifiesTrait : FAdditiveTrait, ITimeline, IUpdate, IEvaluate
 {
-	// @TODO: Hash we use use FNV1A32. Update macro to mention that and also comment if we still need it (?)
-	DECLARE_ANIM_TRAIT(FEvaluationNotifiesTrait, 0x4aa0ccf2, FAdditiveTrait)
+	DECLARE_ANIM_TRAIT(FEvaluationNotifiesTrait, FAdditiveTrait)
 
 	using FSharedData = FEvaluationNotifiesTraitSharedData;
 
diff --git a/Engine/Plugins/Experimental/RigLogicAnimNext/Source/RigLogicAnimNext/Public/RigLogicTrait.h b/Engine/Plugins/Experimental/RigLogicAnimNext/Source/RigLogicAnimNext/Public/RigLogicTrait.h
--- a/Engine/Plugins/Experimental/RigLogicAnimNext/Source/RigLogicAnimNext/Public/RigLogicTrait.h
+++ b/Engine/Plugins/Experimental/RigLogicAnimNext/Source/RigLogicAnimNext/Public/RigLogicTrait.h
@@ -46,7 +46,7 @@ namespace UE::AnimNext
 	 */
 	struct RIGLOGICANIMNEXT_API FRigLogicTrait : FBaseTrait, IEvaluate, IUpdate, IUpdateTraversal, IHierarchy
 	{
-		DECLARE_ANIM_TRAIT(FRigLogicTrait, 0x4311e25e, FBaseTrait)
+		DECLARE_ANIM_TRAIT(FRigLogicTrait, FBaseTrait)
 
 		using FSharedData = FAnimNextRigLogicTraitSharedData;
 
