AnimNext: Prevent notify dispatcher from firing notifies with a null mesh skeletal component

Also now defaults to the host skeletal mesh component if SetNotifyContext is not called, so in simple situations this is no longer required

diff --git a/Engine/Plugins/Experimental/AnimNextAnimGraph/Source/AnimNextAnimGraph/Private/Graph/RigUnit_AnimNextSetNotifyContext.cpp b/Engine/Plugins/Experimental/AnimNextAnimGraph/Source/AnimNextAnimGraph/Private/Graph/RigUnit_AnimNextSetNotifyContext.cpp
--- a/Engine/Plugins/Experimental/AnimNextAnimGraph/Source/AnimNextAnimGraph/Private/Graph/RigUnit_AnimNextSetNotifyContext.cpp
+++ b/Engine/Plugins/Experimental/AnimNextAnimGraph/Source/AnimNextAnimGraph/Private/Graph/RigUnit_AnimNextSetNotifyContext.cpp
@@ -1,6 +1,8 @@
 ﻿// Copyright Epic Games, Inc. All Rights Reserved.
 
 #include "RigUnit_AnimNextSetNotifyContext.h"
+
+#include "Components/SkeletalMeshComponent.h"
 #include "Traits/NotifyDispatcher.h"
 
 FRigUnit_AnimNextSetNotifyContext_Execute()
@@ -8,4 +10,5 @@ FRigUnit_AnimNextSetNotifyContext_Execute()
 	FAnimNextModuleInstance& ModuleInstance = ExecuteContext.GetContextData<FAnimNextModuleContextData>().GetModuleInstance();
 	FAnimNextNotifyDispatcherComponent& NotifyDispatcher = ModuleInstance.GetComponent<FAnimNextNotifyDispatcherComponent>();
 	NotifyDispatcher.SkeletalMeshComponent = SkeletalMeshComponent;
+	NotifyDispatcher.NotifyQueue.PredictedLODLevel = SkeletalMeshComponent ? SkeletalMeshComponent->GetPredictedLODLevel() : 0;
 }
\ No newline at end of file
diff --git a/Engine/Plugins/Experimental/AnimNextAnimGraph/Source/AnimNextAnimGraph/Private/Traits/NotifyDispatcher.cpp b/Engine/Plugins/Experimental/AnimNextAnimGraph/Source/AnimNextAnimGraph/Private/Traits/NotifyDispatcher.cpp
--- a/Engine/Plugins/Experimental/AnimNextAnimGraph/Source/AnimNextAnimGraph/Private/Traits/NotifyDispatcher.cpp
+++ b/Engine/Plugins/Experimental/AnimNextAnimGraph/Source/AnimNextAnimGraph/Private/Traits/NotifyDispatcher.cpp
@@ -8,6 +8,7 @@
 #include "Components/SkeletalMeshComponent.h"
 #include "Graph/AnimNextGraphInstance.h"
 #include "Module/AnimNextModuleInstance.h"
+#include "Module/AnimNextSkeletalMeshComponentReferenceComponent.h"
 
 namespace UE::AnimNext
 {
@@ -110,6 +111,15 @@ namespace UE::AnimNext
 	}
 }
 
+void FAnimNextNotifyDispatcherComponent::OnInitialize()
+{
+	// Default to the current host's skeletal mesh component, if any
+	FAnimNextModuleInstance& ModuleInstance = GetModuleInstance();
+	const FAnimNextSkeletalMeshComponentReferenceComponent& ComponentReference = ModuleInstance.GetComponent<FAnimNextSkeletalMeshComponentReferenceComponent>();
+	SkeletalMeshComponent = ComponentReference.GetComponent();
+	NotifyQueue.PredictedLODLevel = SkeletalMeshComponent ? SkeletalMeshComponent->GetPredictedLODLevel() : 0;
+}
+
 void FAnimNextNotifyDispatcherComponent::OnTraitEvent(FAnimNextTraitEvent& Event)
 {
 	if(UE::AnimNext::FNotifyDispatchEvent* NotifyDispatchEvent = Event.AsType<UE::AnimNext::FNotifyDispatchEvent>())
@@ -129,9 +139,13 @@ void FAnimNextNotifyDispatcherComponent::OnEndExecution(float InDeltaTime)
 
 	FAnimNextModuleInstance& ModuleInstance = GetModuleInstance();
 
-	TSharedPtr<FNotifyQueueDispatchEvent> NotifyQueueDispatchEvent = MakeTraitEvent<FNotifyQueueDispatchEvent>();
-	NotifyQueueDispatchEvent->WeakSkeletalMeshComponent = SkeletalMeshComponent;
-	NotifyQueueDispatchEvent->DeltaSeconds = InDeltaTime;
+	TSharedPtr<FNotifyQueueDispatchEvent> NotifyQueueDispatchEvent;
+	if(SkeletalMeshComponent != nullptr)
+	{
+		NotifyQueueDispatchEvent = MakeTraitEvent<FNotifyQueueDispatchEvent>();
+		NotifyQueueDispatchEvent->WeakSkeletalMeshComponent = SkeletalMeshComponent;
+		NotifyQueueDispatchEvent->DeltaSeconds = InDeltaTime;
+	}
 
 	NotifyQueue.PredictedLODLevel = SkeletalMeshComponent ? SkeletalMeshComponent->GetPredictedLODLevel() : 0;
 
@@ -179,7 +193,7 @@ void FAnimNextNotifyDispatcherComponent::OnEndExecution(float InDeltaTime)
 			}
 
 			// Trigger non 'state' AnimNotifies
-			TriggerSingleAnimNotify(InDeltaTime, *NotifyQueueDispatchEvent.Get(), NotifyQueue.AnimNotifies[Index]);
+			TriggerSingleAnimNotify(InDeltaTime, NotifyQueueDispatchEvent.Get(), NotifyQueue.AnimNotifies[Index]);
 		}
 	}
 
@@ -188,7 +202,7 @@ void FAnimNextNotifyDispatcherComponent::OnEndExecution(float InDeltaTime)
 	{
 		const FAnimNotifyEvent* AnimNotifyEvent = ActiveAnimNotifyState[Index];
 		const FAnimNotifyEventReference& EventReference = ActiveAnimNotifyEventReference[Index];
-		if (AnimNotifyEvent->NotifyStateClass)
+		if (NotifyQueueDispatchEvent.IsValid() && AnimNotifyEvent->NotifyStateClass)
 		{
 #if WITH_EDITOR
 			// Prevent firing notifies in animation editors if requested
@@ -205,7 +219,7 @@ void FAnimNextNotifyDispatcherComponent::OnEndExecution(float InDeltaTime)
 	{
 		const FAnimNotifyEvent* AnimNotifyEvent = NotifyStateBeginEvent[Index];
 		const FAnimNotifyEventReference* EventReference = NotifyStateBeginEventReference[Index];
-		if(AnimNotifyEvent->NotifyStateClass)
+		if(NotifyQueueDispatchEvent.IsValid() && AnimNotifyEvent->NotifyStateClass)
 		{
 #if WITH_EDITOR
 			// Prevent firing notifies in animation editors if requested 
@@ -226,7 +240,7 @@ void FAnimNextNotifyDispatcherComponent::OnEndExecution(float InDeltaTime)
 	{
 		const FAnimNotifyEvent* AnimNotifyEvent = ActiveAnimNotifyState[Index];
 		const FAnimNotifyEventReference& EventReference = ActiveAnimNotifyEventReference[Index];
-		if(AnimNotifyEvent->NotifyStateClass)
+		if(NotifyQueueDispatchEvent.IsValid() && AnimNotifyEvent->NotifyStateClass)
 		{
 #if WITH_EDITOR
 			// Prevent firing notifies in animation editors if requested 
@@ -242,21 +256,21 @@ void FAnimNextNotifyDispatcherComponent::OnEndExecution(float InDeltaTime)
 	NotifyQueue.AnimNotifies.Reset();
 }
 
-void FAnimNextNotifyDispatcherComponent::TriggerSingleAnimNotify(float InDeltaTime, UE::AnimNext::FNotifyQueueDispatchEvent& Dispatcher, const FAnimNotifyEventReference& EventReference)
+void FAnimNextNotifyDispatcherComponent::TriggerSingleAnimNotify(float InDeltaTime, UE::AnimNext::FNotifyQueueDispatchEvent* InDispatcher, const FAnimNotifyEventReference& EventReference)
 {
 	// This is for non 'state' anim notifies.
 	const FAnimNotifyEvent* AnimNotifyEvent = EventReference.GetNotify();
 	if (AnimNotifyEvent && (AnimNotifyEvent->NotifyStateClass == NULL))
 	{
 		FAnimNextModuleInstance& ModuleInstance = GetModuleInstance();
-		if (AnimNotifyEvent->Notify != nullptr)
+		if (InDispatcher != nullptr && AnimNotifyEvent->Notify != nullptr)
 		{
 #if WITH_EDITOR
 			// Prevent firing notifies in animation editors if requested 
 			if(ModuleInstance.GetWorldType() != EWorldType::EditorPreview || AnimNotifyEvent->Notify->ShouldFireInEditor())
 #endif
 			{
-				Dispatcher.EventsToNotify.Add(EventReference);
+				InDispatcher->EventsToNotify.Add(EventReference);
 			}
 		}
 		else if (AnimNotifyEvent->NotifyName != NAME_None)
diff --git a/Engine/Plugins/Experimental/AnimNextAnimGraph/Source/AnimNextAnimGraph/Private/Traits/NotifyDispatcher.h b/Engine/Plugins/Experimental/AnimNextAnimGraph/Source/AnimNextAnimGraph/Private/Traits/NotifyDispatcher.h
--- a/Engine/Plugins/Experimental/AnimNextAnimGraph/Source/AnimNextAnimGraph/Private/Traits/NotifyDispatcher.h
+++ b/Engine/Plugins/Experimental/AnimNextAnimGraph/Source/AnimNextAnimGraph/Private/Traits/NotifyDispatcher.h
@@ -40,11 +40,12 @@ struct FAnimNextNotifyDispatcherComponent : public FAnimNextModuleInstanceCompon
 	FAnimNextNotifyDispatcherComponent() = default;
 
 	// FAnimNextModuleInstanceComponent interface
+	virtual void OnInitialize() override;
 	virtual void OnTraitEvent(FAnimNextTraitEvent& Event) override;
 	virtual void OnEndExecution(float InDeltaTime) override;
 
 	// Triggers a single anim notify in the dispatcher
-	void TriggerSingleAnimNotify(float InDeltaTime, UE::AnimNext::FNotifyQueueDispatchEvent& Dispatcher, const FAnimNotifyEventReference& EventReference);
+	void TriggerSingleAnimNotify(float InDeltaTime, UE::AnimNext::FNotifyQueueDispatchEvent* InDispatcher, const FAnimNotifyEventReference& EventReference);
 
 	// Notify queue to dispatch
 	UPROPERTY()
