AnimNext : Renamed AlphaBlendTrait to PassthroughBlendTrait

diff --git a/Engine/Plugins/Experimental/AnimNextAnimGraph/Source/AnimNextAnimGraph/Private/Traits/AlphaBlendTrait.cpp b/Engine/Plugins/Experimental/AnimNextAnimGraph/Source/AnimNextAnimGraph/Private/Traits/AlphaBlendTrait.cpp
deleted file mode 100644
diff --git a/Engine/Plugins/Experimental/AnimNextAnimGraph/Source/AnimNextAnimGraph/Private/Traits/PassthroughBlendTrait.cpp b/Engine/Plugins/Experimental/AnimNextAnimGraph/Source/AnimNextAnimGraph/Private/Traits/PassthroughBlendTrait.cpp
new file mode 100644
--- /dev/null
+++ b/Engine/Plugins/Experimental/AnimNextAnimGraph/Source/AnimNextAnimGraph/Private/Traits/PassthroughBlendTrait.cpp
@@ -0,0 +1,122 @@
+// Copyright Epic Games, Inc. All Rights Reserved.
+
+#include "Traits/PassthroughBlendTrait.h"
+
+#include "Animation/AnimTypes.h"
+#include "TraitCore/ExecutionContext.h"
+#include "EvaluationVM/Tasks/BlendKeyframes.h"
+#include "EvaluationVM/Tasks/StoreKeyframe.h"
+
+namespace UE::AnimNext
+{
+
+AUTO_REGISTER_ANIM_TRAIT(FPassthroughBlendTrait)
+
+#define TRAIT_INTERFACE_ENUMERATOR(GeneratorMacro) \
+	GeneratorMacro(IUpdate) \
+	GeneratorMacro(IEvaluate) \
+
+// Trait implementation boilerplate
+GENERATE_ANIM_TRAIT_IMPLEMENTATION(FPassthroughBlendTrait, TRAIT_INTERFACE_ENUMERATOR, NULL_ANIM_TRAIT_INTERFACE_ENUMERATOR, NULL_ANIM_TRAIT_EVENT_ENUMERATOR)
+#undef TRAIT_INTERFACE_ENUMERATOR
+
+
+void FPassthroughBlendTrait::FInstanceData::Construct(const FExecutionContext& Context, const FTraitBinding& Binding)
+{
+	const FSharedData* SharedData = Binding.GetSharedData<FSharedData>();
+
+	InternalAlpha = SharedData->GetAlpha(Binding);
+	ComputedAlphaValue = InternalAlpha;
+	AlphaBoolBlend = SharedData->GetAlphaBoolBlend(Binding);
+	AlphaScaleBiasClamp = SharedData->AlphaScaleBiasClamp;
+}
+
+void FPassthroughBlendTrait::OnBecomeRelevant(FUpdateTraversalContext& Context, const TTraitBinding<IUpdate>& Binding, const FTraitUpdateState& TraitState) const
+{
+	IUpdate::OnBecomeRelevant(Context, Binding, TraitState);
+
+	const FSharedData* SharedData = Binding.GetSharedData<FSharedData>();
+	FInstanceData* InstanceData = Binding.GetInstanceData<FInstanceData>();
+
+	InstanceData->InternalAlpha = SharedData->GetAlpha(Binding);
+	InstanceData->ComputedAlphaValue = ComputeAlphaValue(SharedData->GetAlphaInputType(Binding), SharedData, InstanceData, TraitState.GetDeltaTime());
+
+	InstanceData->AlphaBoolBlend.Reinitialize();
+	InstanceData->AlphaScaleBiasClamp.Reinitialize();
+}
+
+void FPassthroughBlendTrait::PreUpdate(FUpdateTraversalContext& Context, const TTraitBinding<IUpdate>& Binding, const FTraitUpdateState& TraitState) const
+{
+	IUpdate::PreUpdate(Context, Binding, TraitState);
+
+	const FSharedData* SharedData = Binding.GetSharedData<FSharedData>();
+	FInstanceData* InstanceData = Binding.GetInstanceData<FInstanceData>();
+
+	const float DeltaTime = TraitState.GetDeltaTime();
+	InstanceData->DeltaTime = DeltaTime;
+	InstanceData->ComputedAlphaValue = ComputeAlphaValue(SharedData->GetAlphaInputType(Binding), SharedData, InstanceData, DeltaTime);
+}
+
+// IEvaluate impl
+void FPassthroughBlendTrait::PostEvaluate(FEvaluateTraversalContext& Context, const TTraitBinding<IEvaluate>& Binding) const
+{
+	// Duplicate existing pose in the stack
+	Context.AppendTask(FAnimNextDuplicateTopKeyframeTask::Make());
+
+	// Let children perform their tasks on the duplicated pose
+	IEvaluate::PostEvaluate(Context, Binding);
+
+	const FSharedData* SharedData = Binding.GetSharedData<FSharedData>();
+	FInstanceData* InstanceData = Binding.GetInstanceData<FInstanceData>();
+
+	const EAnimAlphaInputType AlphaInputType = SharedData->GetAlphaInputType(Binding);
+	if (AlphaInputType == EAnimAlphaInputType::Curve)
+	{
+		// Add a BlendTwoKey with the original pose and the children modified copy, using alpha curve data (value will be calculated / clamped at the task)
+		const float DeltaTime = InstanceData->DeltaTime;
+		const FName AlphaCurveName = SharedData->GetAlphaCurveName(Binding);
+		FInputScaleBiasClamp* InputScaleBiasClamp = &InstanceData->AlphaScaleBiasClamp;
+		static constexpr int32 AlphaSourceInputKeyframeIndex = 0;
+		Context.AppendTask(FAnimNextBlendTwoKeyframesTask::Make(AlphaCurveName, AlphaSourceInputKeyframeIndex, [DeltaTime, InputScaleBiasClamp](float Alpha) -> float
+			{
+				return InputScaleBiasClamp->ApplyTo(Alpha, DeltaTime);
+			}));
+	}
+	else
+	{
+		// Add a BlendTwoKey with the original pose and the children modified copy, using alpha value as float blend weight
+		const float BlendWeight = InstanceData->ComputedAlphaValue;
+		Context.AppendTask(FAnimNextBlendTwoKeyframesTask::Make(BlendWeight));
+	}
+}
+
+
+float FPassthroughBlendTrait::ComputeAlphaValue(const EAnimAlphaInputType AlphaInputType, const FPassthroughBlendTraitSharedData* SharedData, FInstanceData* InstanceData, float DeltaTime)
+{
+	// alpha handlers
+	float CurrentAlpha = 0.f;
+
+	switch (AlphaInputType)
+	{
+		case EAnimAlphaInputType::Float:
+		{
+			CurrentAlpha = SharedData->AlphaScaleBias.ApplyTo(InstanceData->AlphaScaleBiasClamp.ApplyTo(InstanceData->InternalAlpha, DeltaTime));
+			break;
+		}
+		case EAnimAlphaInputType::Bool:
+		{
+			CurrentAlpha = InstanceData->AlphaBoolBlend.ApplyTo(SharedData->bAlphaBoolEnabled, DeltaTime);
+			break;
+		}
+		case EAnimAlphaInputType::Curve:
+		{
+			CurrentAlpha = 1.0; // This will be calculated on the task when reading the curve, but we need both branches at full weight
+			break;
+		}
+	};
+
+	const float ClampedAlpha = FMath::Clamp(CurrentAlpha, 0.0f, 1.0f);
+	return ClampedAlpha;
+}
+
+} // namespace UE::AnimNext
diff --git a/Engine/Plugins/Experimental/AnimNextAnimGraph/Source/AnimNextAnimGraph/Public/Traits/AlphaBlendTrait.h b/Engine/Plugins/Experimental/AnimNextAnimGraph/Source/AnimNextAnimGraph/Public/Traits/AlphaBlendTrait.h
deleted file mode 100644
diff --git a/Engine/Plugins/Experimental/AnimNextAnimGraph/Source/AnimNextAnimGraph/Public/Traits/PassthroughBlendTrait.h b/Engine/Plugins/Experimental/AnimNextAnimGraph/Source/AnimNextAnimGraph/Public/Traits/PassthroughBlendTrait.h
new file mode 100644
--- /dev/null
+++ b/Engine/Plugins/Experimental/AnimNextAnimGraph/Source/AnimNextAnimGraph/Public/Traits/PassthroughBlendTrait.h
@@ -0,0 +1,87 @@
+// Copyright Epic Games, Inc. All Rights Reserved.
+
+#pragma once
+
+#include "CoreMinimal.h"
+#include "TraitCore/Trait.h"
+#include "TraitCore/TraitSharedData.h"
+#include "TraitInterfaces/IUpdate.h"
+#include "TraitInterfaces/IEvaluate.h"
+#include "Animation/InputScaleBias.h"
+// --- ---
+#include "PassthroughBlendTrait.generated.h"
+
+USTRUCT(meta = (DisplayName = "Passthrough Blend"))
+struct FPassthroughBlendTraitSharedData : public FAnimNextTraitSharedData
+{
+	GENERATED_BODY()
+
+	// alpha value handler
+	UPROPERTY(EditAnywhere, Category = Settings, meta = (EditConditionHides, EditCondition = "AlphaInputType == EAnimAlphaInputType::Float"))
+	float Alpha = 1.f;
+
+	UPROPERTY(EditAnywhere, Category=Settings, meta = (Hidden, EditConditionHides, EditCondition = "AlphaInputType == EAnimAlphaInputType::Float"))
+	FInputScaleBias AlphaScaleBias;
+
+	UPROPERTY(EditAnywhere, Category = Settings, meta = (Hidden, DisplayAfter = "AlphaCurveName", EditConditionHides, EditCondition = "AlphaInputType == EAnimAlphaInputType::Float || AlphaInputType == EAnimAlphaInputType::Curve"))
+	FInputScaleBiasClamp AlphaScaleBiasClamp;
+
+	UPROPERTY(EditAnywhere, Category = Settings, meta = (DisplayName = "bEnabled", EditConditionHides, EditCondition = "AlphaInputType == EAnimAlphaInputType::Bool"))
+	bool bAlphaBoolEnabled = true;
+
+	UPROPERTY(EditAnywhere, Category = Settings, meta = (Hidden, DisplayName = "Blend Settings", EditConditionHides, EditCondition = "AlphaInputType == EAnimAlphaInputType::Bool"))
+	FInputAlphaBoolBlend AlphaBoolBlend;
+
+	UPROPERTY(EditAnywhere, Category = Settings, meta = (EditConditionHides, EditCondition = "AlphaInputType == EAnimAlphaInputType::Curve"))
+	FName AlphaCurveName = NAME_None;
+
+	UPROPERTY(EditAnywhere, Category = Settings)
+	EAnimAlphaInputType AlphaInputType = EAnimAlphaInputType::Float;
+
+	// Latent pin support boilerplate
+#define TRAIT_LATENT_PROPERTIES_ENUMERATOR(GeneratorMacro) \
+		GeneratorMacro(Alpha) \
+		GeneratorMacro(bAlphaBoolEnabled) \
+		GeneratorMacro(AlphaBoolBlend) \
+		GeneratorMacro(AlphaCurveName) \
+		GeneratorMacro(AlphaInputType) \
+
+	GENERATE_TRAIT_LATENT_PROPERTIES(FPassthroughBlendTraitSharedData, TRAIT_LATENT_PROPERTIES_ENUMERATOR)
+#undef TRAIT_LATENT_PROPERTIES_ENUMERATOR
+};
+
+namespace UE::AnimNext
+{
+
+USTRUCT(meta = (DisplayName = "Passthrough Blend"))
+struct FPassthroughBlendTrait : FAdditiveTrait, IUpdate, IEvaluate
+{
+	typedef FBaseTrait Super;
+
+	DECLARE_ANIM_TRAIT(FPassthroughBlendTrait, FAdditiveTrait)
+
+	using FSharedData = FPassthroughBlendTraitSharedData;
+		
+	struct FInstanceData : FTrait::FInstanceData
+	{
+		FInputAlphaBoolBlend AlphaBoolBlend;
+		FInputScaleBiasClamp AlphaScaleBiasClamp;
+
+		float DeltaTime = 0.f;
+		float InternalAlpha = 1.f;
+		float ComputedAlphaValue = 0.f;
+
+		void Construct(const FExecutionContext& Context, const FTraitBinding& Binding);
+	};
+
+	// IUpdate
+	virtual void OnBecomeRelevant(FUpdateTraversalContext& Context, const TTraitBinding<IUpdate>& Binding, const FTraitUpdateState& TraitState) const override;
+	virtual void PreUpdate(FUpdateTraversalContext& Context, const TTraitBinding<IUpdate>& Binding, const FTraitUpdateState& TraitState) const override;
+
+	// IEvaluate impl
+	virtual void PostEvaluate(FEvaluateTraversalContext& Context, const TTraitBinding<IEvaluate>& Binding) const override;
+
+	static float ComputeAlphaValue(const EAnimAlphaInputType AlphaInputType, const FPassthroughBlendTraitSharedData* SharedData, FInstanceData* InstanceData, float DeltaTime);
+};
+
+} // namespace UE::AnimNext
