Show required interfaces in trait list when bShowTraitInterfaces is enabled

diff --git a/Engine/Plugins/Experimental/AnimNextAnimGraph/Source/AnimNextAnimGraphEditor/Private/Graph/STraitListView.cpp b/Engine/Plugins/Experimental/AnimNextAnimGraph/Source/AnimNextAnimGraphEditor/Private/Graph/STraitListView.cpp
--- a/Engine/Plugins/Experimental/AnimNextAnimGraph/Source/AnimNextAnimGraphEditor/Private/Graph/STraitListView.cpp
+++ b/Engine/Plugins/Experimental/AnimNextAnimGraph/Source/AnimNextAnimGraphEditor/Private/Graph/STraitListView.cpp
@@ -665,6 +665,14 @@ TSharedRef<ITableRow> STraitListView::HandleGenerateRow(TSharedRef<FTraitListEnt
 						[
 							SNew(SVerticalBox)
 
+							// --- Trait Required Interfaces ---
+							+ SVerticalBox::Slot()
+							.MinHeight(TraitEditorSharedDataLocal->bShowTraitInterfaces ? 23.0f : 0.0f)
+							.AutoHeight()
+							[
+								FTraitEditorUtils::GetInterfaceListWidget(FTraitEditorUtils::EInterfaceDisplayType::ListRequired, TraitDataShared, TraitEditorSharedDataLocal)
+							]
+
 							// --- Trait Main Button ---
 							+ SVerticalBox::Slot()
 							[
@@ -760,6 +768,7 @@ TSharedRef<ITableRow> STraitListView::HandleGenerateRow(TSharedRef<FTraitListEnt
 
 							// --- Trait Implemented Interfaces ---
 							+ SVerticalBox::Slot()
+							.MinHeight(TraitEditorSharedDataLocal->bShowTraitInterfaces ? 23.0f : 0.0f)
 							.AutoHeight()
 							[
 								FTraitEditorUtils::GetInterfaceListWidget(FTraitEditorUtils::EInterfaceDisplayType::ListImplemented, TraitDataShared, TraitEditorSharedDataLocal)
diff --git a/Engine/Plugins/Experimental/AnimNextAnimGraph/Source/AnimNextAnimGraphEditor/Private/Graph/TraitEditorDefs.cpp b/Engine/Plugins/Experimental/AnimNextAnimGraph/Source/AnimNextAnimGraphEditor/Private/Graph/TraitEditorDefs.cpp
--- a/Engine/Plugins/Experimental/AnimNextAnimGraph/Source/AnimNextAnimGraphEditor/Private/Graph/TraitEditorDefs.cpp
+++ b/Engine/Plugins/Experimental/AnimNextAnimGraph/Source/AnimNextAnimGraphEditor/Private/Graph/TraitEditorDefs.cpp
@@ -106,6 +106,23 @@ namespace UE::AnimNext::Editor
 			break;
 		}
 
+		case EInterfaceDisplayType::ListRequired:
+		{
+			for (const FTraitInterfaceUID& InterfaceUID : InTraitDataShared->RequiredInterfaces)
+			{
+				if (!IsInternal(InterfaceUID))
+				{
+					InterfaceWidgetsHorizontalBox->AddSlot()
+						.AutoWidth()
+						.Padding(FMargin(4.f, 1.f))
+						[
+							GetInterfaceWidget(InterfaceDisplayType, InterfaceUID, InTraitDataShared, InTraitEditorSharedDataShared)
+						];
+				}
+			}
+			break;
+		}
+
 		case EInterfaceDisplayType::StackRequired:
 		case EInterfaceDisplayType::StackImplemented:
 		{
@@ -210,6 +227,7 @@ namespace UE::AnimNext::Editor
 							}
 							break;
 						}
+						case EInterfaceDisplayType::ListRequired:
 						case EInterfaceDisplayType::StackRequired:
 						{
 							if (InTraitDataShared->RequiredInterfaces.Find(InterfaceUID) != INDEX_NONE)
diff --git a/Engine/Plugins/Experimental/AnimNextAnimGraph/Source/AnimNextAnimGraphEditor/Private/Graph/TraitEditorDefs.h b/Engine/Plugins/Experimental/AnimNextAnimGraph/Source/AnimNextAnimGraphEditor/Private/Graph/TraitEditorDefs.h
--- a/Engine/Plugins/Experimental/AnimNextAnimGraph/Source/AnimNextAnimGraphEditor/Private/Graph/TraitEditorDefs.h
+++ b/Engine/Plugins/Experimental/AnimNextAnimGraph/Source/AnimNextAnimGraphEditor/Private/Graph/TraitEditorDefs.h
@@ -149,6 +149,7 @@ struct FTraitEditorUtils
 	enum class EInterfaceDisplayType : uint8
 	{
 		ListImplemented,
+		ListRequired,
 		StackImplemented,
 		StackRequired
 	};
