Add static assert to ensure correct trait data inheritance and fix up the blend stack trait

diff --git a/Engine/Plugins/Experimental/AnimNextAnimGraph/Source/AnimNextAnimGraph/Internal/Traits/BlendStackTrait.h b/Engine/Plugins/Experimental/AnimNextAnimGraph/Source/AnimNextAnimGraph/Internal/Traits/BlendStackTrait.h
--- a/Engine/Plugins/Experimental/AnimNextAnimGraph/Source/AnimNextAnimGraph/Internal/Traits/BlendStackTrait.h
+++ b/Engine/Plugins/Experimental/AnimNextAnimGraph/Source/AnimNextAnimGraph/Internal/Traits/BlendStackTrait.h
@@ -26,7 +26,7 @@ struct FAnimNextBlendStackCoreTraitSharedData : public FAnimNextTraitSharedData
 };
 
 USTRUCT(meta = (DisplayName = "Blend Stack"))
-struct FAnimNextBlendStackTraitSharedData : public FAnimNextTraitSharedData
+struct FAnimNextBlendStackTraitSharedData : public FAnimNextBlendStackCoreTraitSharedData
 {
 	GENERATED_BODY()
 
diff --git a/Engine/Plugins/Experimental/AnimNextAnimGraph/Source/AnimNextAnimGraph/Public/TraitCore/Trait.h b/Engine/Plugins/Experimental/AnimNextAnimGraph/Source/AnimNextAnimGraph/Public/TraitCore/Trait.h
--- a/Engine/Plugins/Experimental/AnimNextAnimGraph/Source/AnimNextAnimGraph/Public/TraitCore/Trait.h
+++ b/Engine/Plugins/Experimental/AnimNextAnimGraph/Source/AnimNextAnimGraph/Public/TraitCore/Trait.h
@@ -25,9 +25,7 @@ struct FRigVMPinInfoArray;
 	/* FTrait impl */ \
 	static constexpr UE::AnimNext::FTraitUID TraitUID = UE::AnimNext::FTraitUID::MakeUID(TEXT(#TraitName)); \
 	virtual UE::AnimNext::FTraitUID GetTraitUID() const override { return TraitUID; } \
-	virtual FString GetTraitName() const override { return TEXT(#TraitName); } \
-	static_assert(std::is_base_of<TraitSuper::FSharedData, FSharedData>::value, "Trait shared data must derive from its Super's FSharedData"); \
-	static_assert(std::is_base_of<TraitSuper::FInstanceData, FInstanceData>::value, "Trait instance data must derive from Super's FInstanceData");
+	virtual FString GetTraitName() const override { return TEXT(#TraitName); }
 
 #define ANIM_NEXT_IMPL_DECLARE_ANIM_TRAIT_INSTANCING_SUPPORT() \
 	static const UE::AnimNext::FTraitMemoryLayout TraitMemoryDescription; \
@@ -83,6 +81,8 @@ struct FRigVMPinInfoArray;
 		}; \
 	void TraitName::ConstructTraitInstance(const UE::AnimNext::FExecutionContext& Context, const UE::AnimNext::FTraitBinding& Binding) const \
 	{ \
+		static_assert(std::is_base_of<TraitSuper::FSharedData, TraitName::FSharedData>::value, "Trait shared data must derive from its Super's FSharedData"); \
+		static_assert(std::is_base_of<TraitSuper::FInstanceData, TraitName::FInstanceData>::value, "Trait instance data must derive from Super's FInstanceData"); \
 		/* Construct the base struct first */ \
 		FInstanceData* Data = new(Binding.GetInstanceData<FInstanceData>()) FInstanceData(); \
 		/* Then construct our latent properties, the Construct implementation below might need them */ \
