Add pins by trait order when exposing variables, which is pre-sorted to be optimal for memory access in UncookedOnlyUtils.

#tests PIE

diff --git a/Engine/Plugins/Experimental/AnimNextAnimGraph/Source/AnimNextAnimGraphEditor/Private/Graph/RigVMTrait_AnimNextPublicVariablesEditor.cpp b/Engine/Plugins/Experimental/AnimNextAnimGraph/Source/AnimNextAnimGraphEditor/Private/Graph/RigVMTrait_AnimNextPublicVariablesEditor.cpp
--- a/Engine/Plugins/Experimental/AnimNextAnimGraph/Source/AnimNextAnimGraphEditor/Private/Graph/RigVMTrait_AnimNextPublicVariablesEditor.cpp
+++ b/Engine/Plugins/Experimental/AnimNextAnimGraph/Source/AnimNextAnimGraphEditor/Private/Graph/RigVMTrait_AnimNextPublicVariablesEditor.cpp
@@ -1,4 +1,4 @@
-﻿// Copyright Epic Games, Inc. All Rights Reserved.
+// Copyright Epic Games, Inc. All Rights Reserved.
 
 #include "RigVMTrait_AnimNextPublicVariablesEditor.h"
 #include "AnimNextRigVMAssetEditorData.h"
@@ -49,15 +49,36 @@ void FPublicVariablesImpl::GetProgrammaticPins(const FRigVMTrait_AnimNextPublicV
 	Descs.Reserve(PublicVariables.Num());
 	TArray<TConstArrayView<uint8>> Values;
 	Values.Reserve(PublicVariables.Num());
-	for(UAnimNextVariableEntry* PublicVariable : PublicVariables)
+
+	// Maintain trait sorted order
+	for (FName VariableName : InTrait.InternalVariableNames)
 	{
-		const FProperty* Property;
-		TConstArrayView<uint8> Value;
-		PublicVariable->GetDefaultValue(Property, Value);
-		Descs.Emplace(PublicVariable->GetEntryName(), Property);
-		Values.Emplace(Value);
+		auto FindMatchingName = [&](const UAnimNextVariableEntry* InVariable)
+		{
+			if (!InVariable)
+			{
+				return false;
+			}
+
+			return InVariable->GetEntryName() == VariableName;
+		};
+
+		if (UAnimNextVariableEntry** PublicVariableItr = PublicVariables.FindByPredicate(FindMatchingName))
+		{
+			if (UAnimNextVariableEntry* PublicVariable = *PublicVariableItr)
+			{
+				const FProperty* Property;
+				TConstArrayView<uint8> Value;
+				PublicVariable->GetDefaultValue(Property, Value);
+				Descs.Emplace(PublicVariable->GetEntryName(), Property);
+				Values.Emplace(Value);
+			}
+		}
+
 	}
 
+	ensure(Values.Num() == InTrait.InternalVariableNames.Num());
+
 	FInstancedPropertyBag Defaults;
 	Defaults.ReplaceAllPropertiesAndValues(Descs, Values);
 
