Fix bindings not mapping when indexes do not match. There is no need for them to match.
Fix crash togging visibility variable to off, fire an ensure in place

#tests PIE

diff --git a/Engine/Plugins/Experimental/AnimNextAnimGraph/Source/AnimNextAnimGraph/Private/Graph/AnimNextGraphInstance.cpp b/Engine/Plugins/Experimental/AnimNextAnimGraph/Source/AnimNextAnimGraph/Private/Graph/AnimNextGraphInstance.cpp
--- a/Engine/Plugins/Experimental/AnimNextAnimGraph/Source/AnimNextAnimGraph/Private/Graph/AnimNextGraphInstance.cpp
+++ b/Engine/Plugins/Experimental/AnimNextAnimGraph/Source/AnimNextAnimGraph/Private/Graph/AnimNextGraphInstance.cpp
@@ -291,22 +291,35 @@ bool FAnimNextGraphInstance::BindToCachedBindings()
 		int32 VariableIndex = ImplementedInterface.VariableIndex;
 		int32 InterfaceVariableIndex = 0;
 		const int32 EndVariableIndex = ImplementedInterface.VariableIndex + ImplementedInterface.NumVariables;
+		const int32 BindingStartIndex = CachedBindings.IsEmpty() ? 0 : CachedBinding->CachedBindings[0].VariableIndex;
 		for(; VariableIndex < EndVariableIndex; ++VariableIndex, ++InterfaceVariableIndex)
 		{
+			if (!CachedBinding->CachedBindings.IsValidIndex(InterfaceVariableIndex))
+			{
+				ensureMsgf(false, TEXT("Bindings indexes should always contain variable indexes"));
+				continue;
+			}
+
 			const FPropertyBagPropertyDesc& Desc = PropertyBag->GetPropertyDescs()[VariableIndex];
 			const FCachedDataInterfaceBinding::FVariable& CachedBindingVariable = CachedBinding->CachedBindings[InterfaceVariableIndex];
 			if(CachedBindingVariable.Memory == nullptr)
 			{
+				UE_LOG(LogAnimation, Error, TEXT("FAnimNextGraphInstance::BindToCachedBindings: Did not find cached binding for %s - variable's values will not be updated at runtime.'"), *Desc.Name.ToString());
 				continue;
 			}
 
 			if( Desc.Name == CachedBindingVariable.VariableName &&
-				VariableIndex == CachedBindingVariable.VariableIndex &&
+				VariableIndex == (CachedBindingVariable.VariableIndex - BindingStartIndex) &&
 				Desc.CachedProperty->GetClass() == CachedBindingVariable.Property->GetClass())
 			{
 				ExtendedExecuteContext.ExternalVariableRuntimeData[VariableIndex].Memory = CachedBindingVariable.Memory;
 				bPublicVariablesBound = true;
 			}
+			else
+			{
+				UE_LOG(LogAnimation, Error, TEXT("FAnimNextGraphInstance::BindToCachedBindings: Did not find cached binding for %s - variable's values will not be updated at runtime.'"), *Desc.Name.ToString());
+				continue;
+			}
 		}
 	}
 
@@ -380,10 +393,11 @@ void FAnimNextGraphInstance::UpdateCachedBindingsForHostHelper(const HostType& I
 		{
 			const FPropertyBagPropertyDesc& Desc = PropertyBag->GetPropertyDescs()[VariableIndex];
 			uint8* HostMemory = InHost.GetMemoryForVariable(VariableIndex, Desc.Name, Desc.CachedProperty);
-			if(HostMemory == nullptr)
-			{
-				continue;
-			}
+			uint8* PrevMemory = CachedBinding.CachedBindings.IsValidIndex(BindingIndex) 
+				? CachedBinding.CachedBindings[BindingIndex].Memory 
+				: nullptr;
+			HostMemory = HostMemory ? HostMemory : PrevMemory;
+
 			CachedBinding.CachedBindings[BindingIndex] = FCachedDataInterfaceBinding::FVariable(Desc.Name, VariableIndex, Desc.CachedProperty, HostMemory);
 		}
 	}
