Migrate sync group trait to use its own role enum instead of the legacy AnimBP one

diff --git a/Engine/Plugins/Experimental/AnimNextAnimGraph/Source/AnimNextAnimGraph/Private/Graph/SyncGroup_GraphInstanceComponent.cpp b/Engine/Plugins/Experimental/AnimNextAnimGraph/Source/AnimNextAnimGraph/Private/Graph/SyncGroup_GraphInstanceComponent.cpp
--- a/Engine/Plugins/Experimental/AnimNextAnimGraph/Source/AnimNextAnimGraph/Private/Graph/SyncGroup_GraphInstanceComponent.cpp
+++ b/Engine/Plugins/Experimental/AnimNextAnimGraph/Source/AnimNextAnimGraph/Private/Graph/SyncGroup_GraphInstanceComponent.cpp
@@ -137,8 +137,8 @@ namespace UE::AnimNext
 
 				switch (GroupMember.GroupParameters.GroupRole)
 				{
-				case EAnimGroupRole::CanBeLeader:
-				case EAnimGroupRole::TransitionLeader:
+				case EAnimGroupSynchronizationRole::CanBeLeader:
+				case EAnimGroupSynchronizationRole::TransitionLeader:
 					// Highest weight is the leader
 					if (GroupMember.TraitState.GetTotalWeight() > LeaderTotalWeight)
 					{
@@ -146,15 +146,15 @@ namespace UE::AnimNext
 						LeaderTotalWeight = GroupMember.TraitState.GetTotalWeight();
 					}
 					break;
-				case EAnimGroupRole::AlwaysLeader:
-				case EAnimGroupRole::ExclusiveAlwaysLeader:
+				case EAnimGroupSynchronizationRole::AlwaysLeader:
+				case EAnimGroupSynchronizationRole::ExclusiveAlwaysLeader:
 					// Always set the leader index
 					LeaderIndex = MemberIndex;
 					LeaderTotalWeight = 2.0f;		// Some high value
 					break;
 				default:
-				case EAnimGroupRole::AlwaysFollower:
-				case EAnimGroupRole::TransitionFollower:
+				case EAnimGroupSynchronizationRole::AlwaysFollower:
+				case EAnimGroupSynchronizationRole::TransitionFollower:
 					// Never set the leader index
 					// If we find no leader, we'll use the first index as set below
 					break;
@@ -952,7 +952,7 @@ namespace UE::AnimNext
 
 				float DeltaTime = 0.0f;
 
-				if (GroupMember.GroupParameters.GroupRole == EAnimGroupRole::ExclusiveAlwaysLeader)
+				if (GroupMember.GroupParameters.GroupRole == EAnimGroupSynchronizationRole::ExclusiveAlwaysLeader)
 				{
 					// We asked to be a leader and we weren't picked as leader
 					// We don't want to follow anyone so we tick on our own
diff --git a/Engine/Plugins/Experimental/AnimNextAnimGraph/Source/AnimNextAnimGraph/Private/Traits/SynchronizeUsingGroups.cpp b/Engine/Plugins/Experimental/AnimNextAnimGraph/Source/AnimNextAnimGraph/Private/Traits/SynchronizeUsingGroups.cpp
--- a/Engine/Plugins/Experimental/AnimNextAnimGraph/Source/AnimNextAnimGraph/Private/Traits/SynchronizeUsingGroups.cpp
+++ b/Engine/Plugins/Experimental/AnimNextAnimGraph/Source/AnimNextAnimGraph/Private/Traits/SynchronizeUsingGroups.cpp
@@ -65,7 +65,7 @@ namespace UE::AnimNext
 
 		// When the role is set to TransitionLeader or TransitionFollower, we will not be part of a sync group until we are full weight
 		// meaning we ignore the sync group while blending in, but not while blending out
-		if (SharedData->GroupRole == EAnimGroupRole::TransitionLeader || SharedData->GroupRole == EAnimGroupRole::TransitionFollower)
+		if (SharedData->GroupRole == EAnimGroupSynchronizationRole::TransitionLeader || SharedData->GroupRole == EAnimGroupSynchronizationRole::TransitionFollower)
 		{
 			const FInstanceData* InstanceData = Binding.GetInstanceData<FInstanceData>();
 
diff --git a/Engine/Plugins/Experimental/AnimNextAnimGraph/Source/AnimNextAnimGraph/Private/Traits/SynchronizeUsingGroups.h b/Engine/Plugins/Experimental/AnimNextAnimGraph/Source/AnimNextAnimGraph/Private/Traits/SynchronizeUsingGroups.h
--- a/Engine/Plugins/Experimental/AnimNextAnimGraph/Source/AnimNextAnimGraph/Private/Traits/SynchronizeUsingGroups.h
+++ b/Engine/Plugins/Experimental/AnimNextAnimGraph/Source/AnimNextAnimGraph/Private/Traits/SynchronizeUsingGroups.h
@@ -3,7 +3,6 @@
 #pragma once
 
 #include "CoreMinimal.h"
-#include "Animation/AnimationAsset.h"
 #include "TraitCore/Trait.h"
 #include "TraitInterfaces/IGroupSynchronization.h"
 #include "TraitInterfaces/IUpdate.h"
@@ -23,7 +22,7 @@ struct FAnimNextSynchronizeUsingGroupsTraitSharedData : public FAnimNextTraitSha
 
 	// The role this player can assume within the group
 	UPROPERTY(EditAnywhere, Category = "Default", meta = (Inline))
-	TEnumAsByte<EAnimGroupRole::Type> GroupRole = EAnimGroupRole::CanBeLeader;
+	EAnimGroupSynchronizationRole GroupRole = EAnimGroupSynchronizationRole::CanBeLeader;
 
 	// Whether or not to match the group sync point when joining as leader
 	// When disabled, the leader will snap the position of the group to its current position
diff --git a/Engine/Plugins/Experimental/AnimNextAnimGraph/Source/AnimNextAnimGraph/Public/TraitInterfaces/IGroupSynchronization.h b/Engine/Plugins/Experimental/AnimNextAnimGraph/Source/AnimNextAnimGraph/Public/TraitInterfaces/IGroupSynchronization.h
--- a/Engine/Plugins/Experimental/AnimNextAnimGraph/Source/AnimNextAnimGraph/Public/TraitInterfaces/IGroupSynchronization.h
+++ b/Engine/Plugins/Experimental/AnimNextAnimGraph/Source/AnimNextAnimGraph/Public/TraitInterfaces/IGroupSynchronization.h
@@ -3,13 +3,36 @@
 #pragma once
 
 #include "CoreMinimal.h"
-#include "Animation/AnimationAsset.h"
 #include "TraitCore/ITraitInterface.h"
 #include "TraitCore/TraitBinding.h"
 #include "TraitInterfaces/ITimeline.h"
 
+#include "IGroupSynchronization.generated.h"
+
 class FMarkerTickContext;
 
+UENUM(BlueprintType)
+enum class EAnimGroupSynchronizationRole : uint8
+{
+	/** This node can be the leader, as long as it has a higher blend weight than the previous best leader. */
+	CanBeLeader,
+		
+	/** This node will always be a follower (unless there are only followers, in which case the first one ticked wins). */
+	AlwaysFollower,
+
+	/** This node will always be a leader (if more than one node is AlwaysLeader, the last one ticked wins). */
+	AlwaysLeader,
+
+	/** This node will be excluded from the sync group while blending in. Once blended in it will be the sync group leader until blended out*/
+	TransitionLeader,
+
+	/** This node will be excluded from the sync group while blending in. Once blended in it will be a follower until blended out*/
+	TransitionFollower,
+
+	/** This node will always be a leader. If it fails to be ticked as a leader it will be run as ungrouped asset player (EAnimSyncMethod::DoNotSync) .*/
+	ExclusiveAlwaysLeader,
+};
+
 namespace UE::AnimNext
 {
 	// Encapsulates the parameters that control sync group behavior
@@ -19,7 +42,7 @@ namespace UE::AnimNext
 		FName GroupName;
 
 		// The role within the group
-		TEnumAsByte<EAnimGroupRole::Type> GroupRole = EAnimGroupRole::CanBeLeader;
+		EAnimGroupSynchronizationRole GroupRole = EAnimGroupSynchronizationRole::CanBeLeader;
 
 		// Whether or not to match the group sync point when joining as leader
 		// When disabled, the leader will snap the position of the group to its current position
