Add ability to add more than one native data interface in an AnimNext graph
Add ability to require a number of native data interfaces for a graph factory entry
Rework custom graph factory payload support to allow traits to populate their own
Ensure sync group trait does nothing when no timeline is present

diff --git a/Engine/Plugins/Experimental/AnimNext/Source/AnimNextUncookedOnly/Internal/AnimNextRigVMAssetEditorData.h b/Engine/Plugins/Experimental/AnimNext/Source/AnimNextUncookedOnly/Internal/AnimNextRigVMAssetEditorData.h
--- a/Engine/Plugins/Experimental/AnimNext/Source/AnimNextUncookedOnly/Internal/AnimNextRigVMAssetEditorData.h
+++ b/Engine/Plugins/Experimental/AnimNext/Source/AnimNextUncookedOnly/Internal/AnimNextRigVMAssetEditorData.h
@@ -506,8 +506,13 @@ protected:
 	FRigVMGraphFunctionStore GraphFunctionStore;
 
 	// The native C++ struct that is used to communicate with this asset
+	UE_DEPRECATED(5.6, "This property is deprecated. Please use NativeInterfaces instead")
+	UPROPERTY(meta = (ShowDisplayNames, MetaStruct = "/Script/AnimNext.AnimNextNativeDataInterface"))
+	TObjectPtr<const UScriptStruct> NativeInterface_DEPRECATED;
+
+	// The list of native C++ structs that are used to communicate with this asset
 	UPROPERTY(EditAnywhere, Category = "Native", meta = (ShowDisplayNames, MetaStruct = "/Script/AnimNext.AnimNextNativeDataInterface"))
-	TObjectPtr<const UScriptStruct> NativeInterface;
+	TArray<TObjectPtr<const UScriptStruct>> NativeInterfaces;
 
 	UPROPERTY(EditAnywhere, Category = "User Interface")
 	FRigVMEdGraphDisplaySettings RigGraphDisplaySettings;
diff --git a/Engine/Plugins/Experimental/AnimNext/Source/AnimNextUncookedOnly/Private/AnimNextRigVMAssetEditorData.cpp b/Engine/Plugins/Experimental/AnimNext/Source/AnimNextUncookedOnly/Private/AnimNextRigVMAssetEditorData.cpp
--- a/Engine/Plugins/Experimental/AnimNext/Source/AnimNextUncookedOnly/Private/AnimNextRigVMAssetEditorData.cpp
+++ b/Engine/Plugins/Experimental/AnimNext/Source/AnimNextUncookedOnly/Private/AnimNextRigVMAssetEditorData.cpp
@@ -177,6 +177,15 @@ void UAnimNextRigVMAssetEditorData::PostLoad()
 {
 	Super::PostLoad();
 
+	// Handle deprecation
+	PRAGMA_DISABLE_DEPRECATION_WARNINGS
+	if (NativeInterface_DEPRECATED)
+	{
+		NativeInterfaces.Add(NativeInterface_DEPRECATED);
+		NativeInterface_DEPRECATED = nullptr;
+	}
+	PRAGMA_ENABLE_DEPRECATION_WARNINGS
+
 	GraphModels.Reset();
 	
 	PostLoadExternalPackages();
diff --git a/Engine/Plugins/Experimental/AnimNext/Source/AnimNextUncookedOnly/Private/UncookedOnlyUtils.cpp b/Engine/Plugins/Experimental/AnimNext/Source/AnimNextUncookedOnly/Private/UncookedOnlyUtils.cpp
--- a/Engine/Plugins/Experimental/AnimNext/Source/AnimNextUncookedOnly/Private/UncookedOnlyUtils.cpp
+++ b/Engine/Plugins/Experimental/AnimNext/Source/AnimNextUncookedOnly/Private/UncookedOnlyUtils.cpp
@@ -149,26 +149,56 @@ void FUtils::CompileVariables(const FRigVMCompileSettings& InSettings, UAnimNext
 
 		if(ExistingIndexPtr)
 		{
+			// Found a variable of the same name/type, overwrite its value
 			StructEntryInfos[*ExistingIndexPtr].Value = Value;
 		}
 		else
 		{
+			// This is a new variable, check if it belongs to a native interface
 			const UAnimNextDataInterface_EditorData* FromInterfaceEditorData = GetEditorData<const UAnimNextDataInterface_EditorData>(InFromInterface);
 			check(FromInterfaceEditorData != nullptr);
 
-			int32 Index = StructEntryInfos.Add(
+			bool bFound = false;
+			for (const UScriptStruct* NativeInterfaceStruct : FromInterfaceEditorData->NativeInterfaces)
 			{
-				InFromInterface,
-				FromInterfaceEditorData->NativeInterface,
-				Name,
-				FAnimNextParamType(Type.GetValueType(), Type.GetContainerType(), Type.GetValueTypeObject()),
-				AccessSpecifier,
-				bInAutoBindInterface,
-				Value,
-				CPF_Edit
-			});
-
-			EntryInfoIndexMap.Add(Name, Index);
+				if (NativeInterfaceStruct != nullptr && NativeInterfaceStruct->FindPropertyByName(Name) != nullptr)
+				{
+					// Found it
+					int32 Index = StructEntryInfos.Add(
+						{
+							InFromInterface,
+							NativeInterfaceStruct,
+							Name,
+							FAnimNextParamType(Type.GetValueType(), Type.GetContainerType(), Type.GetValueTypeObject()),
+							AccessSpecifier,
+							bInAutoBindInterface,
+							Value,
+							CPF_Edit
+						});
+
+					EntryInfoIndexMap.Add(Name, Index);
+					bFound = true;
+					break;
+				}
+			}
+
+			if (!bFound)
+			{
+				// Legacy codepath
+				int32 Index = StructEntryInfos.Add(
+					{
+						InFromInterface,
+						nullptr,
+						Name,
+						FAnimNextParamType(Type.GetValueType(), Type.GetContainerType(), Type.GetValueTypeObject()),
+						AccessSpecifier,
+						bInAutoBindInterface,
+						Value,
+						CPF_Edit
+					});
+
+				EntryInfoIndexMap.Add(Name, Index);
+			}
 		}
 	};
 
diff --git a/Engine/Plugins/Experimental/AnimNextAnimGraph/Source/AnimNextAnimGraph/Internal/AnimNextAnimGraphSettings.h b/Engine/Plugins/Experimental/AnimNextAnimGraph/Source/AnimNextAnimGraph/Internal/AnimNextAnimGraphSettings.h
--- a/Engine/Plugins/Experimental/AnimNextAnimGraph/Source/AnimNextAnimGraph/Internal/AnimNextAnimGraphSettings.h
+++ b/Engine/Plugins/Experimental/AnimNextAnimGraph/Source/AnimNextAnimGraph/Internal/AnimNextAnimGraphSettings.h
@@ -36,6 +36,12 @@ private:
 	// Must be an 'object' property of a compatible type to AssetType.
 	UPROPERTY(EditAnywhere, Category = "Mapping")
 	FName Variable;
+
+	UPROPERTY(EditAnywhere, Category = "Mapping")
+	TArray<TSoftObjectPtr<const UScriptStruct>> RequiredDataInterfaces;
+
+	// Index of the next mapping for this AssetType
+	int32 NextMappingIndex = INDEX_NONE;
 };
 
 UCLASS(Config=AnimNextAnimGraph, defaultconfig)
@@ -47,6 +53,9 @@ public:
 	// Given an object, return an animation graph to instantiate and a set of native interface structs that can be used to communicate with the instance
 	const UAnimNextAnimationGraph* GetGraphFromObject(const UObject* InObject) const;
 
+	// Given an object, return an animation graph to instantiate and a set of native interface structs that can be used to communicate with the instance
+	const UAnimNextAnimationGraph* GetGraphFromObject(const UObject* InObject, const FAnimNextDataInterfacePayload& InGraphPayload) const;
+
 	// Given an asset class, return whether an animation graph can be made via GetGraphFromObject.
 	bool CanGetGraphFromAssetClass(const UClass* InClass) const;
 
@@ -95,6 +104,7 @@ private:
 	TArray<FAnimNextAssetGraphMapping> AssetGraphMappings;
 
 	// Map derived from AssetGraphMappings
+	// Maps an asset type (e.g. anim sequence) to the index in the mappings array for the first entry for that object
 	TMap<TObjectKey<const UClass>, int32> AssetGraphMap;
 
 	// Flag to verify correct mapping rebuild behavior during startup
diff --git a/Engine/Plugins/Experimental/AnimNextAnimGraph/Source/AnimNextAnimGraph/Private/AnimNextAnimGraphSettings.cpp b/Engine/Plugins/Experimental/AnimNextAnimGraph/Source/AnimNextAnimGraph/Private/AnimNextAnimGraphSettings.cpp
--- a/Engine/Plugins/Experimental/AnimNextAnimGraph/Source/AnimNextAnimGraph/Private/AnimNextAnimGraphSettings.cpp
+++ b/Engine/Plugins/Experimental/AnimNextAnimGraph/Source/AnimNextAnimGraph/Private/AnimNextAnimGraphSettings.cpp
@@ -1,4 +1,4 @@
-﻿// Copyright Epic Games, Inc. All Rights Reserved.
+// Copyright Epic Games, Inc. All Rights Reserved.
 
 #include "AnimNextAnimGraphSettings.h"
 
@@ -9,7 +9,6 @@
 #include "Logging/StructuredLog.h"
 #include "Misc/EnumerateRange.h"
 #include "Modules/ModuleManager.h"
-#include "TraitInterfaces/IPayloadFactory.h"
 
 #if WITH_EDITOR
 void UAnimNextAnimGraphSettings::PostEditChangeProperty(FPropertyChangedEvent& PropertyChangedEvent)
@@ -19,7 +18,8 @@ void UAnimNextAnimGraphSettings::PostEditChangeProperty(FPropertyChangedEvent& P
 		if( PropertyChangedEvent.Property->GetFName() == GET_MEMBER_NAME_CHECKED(UAnimNextAnimGraphSettings, AssetGraphMappings) ||
 			PropertyChangedEvent.Property->GetFName() == GET_MEMBER_NAME_CHECKED(FAnimNextAssetGraphMapping, AssetType) ||
 			PropertyChangedEvent.Property->GetFName() == GET_MEMBER_NAME_CHECKED(FAnimNextAssetGraphMapping, AnimationGraph) ||
-			PropertyChangedEvent.Property->GetFName() == GET_MEMBER_NAME_CHECKED(FAnimNextAssetGraphMapping, Variable))
+			PropertyChangedEvent.Property->GetFName() == GET_MEMBER_NAME_CHECKED(FAnimNextAssetGraphMapping, Variable) ||
+			PropertyChangedEvent.Property->GetFName() == GET_MEMBER_NAME_CHECKED(FAnimNextAssetGraphMapping, RequiredDataInterfaces))
 		{
 			LoadAndRebuildMappings();
 		}
@@ -28,6 +28,12 @@ void UAnimNextAnimGraphSettings::PostEditChangeProperty(FPropertyChangedEvent& P
 #endif
 
 const UAnimNextAnimationGraph* UAnimNextAnimGraphSettings::GetGraphFromObject(const UObject* InObject) const
+{
+	FAnimNextDataInterfacePayload GraphPayload;
+	return GetGraphFromObject(InObject, GraphPayload);
+}
+
+const UAnimNextAnimationGraph* UAnimNextAnimGraphSettings::GetGraphFromObject(const UObject* InObject, const FAnimNextDataInterfacePayload& InGraphPayload) const
 {
 	check(InObject);
 	ensure(bMappingsBuiltAtLeastOnce);
@@ -59,7 +65,62 @@ const UAnimNextAnimationGraph* UAnimNextAnimGraphSettings::GetGraphFromObject(co
 		return nullptr;
 	}
 
-	return AssetGraphMappings[Index].AnimationGraph.Get();
+	TConstArrayView<FInstancedStruct> NativePayloads = InGraphPayload.GetNativePayloads();
+
+	int32 BestNumMatchedInterfaces = -1;
+	int32 BestMappingIndex = INDEX_NONE;
+
+	while (Index != INDEX_NONE)
+	{
+		const FAnimNextAssetGraphMapping& Mapping = AssetGraphMappings[Index];
+
+		bool bIsMappingValid = true;
+		int32 NumMatchedInterfaces = 0;
+
+		for (const TSoftObjectPtr<const UScriptStruct>& DataInterfacePtr : Mapping.RequiredDataInterfaces)
+		{
+			const UScriptStruct* DataInterface = DataInterfacePtr.Get();
+			if (DataInterface == nullptr)
+			{
+				continue;
+			}
+
+			bool bFoundRequiredInterface = false;
+			for (const FInstancedStruct& NativePayload : NativePayloads)
+			{
+				if (NativePayload.GetScriptStruct() == DataInterface)
+				{
+					bFoundRequiredInterface = true;
+					break;
+				}
+			}
+
+			if (bFoundRequiredInterface)
+			{
+				NumMatchedInterfaces++;
+			}
+			else
+			{
+				bIsMappingValid = false;
+				break;
+			}
+		}
+
+		// This mapping is valid, use it if its our best match
+		// If the number of interfaces matched is identical, then we'll pick the first mapping seen,
+		// which will be the last one found through the config files
+		if (bIsMappingValid && NumMatchedInterfaces > BestNumMatchedInterfaces)
+		{
+			BestNumMatchedInterfaces = NumMatchedInterfaces;
+			BestMappingIndex = Index;
+		}
+
+		// Try the next mapping for the same asset type
+		Index = Mapping.NextMappingIndex;
+	}
+
+	// Returns our best match, if one is found
+	return BestMappingIndex != INDEX_NONE ? AssetGraphMappings[BestMappingIndex].AnimationGraph.Get() : nullptr;
 }
 
 bool UAnimNextAnimGraphSettings::CanGetGraphFromAssetClass(const UClass* InClass) const
@@ -248,7 +309,31 @@ void UAnimNextAnimGraphSettings::LoadAndRebuildMappings()
 			continue;
 		}
 
-		AssetGraphMap.Add(AssetType, Mapping.GetIndex());
+		bool bAreRequiredInterfacesValid = true;
+		for (const TSoftObjectPtr<const UScriptStruct>& DataInterfacePtr : Mapping->RequiredDataInterfaces)
+		{
+			const UScriptStruct* DataInterface = DataInterfacePtr.LoadSynchronous();
+			if (DataInterface == nullptr)
+			{
+				UE_LOGFMT(LogAnimation, Error, "UAnimNextAnimGraphSettings::LoadAndRebuildMappings: Could not load required data interface class '{ClassName}'", DataInterfacePtr.ToString());
+				bAreRequiredInterfacesValid = false;
+				continue;
+			}
+		}
+
+		if (!bAreRequiredInterfacesValid)
+		{
+			continue;
+		}
+
+		// Find or add our new entry
+		int32& FirstAssetMappingIndex = AssetGraphMap.FindOrAdd(AssetType, INDEX_NONE);
+
+		// Update the next mapping index to point to our old entry (or INDEX_NONE if we are the first)
+		Mapping->NextMappingIndex = FirstAssetMappingIndex;
+
+		// Our new first mapping index is the current one
+		FirstAssetMappingIndex = Mapping.GetIndex();
 
 		// Warn if the variable is not settable in the public interface of the graph
 		if(AnimationGraph != nullptr && AssetType != nullptr && !Mapping->Variable.IsNone())
diff --git a/Engine/Plugins/Experimental/AnimNextAnimGraph/Source/AnimNextAnimGraph/Private/Injection/InjectionSiteTrait.cpp b/Engine/Plugins/Experimental/AnimNextAnimGraph/Source/AnimNextAnimGraph/Private/Injection/InjectionSiteTrait.cpp
--- a/Engine/Plugins/Experimental/AnimNextAnimGraph/Source/AnimNextAnimGraph/Private/Injection/InjectionSiteTrait.cpp
+++ b/Engine/Plugins/Experimental/AnimNextAnimGraph/Source/AnimNextAnimGraph/Private/Injection/InjectionSiteTrait.cpp
@@ -12,7 +12,6 @@
 #include "Logging/StructuredLog.h"
 #include "Module/AnimNextModuleInstance.h"
 #include "TraitInterfaces/IGraphFactory.h"
-#include "TraitInterfaces/IPayloadFactory.h"
 #include "TraitInterfaces/IBlendStack.h"
 #include "TraitInterfaces/IInertializerBlend.h"
 #include "TraitInterfaces/ISmoothBlend.h"
@@ -175,12 +174,13 @@ namespace UE::AnimNext
 				const FAnimNextInjectionRequestArgs& RequestArgs = Request->GetArgs();
 				State = EInjectionTrackerState::ActiveEvent;
 
+				BlendStackRequest.GraphPayload = MoveTemp(Request->GetMutableArgs().Payload);
+
 				BlendStackRequest.Type = IBlendStack::EGraphRequestType::Owned;
-				BlendStackRequest.AnimationGraph = IGraphFactory::GetGraphFromObjectWithFallback(Context, Binding, RequestArgs.Object);
-				if (BlendStackRequest.AnimationGraph != nullptr)
+				BlendStackRequest.AnimationGraph = IGraphFactory::GetGraphFromObjectWithFallback(Context, Binding, RequestArgs.Object, BlendStackRequest.GraphPayload);
+				if (BlendStackRequest.AnimationGraph == nullptr)
 				{
-					BlendStackRequest.GraphPayload = MoveTemp(Request->GetMutableArgs().Payload);
-					IPayloadFactory::GetPayloadForGraphWithFallback(Context, Binding, RequestArgs.Object, BlendStackRequest.AnimationGraph, BlendStackRequest.GraphPayload);
+					BlendStackRequest.GraphPayload = FAnimNextDataInterfacePayload();
 				}
 				BlendStackRequest.BindingModuleHandle = RequestArgs.BindingModuleHandle;
 				BlendStackRequest.BlendArgs = RequestArgs.BlendInSettings.Blend;
@@ -190,11 +190,12 @@ namespace UE::AnimNext
 			{
 				// Push request from a pin-value change
 				State = EInjectionTrackerState::ActivePin;
+
 				BlendStackRequest.Type = IBlendStack::EGraphRequestType::Owned;
-				BlendStackRequest.AnimationGraph = IGraphFactory::GetGraphFromObjectWithFallback(Context, Binding, Graph.Asset);
-				if (BlendStackRequest.AnimationGraph != nullptr)
+				BlendStackRequest.AnimationGraph = IGraphFactory::GetGraphFromObjectWithFallback(Context, Binding, Graph.Asset, BlendStackRequest.GraphPayload);
+				if (BlendStackRequest.AnimationGraph == nullptr)
 				{
-					IPayloadFactory::GetPayloadForGraphWithFallback(Context, Binding, Graph.Asset, BlendStackRequest.AnimationGraph, BlendStackRequest.GraphPayload);
+					BlendStackRequest.GraphPayload = FAnimNextDataInterfacePayload();
 				}
 
 				const FAnimNextInjectionBlendSettings& DefaultBlendInSettings = SharedData->GetDefaultBlendInSettings(Binding);
@@ -205,6 +206,7 @@ namespace UE::AnimNext
 			{
 				// No request, or invalid/unhandled, so push our source child
 				State = EInjectionTrackerState::ActiveSource;
+
 				BlendStackRequest.Type = IBlendStack::EGraphRequestType::Child;
 				BlendStackRequest.ChildPtr = Context.AllocateNodeInstance(Binding, SharedData->Source);
 
@@ -323,10 +325,10 @@ namespace UE::AnimNext
 						// Add a request to blend in our fallback graph using blend out time
 						IBlendStack::FGraphRequest BlendStackRequest;
 						BlendStackRequest.Type = IBlendStack::EGraphRequestType::Owned;
-						BlendStackRequest.AnimationGraph = IGraphFactory::GetGraphFromObjectWithFallback(Context, Binding, Graph.Asset);
-						if (BlendStackRequest.AnimationGraph != nullptr)
+						BlendStackRequest.AnimationGraph = IGraphFactory::GetGraphFromObjectWithFallback(Context, Binding, Graph.Asset, BlendStackRequest.GraphPayload);
+						if (BlendStackRequest.AnimationGraph == nullptr)
 						{
-							IPayloadFactory::GetPayloadForGraphWithFallback(Context, Binding, Graph.Asset, BlendStackRequest.AnimationGraph, BlendStackRequest.GraphPayload);
+							BlendStackRequest.GraphPayload = FAnimNextDataInterfacePayload();
 						}
 						BlendStackRequest.BlendArgs.BlendTime = BlendOutTime;
 						BlendStackRequest.BlendMode = TranslateBlendMode(RequestArgs.BlendOutSettings.BlendMode);
diff --git a/Engine/Plugins/Experimental/AnimNextAnimGraph/Source/AnimNextAnimGraph/Private/PlayAnim/PlayAnimSlotTrait.cpp b/Engine/Plugins/Experimental/AnimNextAnimGraph/Source/AnimNextAnimGraph/Private/PlayAnim/PlayAnimSlotTrait.cpp
--- a/Engine/Plugins/Experimental/AnimNextAnimGraph/Source/AnimNextAnimGraph/Private/PlayAnim/PlayAnimSlotTrait.cpp
+++ b/Engine/Plugins/Experimental/AnimNextAnimGraph/Source/AnimNextAnimGraph/Private/PlayAnim/PlayAnimSlotTrait.cpp
@@ -14,7 +14,6 @@
 #include "Logging/StructuredLog.h"
 #include "Module/ModuleEvents.h"
 #include "TraitInterfaces/IGraphFactory.h"
-#include "TraitInterfaces/IPayloadFactory.h"
 
 namespace UE::AnimNext
 {
@@ -172,9 +171,15 @@ namespace UE::AnimNext
 				FAnimNextInjectionRequestArgs& RequestArgs = Request->GetMutableArgs();
 				if (RequestArgs.Object)
 				{
-					AnimationGraph = IGraphFactory::GetGraphFromObjectWithFallback(Context, Binding, RequestArgs.Object);
+					Payload = MoveTemp(RequestArgs.Payload);
 
-					if (AnimationGraph != nullptr)
+					AnimationGraph = IGraphFactory::GetGraphFromObjectWithFallback(Context, Binding, RequestArgs.Object, Payload);
+
+					if (AnimationGraph == nullptr)
+					{
+						Payload = FAnimNextDataInterfacePayload();
+					}
+					else
 					{
 						// Check for re-entrancy and early-out if we are linking back to the current instance or one of its parents
 						const FName EntryPoint = AnimationGraph->DefaultEntryPoint;
@@ -189,10 +194,6 @@ namespace UE::AnimNext
 
 							OwnerGraphInstance = OwnerGraphInstance->GetParentGraphInstance();
 						}
-
-						Payload = MoveTemp(RequestArgs.Payload);
-
-						IPayloadFactory::GetPayloadForGraphWithFallback(Context, Binding, RequestArgs.Object, AnimationGraph, Payload);
 					}
 
 					BlendSettings = RequestArgs.BlendInSettings;
diff --git a/Engine/Plugins/Experimental/AnimNextAnimGraph/Source/AnimNextAnimGraph/Private/TraitInterfaces/IGraphFactory.cpp b/Engine/Plugins/Experimental/AnimNextAnimGraph/Source/AnimNextAnimGraph/Private/TraitInterfaces/IGraphFactory.cpp
--- a/Engine/Plugins/Experimental/AnimNextAnimGraph/Source/AnimNextAnimGraph/Private/TraitInterfaces/IGraphFactory.cpp
+++ b/Engine/Plugins/Experimental/AnimNextAnimGraph/Source/AnimNextAnimGraph/Private/TraitInterfaces/IGraphFactory.cpp
@@ -1,4 +1,4 @@
-﻿// Copyright Epic Games, Inc. All Rights Reserved.
+// Copyright Epic Games, Inc. All Rights Reserved.
 
 #include "TraitInterfaces/IGraphFactory.h"
 
@@ -22,29 +22,60 @@ namespace UE::AnimNext
 	}
 #endif // WITH_EDITOR
 
-	const UAnimNextAnimationGraph* IGraphFactory::GetGraphFromObject(FExecutionContext& Context, const TTraitBinding<IGraphFactory>& Binding, const UObject* InObject) const
+	const UAnimNextAnimationGraph* IGraphFactory::GetGraphFromObject(FExecutionContext& Context, const TTraitBinding<IGraphFactory>& Binding, const UObject* InObject, FAnimNextDataInterfacePayload& InOutPayload) const
 	{
 		TTraitBinding<IGraphFactory> SuperBinding;
 		if (Binding.GetStackInterfaceSuper(SuperBinding))
 		{
-			return SuperBinding.GetGraphFromObject(Context, InObject);
+			return SuperBinding.GetGraphFromObject(Context, InObject, InOutPayload);
 		}
+
 		return nullptr;
 	}
 
-	const UAnimNextAnimationGraph* IGraphFactory::GetGraphFromObjectWithFallback(FExecutionContext& Context, const FTraitBinding& InBinding, const UObject* InObject)
+	void IGraphFactory::CreatePayloadForObject(FExecutionContext& Context, const TTraitBinding<IGraphFactory>& Binding, const UObject* InObject, FAnimNextDataInterfacePayload& InOutPayload) const
 	{
+		TTraitBinding<IGraphFactory> SuperBinding;
+		if (Binding.GetStackInterfaceSuper(SuperBinding))
+		{
+			return SuperBinding.CreatePayloadForObject(Context, InObject, InOutPayload);
+		}
+	}
+
+	const UAnimNextAnimationGraph* IGraphFactory::GetGraphFromObjectWithFallback(FExecutionContext& Context, const FTraitBinding& InBinding, const UObject* InObject, FAnimNextDataInterfacePayload& InOutPayload)
+	{
+		// Graph creation works as follows:
+		//	- Caller needs to create and populate the native interface for the specified object (e.g. anim sequence)
+		//	  The caller knows the object it has and how to populate that native interface
+		//	- We then ask the trait stack to populate other native interfaces it knows about (e.g. sync group trait)
+		//	- Now we have everything we need to create the graph, ask the trait stack to do so
+		//		- If the trait stack fails to create a graph, create one using our user settings
+		//	- Now that we have a graph, create any remaining payloads that might be missing
+
+		const UAnimNextAnimationGraph* AnimationGraph = nullptr;
+
 		TTraitBinding<IGraphFactory> GraphFactoryBinding;
-		if(InBinding.GetStackInterface<IGraphFactory>(GraphFactoryBinding))
+		if (InBinding.GetStackInterface<IGraphFactory>(GraphFactoryBinding))
 		{
-			const UAnimNextAnimationGraph* AnimationGraph = GraphFactoryBinding.GetGraphFromObject(Context, InObject);
-			if(AnimationGraph != nullptr)
-			{
-				return AnimationGraph;
-			}
+			// Ask trait stack to create payloads it cares about
+			GraphFactoryBinding.CreatePayloadForObject(Context, InObject, InOutPayload);
+
+			// Ask trait stack to create a graph from our object/payload
+			AnimationGraph = GraphFactoryBinding.GetGraphFromObject(Context, InObject, InOutPayload);
 		}
 
-		return GetDefault<UAnimNextAnimGraphSettings>()->GetGraphFromObject(InObject);
-	}
+		// If we failed, fallback to our user settings
+		if (AnimationGraph == nullptr)
+		{
+			AnimationGraph = GetDefault<UAnimNextAnimGraphSettings>()->GetGraphFromObject(InObject, InOutPayload);
+		}
 
+		// Create missing payloads, if any
+		if (AnimationGraph != nullptr)
+		{
+			GetDefault<UAnimNextAnimGraphSettings>()->GetNativePayloadFromGraph(InObject, AnimationGraph, InOutPayload);
+		}
+
+		return AnimationGraph;
+	}
 }
diff --git a/Engine/Plugins/Experimental/AnimNextAnimGraph/Source/AnimNextAnimGraph/Private/TraitInterfaces/IPayloadFactory.cpp b/Engine/Plugins/Experimental/AnimNextAnimGraph/Source/AnimNextAnimGraph/Private/TraitInterfaces/IPayloadFactory.cpp
deleted file mode 100644
diff --git a/Engine/Plugins/Experimental/AnimNextAnimGraph/Source/AnimNextAnimGraph/Private/Traits/BlendSpacePlayer.cpp b/Engine/Plugins/Experimental/AnimNextAnimGraph/Source/AnimNextAnimGraph/Private/Traits/BlendSpacePlayer.cpp
--- a/Engine/Plugins/Experimental/AnimNextAnimGraph/Source/AnimNextAnimGraph/Private/Traits/BlendSpacePlayer.cpp
+++ b/Engine/Plugins/Experimental/AnimNextAnimGraph/Source/AnimNextAnimGraph/Private/Traits/BlendSpacePlayer.cpp
@@ -13,7 +13,6 @@
 #include "Graph/AnimNextGraphInstance.h"
 #include "TraitCore/NodeInstance.h"
 #include "TraitInterfaces/IGraphFactory.h"
-#include "TraitInterfaces/IPayloadFactory.h"
 
 
 namespace UE::AnimNext
@@ -120,30 +119,32 @@ namespace UE::AnimNext
 				for (int32 CurSampleIndex = 0; CurSampleIndex < BlendSamples.Num(); CurSampleIndex++)
 				{
 					const UAnimSequence* CurSequence = BlendSamples[CurSampleIndex].Animation;
-					const UAnimNextAnimationGraph* AnimationGraph = IGraphFactory::GetGraphFromObjectWithFallback(Context, Binding, CurSequence);
-					if (AnimationGraph != nullptr)
-					{
-						FSampleGraphState& TargetGraphState = InstanceData->SampleGraphs[CurSampleIndex];
-						// Create a default 'native data interface' payload via which we expect to communicate with the graph 
-						FInstancedStruct NewPayload;
-						NewPayload.InitializeAs<FAnimNextNativeDataInterface_AnimSequencePlayer>();
-						FAnimNextNativeDataInterface_AnimSequencePlayer& NewNativeInterface = NewPayload.GetMutable<FAnimNextNativeDataInterface_AnimSequencePlayer>();
-						NewNativeInterface.AnimSequence = CurSequence;
 
-						// @TODO: Playrate should be normalized too.
-						NewNativeInterface.PlayRate = SharedData->GetPlayRate(Binding);
+					FAnimNextDataInterfacePayload GraphPayload;
+
+					// Create a default 'native data interface' payload via which we expect to communicate with the graph 
+					FInstancedStruct NewPayload;
+					FAnimNextNativeDataInterface_AnimSequencePlayer& NewNativeInterface = NewPayload.InitializeAs<FAnimNextNativeDataInterface_AnimSequencePlayer>();
+					NewNativeInterface.AnimSequence = CurSequence;
 
-						// Assume start position is normalized.
-						NewNativeInterface.StartPosition = SharedData->GetStartPosition(Binding) * CurSequence->GetPlayLength();
-						NewNativeInterface.Loop = SharedData->GetbLoop(Binding);
+					// @TODO: Playrate should be normalized too.
+					NewNativeInterface.PlayRate = SharedData->GetPlayRate(Binding);
 
-						TargetGraphState.GraphPayload.AddNative(MoveTemp(NewPayload));
+					// Assume start position is normalized.
+					NewNativeInterface.StartPosition = SharedData->GetStartPosition(Binding) * CurSequence->GetPlayLength();
+					NewNativeInterface.Loop = SharedData->GetbLoop(Binding);
 
-						IPayloadFactory::GetPayloadForGraphWithFallback(Context, Binding, CurSequence, AnimationGraph, TargetGraphState.GraphPayload);
+					GraphPayload.AddNative(MoveTemp(NewPayload));
+
+					const UAnimNextAnimationGraph* AnimationGraph = IGraphFactory::GetGraphFromObjectWithFallback(Context, Binding, CurSequence, GraphPayload);
+					if (AnimationGraph != nullptr)
+					{
+						FSampleGraphState& TargetGraphState = InstanceData->SampleGraphs[CurSampleIndex];
 
 						const FName EntryPoint = NAME_None;
 						TargetGraphState.Instance = AnimationGraph->AllocateInstance(nullptr, &Binding.GetTraitPtr().GetNodeInstance()->GetOwner(), EntryPoint);
 						TargetGraphState.ChildPtr = TargetGraphState.Instance->GetGraphRootPtr();
+						TargetGraphState.GraphPayload = MoveTemp(GraphPayload);
 
 						TargetGraphState.Instance->BindPublicVariables(TargetGraphState.GraphPayload.Get());
 					}
diff --git a/Engine/Plugins/Experimental/AnimNextAnimGraph/Source/AnimNextAnimGraph/Private/Traits/BlendStackTrait.cpp b/Engine/Plugins/Experimental/AnimNextAnimGraph/Source/AnimNextAnimGraph/Private/Traits/BlendStackTrait.cpp
--- a/Engine/Plugins/Experimental/AnimNextAnimGraph/Source/AnimNextAnimGraph/Private/Traits/BlendStackTrait.cpp
+++ b/Engine/Plugins/Experimental/AnimNextAnimGraph/Source/AnimNextAnimGraph/Private/Traits/BlendStackTrait.cpp
@@ -14,7 +14,6 @@
 #include "TraitCore/ExecutionContext.h"
 #include "TraitInterfaces/IGraphFactory.h"
 #include "TraitInterfaces/IInertializerBlend.h"
-#include "TraitInterfaces/IPayloadFactory.h"
 
 namespace UE::AnimNext
 {
@@ -474,7 +473,8 @@ namespace UE::AnimNext
 		const UObject* DesiredObject = SharedData->GetObject(Binding);
 		if (DesiredObject != nullptr)
 		{
-			const UAnimNextAnimationGraph* AnimationGraph = IGraphFactory::GetGraphFromObjectWithFallback(Context, Binding, DesiredObject);
+			FAnimNextDataInterfacePayload GraphPayload;
+			const UAnimNextAnimationGraph* AnimationGraph = IGraphFactory::GetGraphFromObjectWithFallback(Context, Binding, DesiredObject, GraphPayload);
 			if (AnimationGraph != nullptr)
 			{
 				const int32 CurrentlyActiveGraphIndex = InstanceData->CurrentlyActiveGraphIndex;
@@ -482,9 +482,6 @@ namespace UE::AnimNext
 				const bool bIsEmpty = CurrentlyActiveGraphIndex == INDEX_NONE;
 				if (bForceBlend || bIsEmpty || (DesiredObject != InstanceData->ChildGraphs[CurrentlyActiveGraphIndex].Request.FactoryObject))
 				{
-					FAnimNextDataInterfacePayload GraphPayload;
-					IPayloadFactory::GetPayloadForGraphWithFallback(Context, Binding, DesiredObject, AnimationGraph, GraphPayload);
-
 					TTraitBinding<IBlendStack> BlendStackTrait;
 					Binding.GetStackInterface(BlendStackTrait);
 
@@ -509,7 +506,8 @@ namespace UE::AnimNext
 		const UObject* DesiredObject = SharedData->GetObject(Binding);
 		if (DesiredObject != nullptr)
 		{
-			const UAnimNextAnimationGraph* AnimationGraph = IGraphFactory::GetGraphFromObjectWithFallback(Context, Binding, DesiredObject);
+			FAnimNextDataInterfacePayload GraphPayload;
+			const UAnimNextAnimationGraph* AnimationGraph = IGraphFactory::GetGraphFromObjectWithFallback(Context, Binding, DesiredObject, GraphPayload);
 			if (AnimationGraph != nullptr)
 			{
 				TTraitBinding<IBlendStack> BlendStackTrait;
@@ -520,9 +518,6 @@ namespace UE::AnimNext
 				const bool bForceBlend = SharedData->GetbForceBlend(Binding);
 				if (bForceBlend || !ActiveGraphRequest || (DesiredObject != ActiveGraphRequest->FactoryObject))
 				{
-					FAnimNextDataInterfacePayload GraphPayload;
-					IPayloadFactory::GetPayloadForGraphWithFallback(Context, Binding, DesiredObject, AnimationGraph, GraphPayload);
-
 					IBlendStack::FGraphRequest GraphRequest;
 					GraphRequest.FactoryObject = DesiredObject;
 					GraphRequest.AnimationGraph = AnimationGraph;
diff --git a/Engine/Plugins/Experimental/AnimNextAnimGraph/Source/AnimNextAnimGraph/Private/Traits/MakeGraphFromAsset.cpp b/Engine/Plugins/Experimental/AnimNextAnimGraph/Source/AnimNextAnimGraph/Private/Traits/MakeGraphFromAsset.cpp
--- a/Engine/Plugins/Experimental/AnimNextAnimGraph/Source/AnimNextAnimGraph/Private/Traits/MakeGraphFromAsset.cpp
+++ b/Engine/Plugins/Experimental/AnimNextAnimGraph/Source/AnimNextAnimGraph/Private/Traits/MakeGraphFromAsset.cpp
@@ -1,4 +1,4 @@
-﻿// Copyright Epic Games, Inc. All Rights Reserved.
+// Copyright Epic Games, Inc. All Rights Reserved.
 
 #include "Traits/MakeGraphFromAsset.h"
 
@@ -13,10 +13,10 @@ namespace UE::AnimNext
 	#define TRAIT_INTERFACE_ENUMERATOR(GeneratorMacro) \
 		GeneratorMacro(IGraphFactory) \
 
-		GENERATE_ANIM_TRAIT_IMPLEMENTATION(FMakeGraphFromAssetTrait, TRAIT_INTERFACE_ENUMERATOR, NULL_ANIM_TRAIT_INTERFACE_ENUMERATOR, NULL_ANIM_TRAIT_EVENT_ENUMERATOR)
+	GENERATE_ANIM_TRAIT_IMPLEMENTATION(FMakeGraphFromAssetTrait, TRAIT_INTERFACE_ENUMERATOR, NULL_ANIM_TRAIT_INTERFACE_ENUMERATOR, NULL_ANIM_TRAIT_EVENT_ENUMERATOR)
 	#undef TRAIT_INTERFACE_ENUMERATOR
 
-	const UAnimNextAnimationGraph* FMakeGraphFromAssetTrait::GetGraphFromObject(FExecutionContext& Context, const TTraitBinding<IGraphFactory>& Binding, const UObject* InObject) const
+	const UAnimNextAnimationGraph* FMakeGraphFromAssetTrait::GetGraphFromObject(FExecutionContext& Context, const TTraitBinding<IGraphFactory>& Binding, const UObject* InObject, FAnimNextDataInterfacePayload& InOutPayload) const
 	{
 		if(InObject == nullptr)
 		{
diff --git a/Engine/Plugins/Experimental/AnimNextAnimGraph/Source/AnimNextAnimGraph/Private/Traits/MakeGraphFromAsset.h b/Engine/Plugins/Experimental/AnimNextAnimGraph/Source/AnimNextAnimGraph/Private/Traits/MakeGraphFromAsset.h
--- a/Engine/Plugins/Experimental/AnimNextAnimGraph/Source/AnimNextAnimGraph/Private/Traits/MakeGraphFromAsset.h
+++ b/Engine/Plugins/Experimental/AnimNextAnimGraph/Source/AnimNextAnimGraph/Private/Traits/MakeGraphFromAsset.h
@@ -32,6 +32,6 @@ namespace UE::AnimNext
 		};
 
 		// IGraphFactory impl
-		virtual const UAnimNextAnimationGraph* GetGraphFromObject(FExecutionContext& Context, const TTraitBinding<IGraphFactory>& Binding, const UObject* InObject) const;
+		virtual const UAnimNextAnimationGraph* GetGraphFromObject(FExecutionContext& Context, const TTraitBinding<IGraphFactory>& Binding, const UObject* InObject, FAnimNextDataInterfacePayload& InOutPayload) const;
 	};
 }
diff --git a/Engine/Plugins/Experimental/AnimNextAnimGraph/Source/AnimNextAnimGraph/Private/Traits/SynchronizeUsingGroups.cpp b/Engine/Plugins/Experimental/AnimNextAnimGraph/Source/AnimNextAnimGraph/Private/Traits/SynchronizeUsingGroups.cpp
--- a/Engine/Plugins/Experimental/AnimNextAnimGraph/Source/AnimNextAnimGraph/Private/Traits/SynchronizeUsingGroups.cpp
+++ b/Engine/Plugins/Experimental/AnimNextAnimGraph/Source/AnimNextAnimGraph/Private/Traits/SynchronizeUsingGroups.cpp
@@ -1,7 +1,10 @@
 // Copyright Epic Games, Inc. All Rights Reserved.
 
 #include "Traits/SynchronizeUsingGroups.h"
+
+#include "AnimNextDataInterfacePayload.h"
 #include "Graph/SyncGroup_GraphInstanceComponent.h"
+#include "GraphInterfaces/AnimNextNativeDataInterface_SynchronizeUsingGroups.h"
 
 namespace UE::AnimNext
 {
@@ -9,6 +12,7 @@ namespace UE::AnimNext
 
 	// Trait implementation boilerplate
 	#define TRAIT_INTERFACE_ENUMERATOR(GeneratorMacro) \
+		GeneratorMacro(IGraphFactory) \
 		GeneratorMacro(IGroupSynchronization) \
 		GeneratorMacro(ITimeline) \
 		GeneratorMacro(IUpdate) \
@@ -16,6 +20,12 @@ namespace UE::AnimNext
 	GENERATE_ANIM_TRAIT_IMPLEMENTATION(FSynchronizeUsingGroupsTrait, TRAIT_INTERFACE_ENUMERATOR, NULL_ANIM_TRAIT_INTERFACE_ENUMERATOR, NULL_ANIM_TRAIT_EVENT_ENUMERATOR)
 	#undef TRAIT_INTERFACE_ENUMERATOR
 
+	void FSynchronizeUsingGroupsTrait::FInstanceData::Construct(const FExecutionContext& Context, const FTraitBinding& Binding)
+	{
+		TTraitBinding<ITimeline> TimelineTrait;
+		bHasTimeline = Binding.GetStackInterface(TimelineTrait);
+	}
+
 	void FSynchronizeUsingGroupsTrait::OnBecomeRelevant(FUpdateTraversalContext& Context, const TTraitBinding<IUpdate>& Binding, const FTraitUpdateState& TraitState) const
 	{
 		IUpdate::OnBecomeRelevant(Context, Binding, TraitState);
@@ -29,6 +39,13 @@ namespace UE::AnimNext
 	{
 		FInstanceData* InstanceData = Binding.GetInstanceData<FInstanceData>();
 
+		if (!InstanceData->bHasTimeline)
+		{
+			// If we don't have a timeline on the trait stack, there is nothing to synchronize
+			IUpdate::PreUpdate(Context, Binding, TraitState);
+			return;
+		}
+
 		TTraitBinding<IGroupSynchronization> GroupSyncTrait;
 		Binding.GetStackInterface(GroupSyncTrait);
 
@@ -59,13 +76,13 @@ namespace UE::AnimNext
 		const FSharedData* SharedData = Binding.GetSharedData<FSharedData>();
 
 		FSyncGroupParameters Parameters;
-		Parameters.GroupName = SharedData->GroupName;
-		Parameters.GroupRole = SharedData->GroupRole;
-		Parameters.bMatchSyncPoint = SharedData->bMatchSyncPoint;
+		Parameters.GroupName = SharedData->GetGroupName(Binding);
+		Parameters.GroupRole = SharedData->GetGroupRole(Binding);
+		Parameters.bMatchSyncPoint = SharedData->GetbMatchSyncPoint(Binding);
 
 		// When the role is set to TransitionLeader or TransitionFollower, we will not be part of a sync group until we are full weight
 		// meaning we ignore the sync group while blending in, but not while blending out
-		if (SharedData->GroupRole == EAnimGroupSynchronizationRole::TransitionLeader || SharedData->GroupRole == EAnimGroupSynchronizationRole::TransitionFollower)
+		if (Parameters.GroupRole == EAnimGroupSynchronizationRole::TransitionLeader || Parameters.GroupRole == EAnimGroupSynchronizationRole::TransitionFollower)
 		{
 			const FInstanceData* InstanceData = Binding.GetInstanceData<FInstanceData>();
 
@@ -82,6 +99,11 @@ namespace UE::AnimNext
 	{
 		FInstanceData* InstanceData = Binding.GetInstanceData<FInstanceData>();
 
+		if (!InstanceData->bHasTimeline)
+		{
+			return;	//  No timeline on the trait stack, nothing to synchronize
+		}
+
 		// When the group advances the timeline, we thaw it to advance
 		InstanceData->bFreezeTimeline = false;
 
@@ -103,4 +125,22 @@ namespace UE::AnimNext
 
 		ITimeline::AdvanceBy(Context, Binding, DeltaTime);
 	}
+
+	void FSynchronizeUsingGroupsTrait::CreatePayloadForObject(FExecutionContext& Context, const TTraitBinding<IGraphFactory>& Binding, const UObject* InObject, FAnimNextDataInterfacePayload& InOutPayload) const
+	{
+		IGraphFactory::CreatePayloadForObject(Context, Binding, InObject, InOutPayload);
+
+		TTraitBinding<IGroupSynchronization> GroupSyncTrait;
+		Binding.GetStackInterface(GroupSyncTrait);
+
+		const FSyncGroupParameters GroupParameters = GroupSyncTrait.GetGroupParameters(Context);
+
+		FInstancedStruct SyncGroupPayload;
+		auto& NativeInterface = SyncGroupPayload.InitializeAs<FAnimNextNativeDataInterface_SynchronizeUsingGroups>();
+		NativeInterface.GroupName = GroupParameters.GroupName;
+		NativeInterface.GroupRole = GroupParameters.GroupRole;
+		NativeInterface.MatchSyncPoint = GroupParameters.bMatchSyncPoint;
+
+		InOutPayload.AddNative(MoveTemp(SyncGroupPayload));
+	}
 }
diff --git a/Engine/Plugins/Experimental/AnimNextAnimGraph/Source/AnimNextAnimGraph/Private/Traits/SynchronizeUsingGroups.h b/Engine/Plugins/Experimental/AnimNextAnimGraph/Source/AnimNextAnimGraph/Private/Traits/SynchronizeUsingGroups.h
--- a/Engine/Plugins/Experimental/AnimNextAnimGraph/Source/AnimNextAnimGraph/Private/Traits/SynchronizeUsingGroups.h
+++ b/Engine/Plugins/Experimental/AnimNextAnimGraph/Source/AnimNextAnimGraph/Private/Traits/SynchronizeUsingGroups.h
@@ -5,6 +5,7 @@
 #include "CoreMinimal.h"
 #include "TraitCore/Trait.h"
 #include "TraitInterfaces/IGroupSynchronization.h"
+#include "TraitInterfaces/IGraphFactory.h"
 #include "TraitInterfaces/IUpdate.h"
 #include "TraitInterfaces/ITimeline.h"
 
@@ -17,17 +18,26 @@ struct FAnimNextSynchronizeUsingGroupsTraitSharedData : public FAnimNextTraitSha
 
 	// The group name
 	// If no name is provided, this trait is inactive
-	UPROPERTY(EditAnywhere, Category = "Default", meta = (Inline))
+	UPROPERTY(EditAnywhere, Category = "Default")
 	FName GroupName;
 
 	// The role this player can assume within the group
-	UPROPERTY(EditAnywhere, Category = "Default", meta = (Inline))
+	UPROPERTY(EditAnywhere, Category = "Default")
 	EAnimGroupSynchronizationRole GroupRole = EAnimGroupSynchronizationRole::CanBeLeader;
 
 	// Whether or not to match the group sync point when joining as leader
 	// When disabled, the leader will snap the position of the group to its current position
-	UPROPERTY(EditAnywhere, Category = "Default", meta = (Inline))
+	UPROPERTY(EditAnywhere, Category = "Default")
 	bool bMatchSyncPoint = true;
+
+	// Latent pin support boilerplate
+	#define TRAIT_LATENT_PROPERTIES_ENUMERATOR(GeneratorMacro) \
+		GeneratorMacro(GroupName) \
+		GeneratorMacro(GroupRole) \
+		GeneratorMacro(bMatchSyncPoint) \
+
+	GENERATE_TRAIT_LATENT_PROPERTIES(FAnimNextSynchronizeUsingGroupsTraitSharedData, TRAIT_LATENT_PROPERTIES_ENUMERATOR)
+	#undef TRAIT_LATENT_PROPERTIES_ENUMERATOR
 };
 
 namespace UE::AnimNext
@@ -37,7 +47,7 @@ namespace UE::AnimNext
 	 * 
 	 * A trait that synchronizes animation sequence playback using named groups.
 	 */
-	struct FSynchronizeUsingGroupsTrait : FAdditiveTrait, IUpdate, IGroupSynchronization, ITimeline
+	struct FSynchronizeUsingGroupsTrait : FAdditiveTrait, IUpdate, IGroupSynchronization, ITimeline, IGraphFactory
 	{
 		DECLARE_ANIM_TRAIT(FSynchronizeUsingGroupsTrait, FAdditiveTrait)
 
@@ -45,8 +55,11 @@ namespace UE::AnimNext
 
 		struct FInstanceData : FTraitInstanceData
 		{
+			void Construct(const FExecutionContext& Context, const FTraitBinding& Binding);
+
 			bool bFreezeTimeline = false;
 			bool bHasReachedFullWeight = false;
+			bool bHasTimeline = false;
 		};
 
 		// IUpdate impl
@@ -59,5 +72,8 @@ namespace UE::AnimNext
 
 		// ITimeline impl
 		virtual void AdvanceBy(FExecutionContext& Context, const TTraitBinding<ITimeline>& Binding, float DeltaTime) const override;
+
+		// IGraphFactory impl
+		virtual void CreatePayloadForObject(FExecutionContext& Context, const TTraitBinding<IGraphFactory>& Binding, const UObject* InObject, FAnimNextDataInterfacePayload& InOutPayload) const override;
 	};
 }
diff --git a/Engine/Plugins/Experimental/AnimNextAnimGraph/Source/AnimNextAnimGraph/Public/GraphInterfaces/AnimNextNativeDataInterface_SynchronizeUsingGroups.h b/Engine/Plugins/Experimental/AnimNextAnimGraph/Source/AnimNextAnimGraph/Public/GraphInterfaces/AnimNextNativeDataInterface_SynchronizeUsingGroups.h
new file mode 100644
--- /dev/null
+++ b/Engine/Plugins/Experimental/AnimNextAnimGraph/Source/AnimNextAnimGraph/Public/GraphInterfaces/AnimNextNativeDataInterface_SynchronizeUsingGroups.h
@@ -0,0 +1,33 @@
+// Copyright Epic Games, Inc. All Rights Reserved.
+
+#pragma once
+
+#include "CoreMinimal.h"
+#include "DataInterface/AnimNextNativeDataInterface.h"
+#include "TraitInterfaces/IGroupSynchronization.h"
+
+#include "AnimNextNativeDataInterface_SynchronizeUsingGroups.generated.h"
+
+
+/**
+* Native interface to built-in SynchronizeUsingGroups trait
+ */
+USTRUCT(BlueprintType, DisplayName = "AnimNext Group Synchronization Data")
+struct FAnimNextNativeDataInterface_SynchronizeUsingGroups : public FAnimNextNativeDataInterface
+{
+	GENERATED_BODY()
+
+	// The group name
+	// If no name is provided, no synchronization occurs
+	UPROPERTY(EditAnywhere, BlueprintReadWrite, Category = "Group Synchronization", meta = (EnableCategories))
+	FName GroupName;
+
+	// The role this player can assume within the group
+	UPROPERTY(EditAnywhere, BlueprintReadWrite, Category = "Group Synchronization", meta = (EnableCategories))
+	EAnimGroupSynchronizationRole GroupRole = EAnimGroupSynchronizationRole::CanBeLeader;
+
+	// Whether or not to match the group sync point when joining as leader
+	// When disabled, the leader will snap the position of the group to its current position
+	UPROPERTY(EditAnywhere, BlueprintReadWrite, Category = "Group Synchronization", meta = (EnableCategories))
+	bool MatchSyncPoint = true;
+};
diff --git a/Engine/Plugins/Experimental/AnimNextAnimGraph/Source/AnimNextAnimGraph/Public/TraitInterfaces/IBlendStack.h b/Engine/Plugins/Experimental/AnimNextAnimGraph/Source/AnimNextAnimGraph/Public/TraitInterfaces/IBlendStack.h
--- a/Engine/Plugins/Experimental/AnimNextAnimGraph/Source/AnimNextAnimGraph/Public/TraitInterfaces/IBlendStack.h
+++ b/Engine/Plugins/Experimental/AnimNextAnimGraph/Source/AnimNextAnimGraph/Public/TraitInterfaces/IBlendStack.h
@@ -7,7 +7,6 @@
 #include "TraitCore/TraitBinding.h"
 #include "AlphaBlend.h"
 #include "AnimNextDataInterfacePayload.h"
-#include "IPayloadFactory.h"
 #include "Graph/AnimNextAnimationGraph.h"
 #include "Module/ModuleHandle.h"
 
diff --git a/Engine/Plugins/Experimental/AnimNextAnimGraph/Source/AnimNextAnimGraph/Public/TraitInterfaces/IGraphFactory.h b/Engine/Plugins/Experimental/AnimNextAnimGraph/Source/AnimNextAnimGraph/Public/TraitInterfaces/IGraphFactory.h
--- a/Engine/Plugins/Experimental/AnimNextAnimGraph/Source/AnimNextAnimGraph/Public/TraitInterfaces/IGraphFactory.h
+++ b/Engine/Plugins/Experimental/AnimNextAnimGraph/Source/AnimNextAnimGraph/Public/TraitInterfaces/IGraphFactory.h
@@ -1,4 +1,4 @@
-﻿// Copyright Epic Games, Inc. All Rights Reserved.
+// Copyright Epic Games, Inc. All Rights Reserved.
 
 #pragma once
 
@@ -8,6 +8,7 @@
 #include "TraitCore/TraitBinding.h"
 
 class UAnimNextAnimationGraph;
+struct FAnimNextDataInterfacePayload;
 
 namespace UE::AnimNext
 {
@@ -22,14 +23,21 @@ namespace UE::AnimNext
 
 		// Gets a graph and a set of native interfaces from an object
 		// @param	InObject				The object to manufacture a graph for (e.g. AnimSequence). If this is a UAnimNextAnimationGraph, it should be returned verbatim
-		// @return The graph to be used.
-		virtual const UAnimNextAnimationGraph* GetGraphFromObject(FExecutionContext& Context, const TTraitBinding<IGraphFactory>& Binding, const UObject* InObject) const;
+		// @param	InOutGraphPayload			Payload that can be used to communicate with a graph instance
+		// @return	The graph to be used.
+		virtual const UAnimNextAnimationGraph* GetGraphFromObject(FExecutionContext& Context, const TTraitBinding<IGraphFactory>& Binding, const UObject* InObject, FAnimNextDataInterfacePayload& InOutPayload) const;
+
+		// Creates data interface payload to bind when creating a graph for an object
+		// @param	InObject					The object (e.g. AnimSequence). If this is a UAnimNextAnimationGraph, it should be the same object as InGraph
+		// @param	InOutGraphPayload			Payload that can be used to communicate with a graph instance
+		virtual void CreatePayloadForObject(FExecutionContext& Context, const TTraitBinding<IGraphFactory>& Binding, const UObject* InObject, FAnimNextDataInterfacePayload& InOutPayload) const;
 
 		// Gets a graph and a set of native interfaces from an object
 		// Falls back to using user settings if no interface is found in the current stack, or if an interface was found but returns nullptr
 		// @param	InObject				The object to manufacture a graph for (e.g. AnimSequence). If this is a UAnimNextAnimationGraph, it should be returned verbatim
-		// @return The graph to be used.
-		static const UAnimNextAnimationGraph* GetGraphFromObjectWithFallback(FExecutionContext& Context, const FTraitBinding& InBinding, const UObject* InObject);
+		// @param	InOutGraphPayload			Payload that can be used to communicate with a graph instance
+		// @return	The graph to be used.
+		static const UAnimNextAnimationGraph* GetGraphFromObjectWithFallback(FExecutionContext& Context, const FTraitBinding& InBinding, const UObject* InObject, FAnimNextDataInterfacePayload& InOutPayload);
 
 #if WITH_EDITOR
 		virtual const FText& GetDisplayName() const override;
@@ -44,9 +52,15 @@ namespace UE::AnimNext
 	struct TTraitBinding<IGraphFactory> : FTraitBinding
 	{
 		// @see IGraphFactory::GetGraphFromObject
-		const UAnimNextAnimationGraph* GetGraphFromObject(FExecutionContext& Context, const UObject* InObject) const
+		const UAnimNextAnimationGraph* GetGraphFromObject(FExecutionContext& Context, const UObject* InObject, FAnimNextDataInterfacePayload& InOutPayload) const
+		{
+			return GetInterface()->GetGraphFromObject(Context, *this, InObject, InOutPayload);
+		}
+
+		// @see IGraphFactory::CreatePayloadForObject
+		void CreatePayloadForObject(FExecutionContext& Context, const UObject* InObject, FAnimNextDataInterfacePayload& InOutPayload) const
 		{
-			return GetInterface()->GetGraphFromObject(Context, *this, InObject);
+			GetInterface()->CreatePayloadForObject(Context, *this, InObject, InOutPayload);
 		}
 
 	protected:
diff --git a/Engine/Plugins/Experimental/AnimNextAnimGraph/Source/AnimNextAnimGraph/Public/TraitInterfaces/IPayloadFactory.h b/Engine/Plugins/Experimental/AnimNextAnimGraph/Source/AnimNextAnimGraph/Public/TraitInterfaces/IPayloadFactory.h
deleted file mode 100644
diff --git a/Engine/Plugins/Experimental/AnimNextChooser/Source/AnimNextChooser/Private/MakeGraphFromChooser.cpp b/Engine/Plugins/Experimental/AnimNextChooser/Source/AnimNextChooser/Private/MakeGraphFromChooser.cpp
--- a/Engine/Plugins/Experimental/AnimNextChooser/Source/AnimNextChooser/Private/MakeGraphFromChooser.cpp
+++ b/Engine/Plugins/Experimental/AnimNextChooser/Source/AnimNextChooser/Private/MakeGraphFromChooser.cpp
@@ -1,4 +1,4 @@
-﻿// Copyright Epic Games, Inc. All Rights Reserved.
+// Copyright Epic Games, Inc. All Rights Reserved.
 
 #include "MakeGraphFromChooser.h"
 
@@ -18,7 +18,7 @@ namespace UE::AnimNext
 		GENERATE_ANIM_TRAIT_IMPLEMENTATION(FMakeGraphFromChooserTrait, TRAIT_INTERFACE_ENUMERATOR, NULL_ANIM_TRAIT_INTERFACE_ENUMERATOR, NULL_ANIM_TRAIT_EVENT_ENUMERATOR)
 	#undef TRAIT_INTERFACE_ENUMERATOR
 
-	const UAnimNextAnimationGraph* FMakeGraphFromChooserTrait::GetGraphFromObject(FExecutionContext& Context, const TTraitBinding<IGraphFactory>& Binding, const UObject* InObject) const
+	const UAnimNextAnimationGraph* FMakeGraphFromChooserTrait::GetGraphFromObject(FExecutionContext& Context, const TTraitBinding<IGraphFactory>& Binding, const UObject* InObject, FAnimNextDataInterfacePayload& InOutPayload) const
 	{
 		if(InObject == nullptr)
 		{
diff --git a/Engine/Plugins/Experimental/AnimNextChooser/Source/AnimNextChooser/Private/MakeGraphFromChooser.h b/Engine/Plugins/Experimental/AnimNextChooser/Source/AnimNextChooser/Private/MakeGraphFromChooser.h
--- a/Engine/Plugins/Experimental/AnimNextChooser/Source/AnimNextChooser/Private/MakeGraphFromChooser.h
+++ b/Engine/Plugins/Experimental/AnimNextChooser/Source/AnimNextChooser/Private/MakeGraphFromChooser.h
@@ -45,6 +45,6 @@ namespace UE::AnimNext
 		};
 
 		// IGraphFactory impl
-		virtual const UAnimNextAnimationGraph* GetGraphFromObject(FExecutionContext& Context, const TTraitBinding<IGraphFactory>& Binding, const UObject* InObject) const;
+		virtual const UAnimNextAnimationGraph* GetGraphFromObject(FExecutionContext& Context, const TTraitBinding<IGraphFactory>& Binding, const UObject* InObject, FAnimNextDataInterfacePayload& InOutPayload) const override;
 	};
 }
diff --git a/Engine/Plugins/Experimental/AnimNextPoseSearch/Source/AnimNextPoseSearch/Private/MotionMatchingTrait.cpp b/Engine/Plugins/Experimental/AnimNextPoseSearch/Source/AnimNextPoseSearch/Private/MotionMatchingTrait.cpp
--- a/Engine/Plugins/Experimental/AnimNextPoseSearch/Source/AnimNextPoseSearch/Private/MotionMatchingTrait.cpp
+++ b/Engine/Plugins/Experimental/AnimNextPoseSearch/Source/AnimNextPoseSearch/Private/MotionMatchingTrait.cpp
@@ -270,39 +270,38 @@ namespace UE::AnimNext
 			UE::AnimNext::IBlendStack::FGraphRequest NewGraphRequest;
 			NewGraphRequest.BlendArgs = SharedData->BlendArguments;
 
-			const UAnimNextAnimationGraph* AnimationGraph = IGraphFactory::GetGraphFromObjectWithFallback(Context, Binding, AnimationAsset);
-			if(AnimationGraph != nullptr)
+			// Create a default 'native data interface' payload via which we expect to communicate with the graph 
+			FInstancedStruct NewPayload;
+			// @TODO: Support asset->playload mapping in a more data-driven way
+			if (AnimationAsset->IsA<UBlendSpace>())
 			{
-				// Create a default 'native data interface' payload via which we expect to communicate with the graph 
-				FInstancedStruct NewPayload;
-				// @TODO: Support asset->playload mapping in a more data-driven way
-				if (AnimationAsset->IsA<UBlendSpace>())
-				{
-					NewPayload.InitializeAs<FAnimNextNativeDataInterface_BlendSpacePlayer>();
-					FAnimNextNativeDataInterface_BlendSpacePlayer& NewNativeInterface = NewPayload.GetMutable<FAnimNextNativeDataInterface_BlendSpacePlayer>();
-					NewNativeInterface.BlendSpace = Cast<UBlendSpace>(AnimationAsset);
-					
-					NewNativeInterface.XAxisSamplePoint = BlendSpaceParameters.X;
-					NewNativeInterface.YAxisSamplePoint = BlendSpaceParameters.Y;
-
-					NewNativeInterface.PlayRate = DesiredPlayRate;
-					NewNativeInterface.StartPosition = MotionMatchingState.CurrentSearchResult.AssetTime;
-					NewNativeInterface.bLoop = bLoop;
-				}
-				else
-				{
-					NewPayload.InitializeAs<FAnimNextNativeDataInterface_AnimSequencePlayer>();
-					FAnimNextNativeDataInterface_AnimSequencePlayer& NewNativeInterface = NewPayload.GetMutable<FAnimNextNativeDataInterface_AnimSequencePlayer>();
-					NewNativeInterface.AnimSequence = Cast<UAnimSequence>(AnimationAsset);
-					NewNativeInterface.PlayRate = DesiredPlayRate;
-					NewNativeInterface.StartPosition = MotionMatchingState.CurrentSearchResult.AssetTime;
-					NewNativeInterface.Loop = bLoop;
-				}
+				FAnimNextNativeDataInterface_BlendSpacePlayer& NewNativeInterface = NewPayload.InitializeAs<FAnimNextNativeDataInterface_BlendSpacePlayer>();
+				NewNativeInterface.BlendSpace = Cast<UBlendSpace>(AnimationAsset);
 
-				NewGraphRequest.GraphPayload.AddNative(MoveTemp(NewPayload));
 
-				IPayloadFactory::GetPayloadForGraphWithFallback(Context, Binding, AnimationAsset, AnimationGraph, NewGraphRequest.GraphPayload);
+				NewNativeInterface.XAxisSamplePoint = BlendSpaceParameters.X;
+				NewNativeInterface.YAxisSamplePoint = BlendSpaceParameters.Y;
 
+				NewNativeInterface.PlayRate = DesiredPlayRate;
+				NewNativeInterface.StartPosition = MotionMatchingState.CurrentSearchResult.AssetTime;
+				NewNativeInterface.bLoop = bLoop;
+			}
+			else
+			{
+				FAnimNextNativeDataInterface_AnimSequencePlayer& NewNativeInterface = NewPayload.InitializeAs<FAnimNextNativeDataInterface_AnimSequencePlayer>();
+				NewNativeInterface.AnimSequence = Cast<UAnimSequence>(AnimationAsset);
+				NewNativeInterface.PlayRate = DesiredPlayRate;
+				NewNativeInterface.StartPosition = MotionMatchingState.CurrentSearchResult.AssetTime;
+				NewNativeInterface.Loop = bLoop;
+			}
+
+			FAnimNextDataInterfacePayload GraphPayload;
+			GraphPayload.AddNative(MoveTemp(NewPayload));
+
+			const UAnimNextAnimationGraph* AnimationGraph = IGraphFactory::GetGraphFromObjectWithFallback(Context, Binding, AnimationAsset, GraphPayload);
+			if(AnimationGraph != nullptr)
+			{
+				NewGraphRequest.GraphPayload = MoveTemp(GraphPayload);
 				NewGraphRequest.AnimationGraph = AnimationGraph;
 
 				BlendStackBinding.PushGraph(Context, MoveTemp(NewGraphRequest));
diff --git a/Engine/Plugins/Experimental/AnimNextStateTree/Source/AnimNextStateTree/Private/AnimNextStateTreeContext.cpp b/Engine/Plugins/Experimental/AnimNextStateTree/Source/AnimNextStateTree/Private/AnimNextStateTreeContext.cpp
--- a/Engine/Plugins/Experimental/AnimNextStateTree/Source/AnimNextStateTree/Private/AnimNextStateTreeContext.cpp
+++ b/Engine/Plugins/Experimental/AnimNextStateTree/Source/AnimNextStateTree/Private/AnimNextStateTreeContext.cpp
@@ -25,15 +25,12 @@ bool FAnimNextStateTreeTraitContext::PushAssetOntoBlendStack(TNonNullPtr<UObject
 	}
 
 	// Get a graph to instantiate
-	const UAnimNextAnimationGraph* AnimationGraph = IGraphFactory::GetGraphFromObjectWithFallback(*Context, *Binding, InAsset);
+	const UAnimNextAnimationGraph* AnimationGraph = IGraphFactory::GetGraphFromObjectWithFallback(*Context, *Binding, InAsset, InPayload);
 	if(AnimationGraph == nullptr)
 	{
 		return false;
 	}
 
-	// Get a payload to instantiate to control the graph
-	IPayloadFactory::GetPayloadForGraphWithFallback(*Context, *Binding, InAsset, AnimationGraph, InPayload);
-
 	IBlendStack::FGraphRequest NewGraphRequest;
 	NewGraphRequest.BlendArgs = InBlendArguments;
 	NewGraphRequest.FactoryObject = InAsset;
