Add sync group support to AnimNext graph factories

#virtualized

diff --git a/Engine/Plugins/Experimental/AnimNextAnimGraph/Config/BaseAnimNextAnimGraph.ini b/Engine/Plugins/Experimental/AnimNextAnimGraph/Config/BaseAnimNextAnimGraph.ini
--- a/Engine/Plugins/Experimental/AnimNextAnimGraph/Config/BaseAnimNextAnimGraph.ini
+++ b/Engine/Plugins/Experimental/AnimNextAnimGraph/Config/BaseAnimNextAnimGraph.ini
@@ -2,7 +2,9 @@
 [/Script/AnimNextAnimGraph.AnimNextAnimGraphSettings]
 DefaultRunGraphHost=/AnimNextAnimGraph/FactoryGraphs/AG_DefaultGraphHost.AG_DefaultGraphHost
 +AssetGraphMappings=(AssetType="/Script/Engine.AnimSequence",AnimationGraph="/AnimNextAnimGraph/FactoryGraphs/AG_AnimSequencePlayer.AG_AnimSequencePlayer",Variable="AnimSequence")
++AssetGraphMappings=(AssetType="/Script/Engine.AnimSequence",AnimationGraph="/AnimNextAnimGraph/FactoryGraphs/AG_SyncedAnimSequencePlayer.AG_SyncedAnimSequencePlayer",Variable="AnimSequence", RequiredDataInterfaces=("/Script/AnimNextAnimGraph.AnimNextNativeDataInterface_AnimSequencePlayer", "/Script/AnimNextAnimGraph.AnimNextNativeDataInterface_SynchronizeUsingGroups"))
 +AssetGraphMappings=(AssetType="/Script/Engine.BlendSpace",AnimationGraph="/AnimNextAnimGraph/FactoryGraphs/AG_BlendSpacePlayer.AG_BlendSpacePlayer",Variable="BlendSpace")
++AssetGraphMappings=(AssetType="/Script/Engine.BlendSpace",AnimationGraph="/AnimNextAnimGraph/FactoryGraphs/AG_SyncedBlendSpacePlayer.AG_SyncedBlendSpacePlayer",Variable="BlendSpace", RequiredDataInterfaces=("/Script/AnimNextAnimGraph.AnimNextNativeDataInterface_BlendSpacePlayer", "/Script/AnimNextAnimGraph.AnimNextNativeDataInterface_SynchronizeUsingGroups"))
 
 [CoreRedirects]
 ; Traits rename
