Fix time based group syncing in AnimNext

diff --git a/Engine/Plugins/Experimental/AnimNextAnimGraph/Source/AnimNextAnimGraph/Private/Graph/SyncGroup_GraphInstanceComponent.cpp b/Engine/Plugins/Experimental/AnimNextAnimGraph/Source/AnimNextAnimGraph/Private/Graph/SyncGroup_GraphInstanceComponent.cpp
--- a/Engine/Plugins/Experimental/AnimNextAnimGraph/Source/AnimNextAnimGraph/Private/Graph/SyncGroup_GraphInstanceComponent.cpp
+++ b/Engine/Plugins/Experimental/AnimNextAnimGraph/Source/AnimNextAnimGraph/Private/Graph/SyncGroup_GraphInstanceComponent.cpp
@@ -47,9 +47,9 @@ namespace UE::AnimNext
 			const FSyncGroupMember* State = nullptr;
 
 			FTraitStackBinding TraitStack;
-			TTraitBinding<IGroupSynchronization> GroupSyncTrait;
 			
 			FTimelineState TimelineState;
+			TTraitBinding<ITimeline> TimelineTrait;
 
 			FTimelineSyncMarkerArray SyncMarkers;
 			bool bUseMarkerSyncing = false;
@@ -58,8 +58,8 @@ namespace UE::AnimNext
 			{
 				State = &InState;
 				TraitStack.Reset();
-				GroupSyncTrait.Reset();
 				TimelineState.Reset();
+				TimelineTrait.Reset();
 				SyncMarkers.Reset();
 				bUseMarkerSyncing = false;
 			}
@@ -186,43 +186,37 @@ namespace UE::AnimNext
 #endif
 		}
 
-		static void CollectValidMarkers(FExecutionContext& Context, FSyncGroupContext& GroupContext)
+		static void BuildGroupState(FExecutionContext& Context, FSyncGroupContext& GroupContext)
 		{
 			const int32 NumMembers = GroupContext.Members.Num();
-			if (NumMembers == 1)
-			{
-				return;	// If the leader is alone, we have no need for syncing
-			}
-
 			const int32 LeaderIndex = GroupContext.LeaderIndex;
-
 			FSyncGroupMemberContext& GroupLeaderContext = GroupContext.Members[LeaderIndex];
-			const FSyncGroupMember& GroupLeader = *GroupLeaderContext.State;
-
-			FTraitStackBinding& LeaderTraitStack = GroupLeaderContext.TraitStack;
-
-			TTraitBinding<ITimeline> TimelineTrait;
 
 			bool bCanGroupUseMarkerSyncing = false;
 
 			// Get the group sync marker names from the leader
 			{
+				const FSyncGroupMember& GroupLeader = *GroupLeaderContext.State;
+
+				FTraitStackBinding& LeaderTraitStack = GroupLeaderContext.TraitStack;
+
 				Context.BindTo(GroupLeader.TraitPtr);
 				ensure(Context.GetStack(GroupLeader.TraitPtr, LeaderTraitStack));
+				ensure(LeaderTraitStack.GetInterface(GroupLeaderContext.TimelineTrait));
+
+				if (NumMembers == 1)
+				{
+					return;	// If the leader is alone, we have no need for syncing, we'll advance the leader normally
+				}
 
 				FTimelineSyncMarkerArray& SyncMarkers = GroupLeaderContext.SyncMarkers;
 
 				SyncMarkers.Reset();
-				if (LeaderTraitStack.GetInterface(TimelineTrait))
-				{
-					TimelineTrait.GetSyncMarkers(Context, SyncMarkers);
-					CheckSyncMarkersSorted(SyncMarkers);
+				GroupLeaderContext.TimelineTrait.GetSyncMarkers(Context, SyncMarkers);
+				CheckSyncMarkersSorted(SyncMarkers);
 
-					bCanGroupUseMarkerSyncing = !SyncMarkers.IsEmpty();
-					GroupLeaderContext.bUseMarkerSyncing = bCanGroupUseMarkerSyncing;
-
-					GroupLeaderContext.TimelineState = TimelineTrait.GetState(Context);
-				}
+				bCanGroupUseMarkerSyncing = !SyncMarkers.IsEmpty();
+				GroupLeaderContext.bUseMarkerSyncing = bCanGroupUseMarkerSyncing;
 			}
 
 			// If our leader has sync markers, iterate over every member of the group and remove markers
@@ -270,16 +264,12 @@ namespace UE::AnimNext
 
 					Context.BindTo(GroupMember.TraitPtr);
 					ensure(Context.GetStack(GroupMember.TraitPtr, MemberTraitStack));
+					ensure(MemberTraitStack.GetInterface(GroupMemberContext.TimelineTrait));
 
 					FTimelineSyncMarkerArray& SyncMarkers = GroupMemberContext.SyncMarkers;
 					SyncMarkers.Reset();
-					if (MemberTraitStack.GetInterface(TimelineTrait))
-					{
-						TimelineTrait.GetSyncMarkers(Context, SyncMarkers);
-						CheckSyncMarkersSorted(SyncMarkers);
-
-						GroupMemberContext.TimelineState = TimelineTrait.GetState(Context);
-					}
+					GroupMemberContext.TimelineTrait.GetSyncMarkers(Context, SyncMarkers);
+					CheckSyncMarkersSorted(SyncMarkers);
 
 					// Members without markers can still use time based syncing
 					const bool bCanMemberUseMarkerSyncing = !SyncMarkers.IsEmpty();
@@ -322,6 +312,25 @@ namespace UE::AnimNext
 
 				bCanGroupUseMarkerSyncing = !ValidMarkers.IsEmpty();
 			}
+			else
+			{
+				for (int32 MemberIndex = 0; MemberIndex < NumMembers; ++MemberIndex)
+				{
+					if (MemberIndex == LeaderIndex)
+					{
+						continue;	// Ignore the leader
+					}
+
+					FSyncGroupMemberContext& GroupMemberContext = GroupContext.Members[MemberIndex];
+					const FSyncGroupMember& GroupMember = *GroupMemberContext.State;
+
+					FTraitStackBinding& MemberTraitStack = GroupMemberContext.TraitStack;
+
+					Context.BindTo(GroupMember.TraitPtr);
+					ensure(Context.GetStack(GroupMember.TraitPtr, MemberTraitStack));
+					ensure(MemberTraitStack.GetInterface(GroupMemberContext.TimelineTrait));
+				}
+			}
 
 			GroupContext.bCanGroupUseMarkerSyncing = bCanGroupUseMarkerSyncing;
 		}
@@ -781,12 +790,13 @@ namespace UE::AnimNext
 			FTraitStackBinding& LeaderTraitStack = GroupLeaderContext.TraitStack;
 
 			TTraitBinding<IGroupSynchronization> GroupSyncTrait;
-			TTraitBinding<ITimeline> TimelineTrait;
 
 			Context.BindTo(GroupLeader.TraitPtr);
 			ensure(Context.GetStack(GroupLeader.TraitPtr, LeaderTraitStack));
 			ensure(LeaderTraitStack.GetInterface(GroupSyncTrait));
-			ensure(LeaderTraitStack.GetInterface(TimelineTrait));
+
+			// Cache our starting timeline state
+			GroupLeaderContext.TimelineState = GroupLeaderContext.TimelineTrait.GetState(Context);
 
 			float DeltaTime = GroupLeader.TraitState.GetDeltaTime();
 
@@ -848,7 +858,7 @@ namespace UE::AnimNext
 				GroupSyncTrait.AdvanceBy(Context, DeltaTime);
 
 				// Compute our new timeline state
-				const FTimelineState NewTimelineState = TimelineTrait.GetState(Context);
+				const FTimelineState NewTimelineState = GroupLeaderContext.TimelineTrait.GetState(Context);
 
 				GroupContext.LeaderTimelineState = NewTimelineState;
 				GroupContext.LeaderDeltaTime = DeltaTime;
@@ -866,8 +876,8 @@ namespace UE::AnimNext
 				UE_LOG(LogAnimMarkerSync, Warning,
 					TEXT("[%p] Leader CurrentTime(%0.2f, %0.2f%%) MoveDelta(%0.2f) Looping(%d) [(Prev '%s' @ %d | Next '%s' @ %d | %0.2f) -> (Prev '%s' @ %d | Next '%s' @ %d | %0.2f)]"),
 					GroupLeader.TraitPtr.GetNodeInstance(),
-					GroupContext.LeaderProgress.GetPosition(),
-					GroupContext.LeaderProgress.GetPositionRatio(),
+					GroupContext.LeaderTimelineState.GetPosition(),
+					GroupContext.LeaderTimelineState.GetPositionRatio(),
 					DeltaTime,
 					LeaderTimelineState.IsLooping() ? 1 : 0,
 					*GroupContext.SyncPositionStart.PreviousMarkerName.ToString(),
@@ -905,16 +915,16 @@ namespace UE::AnimNext
 
 				GroupSyncTrait.AdvanceBy(Context, DeltaTime);
 
-				GroupContext.LeaderTimelineState = TimelineTrait.GetState(Context);
+				GroupContext.LeaderTimelineState = GroupLeaderContext.TimelineTrait.GetState(Context);
 				GroupContext.LeaderDeltaTime = DeltaTime;
 
 #if UE_LOG_ENABLE_MARKER_DEBUG
 				UE_LOG(LogAnimMarkerSync, Warning,
-					TEXT("[%p] Leader CurrentTime(%0.2f, %0.2f%%) MoveDelta(%0.2f)"),
+					TEXT("[%p] Leader [Time] [%0.2f + %0.2f -> %0.2f / %0.2f (%0.2f%%)]"),
 					GroupLeader.TraitPtr.GetNodeInstance(),
-					GroupContext.LeaderProgress.GetPosition(),
-					GroupContext.LeaderProgress.GetPositionRatio(),
-					DeltaTime);
+					GroupLeaderContext.TimelineState.GetPosition(), DeltaTime,
+					GroupContext.LeaderTimelineState.GetPosition(), GroupContext.LeaderTimelineState.GetDuration(),
+					GroupContext.LeaderTimelineState.GetPositionRatio());
 #endif
 			}
 		}
@@ -950,17 +960,32 @@ namespace UE::AnimNext
 				ensure(Context.GetStack(GroupMember.TraitPtr, MemberTraitStack));
 				ensure(MemberTraitStack.GetInterface(GroupSyncTrait));
 
+				// Cache our starting timeline state
+				GroupMemberContext.TimelineState = GroupMemberContext.TimelineTrait.GetState(Context);
+
 				float DeltaTime = 0.0f;
 
+#if UE_LOG_ENABLE_MARKER_DEBUG
+				const TCHAR* SyncMode = TEXT("????");
+#endif
+
 				if (GroupMember.GroupParameters.GroupRole == EAnimGroupSynchronizationRole::ExclusiveAlwaysLeader)
 				{
 					// We asked to be a leader and we weren't picked as leader
 					// We don't want to follow anyone so we tick on our own
 					DeltaTime = GroupMember.TraitState.GetDeltaTime();
+
+#if UE_LOG_ENABLE_MARKER_DEBUG
+					SyncMode = TEXT("Excl");
+#endif
 				}
 				else if (GroupMemberContext.bUseMarkerSyncing)
 				{
 					DeltaTime = MatchLeaderMarkerPhase(GroupContext, GroupLeaderContext, GroupMemberContext);
+
+#if UE_LOG_ENABLE_MARKER_DEBUG
+					SyncMode = TEXT("Mark");
+#endif
 				}
 				else
 				{
@@ -987,14 +1012,22 @@ namespace UE::AnimNext
 					DeltaTime = CurrentPosition - PreviousPosition;
 
 #if UE_LOG_ENABLE_MARKER_DEBUG
-					UE_LOG(LogAnimMarkerSync, Warning,
-						TEXT("[%p] Follower Ratio(%0.2f)"),
-						GroupMember.TraitPtr.GetNodeInstance(),
-						LeaderProgressRatio);
+					SyncMode = TEXT("Time");
 #endif
 				}
 
 				GroupSyncTrait.AdvanceBy(Context, DeltaTime);
+
+#if UE_LOG_ENABLE_MARKER_DEBUG
+				const FTimelineState CurrTimelineState = GroupMemberContext.TimelineTrait.GetState(Context);
+
+				UE_LOG(LogAnimMarkerSync, Warning,
+					TEXT("[%p] Follower [%s] [%0.2f + %0.2f -> %0.2f / %0.2f (%0.2f%%)]"),
+					GroupMember.TraitPtr.GetNodeInstance(), SyncMode,
+					GroupMemberContext.TimelineState.GetPosition(), DeltaTime,
+					CurrTimelineState.GetPosition(), CurrTimelineState.GetDuration(),
+					CurrTimelineState.GetPositionRatio());
+#endif
 			}
 		}
 	}
@@ -1074,7 +1107,7 @@ namespace UE::AnimNext
 
 			// If the leader has sync markers, collect and filter them
 			// We'll retain only markers common to all members
-			CollectValidMarkers(Context, GroupContext);
+			BuildGroupState(Context, GroupContext);
 
 			// Advance the leader position as it determines the position of its followers
 			AdvanceLeader(Context, GroupContext);
