ApplyAdditiveTrait AnimNext

#tests PIE

diff --git a/Engine/Plugins/Experimental/AnimNext/Source/AnimNext/Private/DecompressionTools.cpp b/Engine/Plugins/Experimental/AnimNext/Source/AnimNext/Private/DecompressionTools.cpp
--- a/Engine/Plugins/Experimental/AnimNext/Source/AnimNext/Private/DecompressionTools.cpp
+++ b/Engine/Plugins/Experimental/AnimNext/Source/AnimNext/Private/DecompressionTools.cpp
@@ -235,10 +235,16 @@ void FDecompressionTools::GetBonePose(const UAnimSequence* AnimSequence, const F
 
 void FDecompressionTools::GetBonePose_Additive(const UAnimSequence* AnimSequence, const FAnimExtractContext& ExtractionContext, FLODPose& OutAnimationPoseData)
 {
+	// @TODO: Implement now that we have additive, only used for 'bForceUseRawData' for now.
+	// Below exists to prevent reading garbage / warnings about non-normalized rotations
+	GetBonePose(AnimSequence, ExtractionContext, OutAnimationPoseData);
 }
 
 void FDecompressionTools::GetBonePose_AdditiveMeshRotationOnly(const UAnimSequence* AnimSequence, const FAnimExtractContext& ExtractionContext, FLODPose& OutAnimationPoseData)
 {
+	// @TODO: Implement now that we have additive, only used for 'bForceUseRawData' for now.
+	// Below exists to prevent reading garbage / warnings about non-normalized rotations
+	GetBonePose(AnimSequence, ExtractionContext, OutAnimationPoseData);
 }
 
 void FDecompressionTools::GetAnimationCurves(const UAnimSequence* AnimSequence, const FAnimExtractContext& ExtractionContext, FBlendedCurve& OutCurves, bool bForceUseRawData)
@@ -347,6 +353,11 @@ void FDecompressionTools::DecompressPose(FLODPose& OutAnimationPoseData,
 	if (bIsMeshSpaceAdditive)
 	{
 		AnimatedCompactRotations.Init(false, NumLODBoneIndexes);
+		OutAnimationPoseData.Flags |= ELODPoseFlags::MeshSpaceAdditive;
+	}
+	else if (DecompressionContext.GetAdditiveType() == AAT_LocalSpaceBase)
+	{
+		OutAnimationPoseData.Flags |= ELODPoseFlags::LocalSpaceAdditive;
 	}
 
 	// Optimization: assuming first index is root bone. That should always be the case in Skeletons.
diff --git a/Engine/Plugins/Experimental/AnimNext/Source/AnimNext/Private/TransformArrayOperations_AoS.cpp b/Engine/Plugins/Experimental/AnimNext/Source/AnimNext/Private/TransformArrayOperations_AoS.cpp
--- a/Engine/Plugins/Experimental/AnimNext/Source/AnimNext/Private/TransformArrayOperations_AoS.cpp
+++ b/Engine/Plugins/Experimental/AnimNext/Source/AnimNext/Private/TransformArrayOperations_AoS.cpp
@@ -51,6 +51,32 @@ namespace UE::AnimNext
 		}
 	}
 
+	void ConvertPoseLocalToMesh(const FTransformArrayAoSView& Input, const TArrayView<const FBoneIndexType>& LODBoneIndexToParentLODBoneIndexMap)
+	{
+		// @TODO: This method is untested as AoS doesn't seem to compile? 
+		const int32 NumLODBoneIndexes = LODBoneIndexToParentLODBoneIndexMap.Num();
+
+		for (int32 LODBoneIndex = 1; LODBoneIndex < NumLODBoneIndexes; ++LODBoneIndex)
+		{
+			const FBoneIndexType ParentLODBoneIndex = LODBoneIndexToParentLODBoneIndexMap[LODBoneIndex];
+			const FQuat MeshSpaceRotation = Input[ParentLODBoneIndex].GetRotation() * Input[LODBoneIndex].GetRotation();
+			Input[LODBoneIndex].SetRotation(MeshSpaceRotation);
+		}
+	}
+
+	void ConvertPoseMeshToLocal(const FTransformArrayAoSView& Input, const TArrayView<const FBoneIndexType>& LODBoneIndexToParentLODBoneIndexMap)
+	{
+		// @TODO: This method is untested as AoS doesn't seem to compile? 
+		const int32 NumLODBoneIndexes = LODBoneIndexToParentLODBoneIndexMap.Num();
+
+		for (int32 LODBoneIndex = NumLODBoneIndexes - 1; LODBoneIndex > 0; --LODBoneIndex)
+		{
+			const FBoneIndexType ParentLODBoneIndex = LODBoneIndexToParentLODBoneIndexMap[LODBoneIndex];
+			const FQuat LocalSpaceRotation = Input[ParentLODBoneIndex].GetRotation().Inverse() * Input[LODBoneIndex].GetRotation();
+			Input[LODBoneIndex].SetRotation(LocalSpaceRotation);
+		}
+	}
+
 	void BlendWithIdentityAndAccumulate(const FTransformArrayAoSView& Base, const FTransformArrayAoSConstView& Additive, const float BlendWeight)
 	{
 		const ScalarRegister VBlendWeight(BlendWeight);
@@ -64,6 +90,18 @@ namespace UE::AnimNext
 		}
 	}
 
+	void BlendWithIdentityAndAccumulateMesh(
+		const FTransformArrayAoSView& Base, const FTransformArrayAoSConstView& Additive, 
+		const TArrayView<const FBoneIndexType>& LODBoneIndexToParentLODBoneIndexMap, const float BlendWeight)
+	{
+		// @TODO: This method is untested as AoS doesn't seem to compile? 
+		ConvertPoseLocalToMesh(Base, LODBoneIndexToParentLODBoneIndexMap);
+
+		BlendWithIdentityAndAccumulate(Base, Additive, BlendWeight);
+
+		ConvertPoseMeshToLocal(Base, LODBoneIndexToParentLODBoneIndexMap);
+	}
+
 	void BlendOverwriteWithScale(const FTransformArrayAoSView& Dest, const FTransformArrayAoSConstView& Source, const float ScaleWeight)
 	{
 		const ScalarRegister VScaleWeight(ScaleWeight);
diff --git a/Engine/Plugins/Experimental/AnimNext/Source/AnimNext/Private/TransformArrayOperations_SoA.cpp b/Engine/Plugins/Experimental/AnimNext/Source/AnimNext/Private/TransformArrayOperations_SoA.cpp
--- a/Engine/Plugins/Experimental/AnimNext/Source/AnimNext/Private/TransformArrayOperations_SoA.cpp
+++ b/Engine/Plugins/Experimental/AnimNext/Source/AnimNext/Private/TransformArrayOperations_SoA.cpp
@@ -100,6 +100,30 @@ namespace UE::AnimNext
 		}
 	}
 
+	void ConvertPoseLocalToMesh(const FTransformArraySoAView& Input, const TArrayView<const FBoneIndexType>& LODBoneIndexToParentLODBoneIndexMap)
+	{
+		const int32 NumLODBoneIndexes = LODBoneIndexToParentLODBoneIndexMap.Num();
+
+		for (int32 LODBoneIndex = 1; LODBoneIndex < NumLODBoneIndexes; ++LODBoneIndex)
+		{
+			const FBoneIndexType ParentLODBoneIndex = LODBoneIndexToParentLODBoneIndexMap[LODBoneIndex];
+			const FQuat MeshSpaceRotation = Input[ParentLODBoneIndex].GetRotation() * Input[LODBoneIndex].GetRotation();
+			Input.Rotations[LODBoneIndex] = MeshSpaceRotation;
+		}
+	}
+
+	void ConvertPoseMeshToLocal(const FTransformArraySoAView& Input, const TArrayView<const FBoneIndexType>& LODBoneIndexToParentLODBoneIndexMap)
+	{
+		const int32 NumLODBoneIndexes = LODBoneIndexToParentLODBoneIndexMap.Num();
+
+		for (int32 LODBoneIndex = NumLODBoneIndexes - 1; LODBoneIndex > 0; --LODBoneIndex)
+		{
+			const FBoneIndexType ParentLODBoneIndex = LODBoneIndexToParentLODBoneIndexMap[LODBoneIndex];
+			const FQuat LocalSpaceRotation = Input[ParentLODBoneIndex].GetRotation().Inverse() * Input[LODBoneIndex].GetRotation();
+			Input.Rotations[LODBoneIndex] = LocalSpaceRotation;
+		}
+	}
+
 	void BlendWithIdentityAndAccumulate(const FTransformArraySoAView& Base, const FTransformArraySoAConstView& Additive, const float BlendWeight)
 	{
 		const int32 NumTransforms = Base.Num();
@@ -150,6 +174,17 @@ namespace UE::AnimNext
 		}
 	}
 
+	void BlendWithIdentityAndAccumulateMesh(
+		const FTransformArraySoAView& Base, const FTransformArraySoAConstView& Additive, 
+		const TArrayView<const FBoneIndexType>& LODBoneIndexToParentLODBoneIndexMap, const float BlendWeight)
+	{
+		ConvertPoseLocalToMesh(Base, LODBoneIndexToParentLODBoneIndexMap);
+
+		BlendWithIdentityAndAccumulate(Base, Additive, BlendWeight);
+
+		ConvertPoseMeshToLocal(Base, LODBoneIndexToParentLODBoneIndexMap);
+	}
+
 	void BlendOverwriteWithScale(const FTransformArraySoAView& Dest, const FTransformArraySoAConstView& Source, const float ScaleWeight)
 	{
 		SCOPE_CYCLE_COUNTER(STAT_AnimNext_BlendOverwrite_SoA);
diff --git a/Engine/Plugins/Experimental/AnimNext/Source/AnimNext/Public/LODPose.h b/Engine/Plugins/Experimental/AnimNext/Source/AnimNext/Public/LODPose.h
--- a/Engine/Plugins/Experimental/AnimNext/Source/AnimNext/Public/LODPose.h
+++ b/Engine/Plugins/Experimental/AnimNext/Source/AnimNext/Public/LODPose.h
@@ -18,7 +18,9 @@ enum class ELODPoseFlags : uint8
 	Additive			= 1 << 0,
 	DisableRetargeting	= 1 << 1,
 	UseRawData			= 1 << 2,
-	UseSourceData		= 1 << 3
+	UseSourceData		= 1 << 3,
+	MeshSpaceAdditive	= 1 << 4,
+	LocalSpaceAdditive	= 1 << 5,
 };
 
 ENUM_CLASS_FLAGS(ELODPoseFlags);
@@ -179,6 +181,16 @@ struct FLODPose
 		return (Flags & ELODPoseFlags::Additive) != ELODPoseFlags::None;
 	}
 
+	bool IsMeshSpaceAdditive() const
+	{
+		return (Flags & ELODPoseFlags::MeshSpaceAdditive) != ELODPoseFlags::None;
+	}
+
+	bool IsLocalSpaceAdditive() const
+	{
+		return (Flags & ELODPoseFlags::LocalSpaceAdditive) != ELODPoseFlags::None;
+	}
+
 	/** Disable Retargeting */
 	void SetDisableRetargeting(bool bDisableRetargeting)
 	{
diff --git a/Engine/Plugins/Experimental/AnimNext/Source/AnimNext/Public/TransformArrayOperations.h b/Engine/Plugins/Experimental/AnimNext/Source/AnimNext/Public/TransformArrayOperations.h
--- a/Engine/Plugins/Experimental/AnimNext/Source/AnimNext/Public/TransformArrayOperations.h
+++ b/Engine/Plugins/Experimental/AnimNext/Source/AnimNext/Public/TransformArrayOperations.h
@@ -4,6 +4,7 @@
 
 #include "CoreMinimal.h"
 #include "TransformArrayView.h"
+#include "BoneIndices.h"
 
 namespace UE::AnimNext
 {
@@ -22,6 +23,14 @@ namespace UE::AnimNext
 	ANIMNEXT_API void NormalizeRotations(const FTransformArrayAoSView& Input);
 	ANIMNEXT_API void NormalizeRotations(const FTransformArraySoAView& Input);
 
+	// Convert a pose from local space to mesh space
+	ANIMNEXT_API void ConvertPoseLocalToMesh(const FTransformArrayAoSView& Input, const TArrayView<const FBoneIndexType>& LODBoneIndexToParentLODBoneIndexMap);
+	ANIMNEXT_API void ConvertPoseLocalToMesh(const FTransformArraySoAView& Input, const TArrayView<const FBoneIndexType>& LODBoneIndexToParentLODBoneIndexMap);
+
+	// Convert a pose from mesh space to local space
+	ANIMNEXT_API void ConvertPoseMeshToLocal(const FTransformArrayAoSView& Input, const TArrayView<const FBoneIndexType>& LODBoneIndexToParentLODBoneIndexMap);
+	ANIMNEXT_API void ConvertPoseMeshToLocal(const FTransformArraySoAView& Input, const TArrayView<const FBoneIndexType>& LODBoneIndexToParentLODBoneIndexMap);
+
 	// The additive transform view is blended with the additive identity using the provided blend weight
 	// We then accumulate the resulting transforms on top of the base transforms
 	// Delta = Blend(Identity, Additive, BlendWeight)
@@ -29,6 +38,14 @@ namespace UE::AnimNext
 	ANIMNEXT_API void BlendWithIdentityAndAccumulate(const FTransformArrayAoSView& Base, const FTransformArrayAoSConstView& Additive, const float BlendWeight);
 	ANIMNEXT_API void BlendWithIdentityAndAccumulate(const FTransformArraySoAView& Base, const FTransformArraySoAConstView& Additive, const float BlendWeight);
 
+	// Performs a BlendWithIdentityAndAccumulate in mesh space, result is in local space
+	ANIMNEXT_API void BlendWithIdentityAndAccumulateMesh(
+		const FTransformArrayAoSView& Base, const FTransformArrayAoSConstView& Additive, 
+		const TArrayView<const FBoneIndexType>& LODBoneIndexToParentLODBoneIndexMap, const float BlendWeight);
+	ANIMNEXT_API void BlendWithIdentityAndAccumulateMesh(
+		const FTransformArraySoAView& Base, const FTransformArraySoAConstView& Additive, 
+		const TArrayView<const FBoneIndexType>& LODBoneIndexToParentLODBoneIndexMap, const float BlendWeight);
+
 	// The source transforms are scaled by the provided weight and the result is written in the destination
 	// Dest = Source * ScaleWeight
 	ANIMNEXT_API void BlendOverwriteWithScale(const FTransformArrayAoSView& Dest, const FTransformArrayAoSConstView& Source, const float ScaleWeight);
diff --git a/Engine/Plugins/Experimental/AnimNextAnimGraph/Source/AnimNextAnimGraph/Private/EvaluationVM/Tasks/ApplyAdditiveKeyframe.cpp b/Engine/Plugins/Experimental/AnimNextAnimGraph/Source/AnimNextAnimGraph/Private/EvaluationVM/Tasks/ApplyAdditiveKeyframe.cpp
--- a/Engine/Plugins/Experimental/AnimNextAnimGraph/Source/AnimNextAnimGraph/Private/EvaluationVM/Tasks/ApplyAdditiveKeyframe.cpp
+++ b/Engine/Plugins/Experimental/AnimNextAnimGraph/Source/AnimNextAnimGraph/Private/EvaluationVM/Tasks/ApplyAdditiveKeyframe.cpp
@@ -36,20 +36,26 @@ void FAnimNextApplyAdditiveKeyframeTask::Execute(UE::AnimNext::FEvaluationVM& VM
 		return;
 	}
 
+	if (!AdditiveKeyframe->Pose.IsAdditive())
+	{
+		// Additive must be addtive type, push reference pose if not the case
+		VM.PushValue(KEYFRAME_STACK_NAME, MakeUnique<FKeyframeState>(VM.MakeReferenceKeyframe(false)));
+		return;
+	}
+
 	if (EnumHasAnyFlags(VM.GetFlags(), EEvaluationFlags::Bones))
 	{
 		check(BaseKeyframe->Pose.GetNumBones() == AdditiveKeyframe->Pose.GetNumBones());
 
 		const FTransformArrayView BaseTransformsView = BaseKeyframe->Pose.LocalTransforms.GetView();
 
-#if 0
-		// TODO: Read the additive type from the pose itself
-		if (AdditiveType == AAT_RotationOffsetMeshSpace)
+		if (AdditiveKeyframe->Pose.IsMeshSpaceAdditive())
 		{
-			AccumulateMeshSpaceRotationAdditiveToLocalPoseInternal(BaseAnimationPoseData.GetPose(), AdditiveAnimationPoseData.GetPose(), Weight);
+			BlendWithIdentityAndAccumulateMesh(
+				BaseTransformsView, AdditiveKeyframe->Pose.LocalTransforms.GetConstView(), 
+				AdditiveKeyframe->Pose.GetLODBoneIndexToParentLODBoneIndexMap(), BlendWeight);
 		}
 		else
-#endif
 		{
 			BlendWithIdentityAndAccumulate(BaseTransformsView, AdditiveKeyframe->Pose.LocalTransforms.GetConstView(), BlendWeight);
 		}
diff --git a/Engine/Plugins/Experimental/AnimNextAnimGraph/Source/AnimNextAnimGraph/Private/Traits/ApplyAdditiveTrait.cpp b/Engine/Plugins/Experimental/AnimNextAnimGraph/Source/AnimNextAnimGraph/Private/Traits/ApplyAdditiveTrait.cpp
new file mode 100644
--- /dev/null
+++ b/Engine/Plugins/Experimental/AnimNextAnimGraph/Source/AnimNextAnimGraph/Private/Traits/ApplyAdditiveTrait.cpp
@@ -0,0 +1,142 @@
+// Copyright Epic Games, Inc. All Rights Reserved.
+
+#include "Traits/ApplyAdditiveTrait.h"
+
+#include "Animation/AnimTypes.h"
+#include "TraitCore/ExecutionContext.h"
+#include "EvaluationVM/Tasks/ApplyAdditiveKeyframe.h"
+
+namespace UE::AnimNext
+{
+	AUTO_REGISTER_ANIM_TRAIT(FApplyAdditiveTrait)
+
+	// Trait implementation boilerplate
+	#define TRAIT_INTERFACE_ENUMERATOR(GeneratorMacro) \
+		GeneratorMacro(IEvaluate) \
+		GeneratorMacro(IHierarchy) \
+		GeneratorMacro(IUpdate) \
+		GeneratorMacro(IUpdateTraversal) \
+		GeneratorMacro(IContinuousBlend) \
+
+	GENERATE_ANIM_TRAIT_IMPLEMENTATION(FApplyAdditiveTrait, TRAIT_INTERFACE_ENUMERATOR, NULL_ANIM_TRAIT_INTERFACE_ENUMERATOR, NULL_ANIM_TRAIT_EVENT_ENUMERATOR)
+	#undef TRAIT_INTERFACE_ENUMERATOR
+
+	void FApplyAdditiveTrait::FInstanceData::Construct(const FExecutionContext& Context, const FTraitBinding& Binding)
+	{
+		FTrait::FInstanceData::Construct(Context, Binding);
+
+		const FSharedData* SharedData = Binding.GetSharedData<FSharedData>();
+
+		check(!Base.IsValid());
+		Base = Context.AllocateNodeInstance(Binding, SharedData->Base);
+	}
+
+	void FApplyAdditiveTrait::PostEvaluate(FEvaluateTraversalContext& Context, const TTraitBinding<IEvaluate>& Binding) const
+	{
+		const FSharedData* SharedData = Binding.GetSharedData<FSharedData>();
+		check(SharedData);
+
+		const FInstanceData* InstanceData = Binding.GetInstanceData<FInstanceData>();
+		check(InstanceData);
+
+		if (InstanceData->Additive.IsValid())
+		{
+			TTraitBinding<IContinuousBlend> ContinuousBlendTrait;
+			Binding.GetStackInterface(ContinuousBlendTrait);
+
+			const float Alpha = ContinuousBlendTrait.GetBlendWeight(Context, 1);
+			Context.AppendTask(FAnimNextApplyAdditiveKeyframeTask::Make(Alpha));
+		}
+		else
+		{
+			// Additive is not is active, do nothing
+		}
+	}
+
+	void FApplyAdditiveTrait::PreUpdate(FUpdateTraversalContext& Context, const TTraitBinding<IUpdate>& Binding, const FTraitUpdateState& TraitState) const
+	{
+		const FSharedData* SharedData = Binding.GetSharedData<FSharedData>();
+		check(SharedData);
+
+		FInstanceData* InstanceData = Binding.GetInstanceData<FInstanceData>();
+		check(InstanceData);
+
+		TTraitBinding<IContinuousBlend> ContinuousBlendTrait;
+		Binding.GetStackInterface(ContinuousBlendTrait);
+
+		const float Alpha = ContinuousBlendTrait.GetBlendWeight(Context, 1);
+		if (FAnimWeight::IsRelevant(Alpha))
+		{
+			if (!InstanceData->Additive.IsValid())
+			{
+				// We need an additive child that isn't instanced yet, allocate it
+				InstanceData->Additive = Context.AllocateNodeInstance(Binding, SharedData->Additive);
+			}
+			else
+			{
+				InstanceData->bWasAdditiveRelevant = true;
+			}
+		}
+		else
+		{
+			// We no longer need this child, release it
+			InstanceData->Additive.Reset();
+			InstanceData->bWasAdditiveRelevant = false;
+		}
+	}
+
+	void FApplyAdditiveTrait::QueueChildrenForTraversal(FUpdateTraversalContext& Context, const TTraitBinding<IUpdateTraversal>& Binding, const FTraitUpdateState& TraitState, FUpdateTraversalQueue& TraversalQueue) const
+	{
+		const FInstanceData* InstanceData = Binding.GetInstanceData<FInstanceData>();
+		check(InstanceData);
+
+		// Note: No need to manage newly relevant for incoming TraitState, will be set as need by caller. See IUpdate.cpp.
+		TraversalQueue.Push(InstanceData->Base, TraitState);
+
+		if (InstanceData->Additive.IsValid())
+		{
+			TTraitBinding<IContinuousBlend> ContinuousBlendTrait;
+			Binding.GetStackInterface(ContinuousBlendTrait);
+
+			const float Alpha = ContinuousBlendTrait.GetBlendWeight(Context, 1);
+			TraversalQueue.Push(InstanceData->Additive, TraitState.WithWeight(Alpha).AsNewlyRelevant(!InstanceData->bWasAdditiveRelevant));
+		}
+	}
+
+	uint32 FApplyAdditiveTrait::GetNumChildren(const FExecutionContext& Context, const TTraitBinding<IHierarchy>& Binding) const
+	{
+		return 2;
+	}
+
+	void FApplyAdditiveTrait::GetChildren(const FExecutionContext& Context, const TTraitBinding<IHierarchy>& Binding, FChildrenArray& Children) const
+	{
+		const FInstanceData* InstanceData = Binding.GetInstanceData<FInstanceData>();
+
+		Children.Add(InstanceData->Base);
+		Children.Add(InstanceData->Additive);
+	}
+
+	float FApplyAdditiveTrait::GetBlendWeight(const FExecutionContext& Context, const TTraitBinding<IContinuousBlend>& Binding, int32 ChildIndex) const
+	{
+		const FSharedData* SharedData = Binding.GetSharedData<FSharedData>();
+		check(SharedData);
+
+		const float Alpha = SharedData->GetAlpha(Binding);
+
+		if (ChildIndex == 0)
+		{
+			// Base is always full weight
+			return 1.0f;
+		}
+		else if (ChildIndex == 1)
+		{
+			return Alpha;
+		}
+		else
+		{
+			// Invalid child index
+			return -1.0f;
+		}
+	}
+
+}  // UE::AnimNext
diff --git a/Engine/Plugins/Experimental/AnimNextAnimGraph/Source/AnimNextAnimGraph/Private/Traits/ApplyAdditiveTrait.h b/Engine/Plugins/Experimental/AnimNextAnimGraph/Source/AnimNextAnimGraph/Private/Traits/ApplyAdditiveTrait.h
new file mode 100644
--- /dev/null
+++ b/Engine/Plugins/Experimental/AnimNextAnimGraph/Source/AnimNextAnimGraph/Private/Traits/ApplyAdditiveTrait.h
@@ -0,0 +1,84 @@
+// Copyright Epic Games, Inc. All Rights Reserved.
+
+#pragma once
+
+#include "CoreMinimal.h"
+
+#include "Animation/AnimTypes.h"
+#include "TraitCore/Trait.h"
+#include "TraitInterfaces/IContinuousBlend.h"
+#include "TraitInterfaces/IEvaluate.h"
+#include "TraitInterfaces/IHierarchy.h"
+#include "TraitInterfaces/IUpdate.h"
+
+#include "ApplyAdditiveTrait.generated.h"
+
+USTRUCT(meta = (DisplayName = "Apply Additive"))
+struct FAnimNextApplyAdditiveTraitSharedData : public FAnimNextTraitSharedData
+{
+	GENERATED_BODY()
+
+	/** Base to apply addtitive too. */
+	UPROPERTY()
+	FAnimNextTraitHandle Base;
+
+	/** Addtive to apply. */
+	UPROPERTY()
+	FAnimNextTraitHandle Additive;
+
+	/** How much to apply our addtive, default is 1. */
+	UPROPERTY(EditAnywhere, Category = "Default")
+	float Alpha = 1.0f;
+
+	// Latent pin support boilerplate
+	#define TRAIT_LATENT_PROPERTIES_ENUMERATOR(GeneratorMacro) \
+		GeneratorMacro(Alpha) \
+
+	GENERATE_TRAIT_LATENT_PROPERTIES(FAnimNextApplyAdditiveTraitSharedData, TRAIT_LATENT_PROPERTIES_ENUMERATOR)
+	#undef TRAIT_LATENT_PROPERTIES_ENUMERATOR
+};
+
+namespace UE::AnimNext
+{
+	/**
+	 * A trait that can apply a mesh or local space additive to this trait stack.
+	 * 
+	 * Ex: LookAt's, minor hit reactions, up-down floating, etc.
+	 */
+	struct FApplyAdditiveTrait : FBaseTrait, IEvaluate, IUpdate, IUpdateTraversal, IHierarchy, IContinuousBlend
+	{
+		DECLARE_ANIM_TRAIT(FApplyAdditiveTrait, FBaseTrait)
+
+		using FSharedData = FAnimNextApplyAdditiveTraitSharedData;
+
+		struct FInstanceData : FTrait::FInstanceData
+		{
+			/** Reference to base for additive we are applying */
+			FTraitPtr Base;
+
+			/** Reference to addtive we are applying */
+			FTraitPtr Additive;
+
+			/** True if Additive branch has any contribution (Alpha non-zero)  */
+			bool bWasAdditiveRelevant = false;
+
+			void Construct(const FExecutionContext& Context, const FTraitBinding& Binding);
+		};
+
+		// IEvaluate impl
+		virtual void PostEvaluate(FEvaluateTraversalContext& Context, const TTraitBinding<IEvaluate>& Binding) const override;
+
+		// IUpdate impl
+		virtual void PreUpdate(FUpdateTraversalContext& Context, const TTraitBinding<IUpdate>& Binding, const FTraitUpdateState& TraitState) const override;
+
+		// IUpdateTraversal impl
+		virtual void QueueChildrenForTraversal(FUpdateTraversalContext& Context, const TTraitBinding<IUpdateTraversal>& Binding, const FTraitUpdateState& TraitState, FUpdateTraversalQueue& TraversalQueue) const override;
+
+		// IHierarchy impl
+		virtual uint32 GetNumChildren(const FExecutionContext& Context, const TTraitBinding<IHierarchy>& Binding) const override;
+		virtual void GetChildren(const FExecutionContext& Context, const TTraitBinding<IHierarchy>& Binding, FChildrenArray& Children) const override;
+
+		// IContinuousBlend impl
+		virtual float GetBlendWeight(const FExecutionContext& Context, const TTraitBinding<IContinuousBlend>& Binding, int32 ChildIndex) const override;
+	};
+}
