Overhaul SmoothBlendPerBone trait to support hierarchy table blend profiles
* Update BlendStack trait to interface with ISmoothBlendPerBone and added optional blend profile property to FGraphRequest
* Added two additive traits for providing blend profile data

diff --git a/Engine/Plugins/Experimental/AnimNextAnimGraph/Source/AnimNextAnimGraph/Internal/Traits/BlendStackTrait.h b/Engine/Plugins/Experimental/AnimNextAnimGraph/Source/AnimNextAnimGraph/Internal/Traits/BlendStackTrait.h
--- a/Engine/Plugins/Experimental/AnimNextAnimGraph/Source/AnimNextAnimGraph/Internal/Traits/BlendStackTrait.h
+++ b/Engine/Plugins/Experimental/AnimNextAnimGraph/Source/AnimNextAnimGraph/Internal/Traits/BlendStackTrait.h
@@ -14,6 +14,7 @@
 #include "TraitInterfaces/IInertializerBlend.h"
 #include "TraitInterfaces/ITimeline.h"
 #include "TraitInterfaces/ISmoothBlend.h"
+#include "TraitInterfaces/ISmoothBlendPerBone.h"
 
 #include "BlendStackTrait.generated.h"
 
@@ -86,7 +87,7 @@ namespace UE::AnimNext
 	 * 
 	 * Used as base trait in state machine-like setups where state machine is the additive trait that pushes graphs and blend settings to the blend stack.
 	 */
-	struct FBlendStackCoreTrait : FBaseTrait, IUpdate, IUpdateTraversal, IHierarchy, IDiscreteBlend, ISmoothBlend, IInertializerBlend, IAttributeProvider, ITimeline, IGarbageCollection, IBlendStack
+	struct FBlendStackCoreTrait : FBaseTrait, IUpdate, IUpdateTraversal, IHierarchy, IDiscreteBlend, ISmoothBlend, IInertializerBlend, IAttributeProvider, ITimeline, IGarbageCollection, IBlendStack, ISmoothBlendPerBone
 	{
 		DECLARE_ANIM_TRAIT(FBlendStackCoreTrait, FBaseTrait)
 
@@ -174,6 +175,9 @@ namespace UE::AnimNext
 		virtual int32 GetActiveGraph(FExecutionContext& Context, const TTraitBinding<IBlendStack>& Binding, IBlendStack::FGraphRequestPtr& OutGraphRequest) const override;
 		virtual IBlendStack::FGraphRequestPtr GetGraph(FExecutionContext& Context, const TTraitBinding<IBlendStack>& Binding, int32 InChildIndex) const override;
 
+		// ISmoothBlendPerBone impl
+		virtual TSharedPtr<const IBlendProfileInterface> GetBlendProfile(FExecutionContext& Context, const TTraitBinding<ISmoothBlendPerBone>& Binding, int32 ChildIndex) const;
+
 #if WITH_EDITOR
 		virtual bool IsHidden() const override { return true; }
 #endif
diff --git a/Engine/Plugins/Experimental/AnimNextAnimGraph/Source/AnimNextAnimGraph/Private/EvaluationVM/Tasks/BlendKeyframesPerBone.cpp b/Engine/Plugins/Experimental/AnimNextAnimGraph/Source/AnimNextAnimGraph/Private/EvaluationVM/Tasks/BlendKeyframesPerBone.cpp
--- a/Engine/Plugins/Experimental/AnimNextAnimGraph/Source/AnimNextAnimGraph/Private/EvaluationVM/Tasks/BlendKeyframesPerBone.cpp
+++ b/Engine/Plugins/Experimental/AnimNextAnimGraph/Source/AnimNextAnimGraph/Private/EvaluationVM/Tasks/BlendKeyframesPerBone.cpp
@@ -12,22 +12,11 @@
 #include "EvaluationVM/KeyframeState.h"
 #include "TransformArrayOperations.h"
 
-FAnimNextBlendOverwriteKeyframePerBoneWithScaleTask FAnimNextBlendOverwriteKeyframePerBoneWithScaleTask::Make(TWeakObjectPtr<const UBlendProfile> BlendProfile, const FBlendSampleData& BlendData, float ScaleFactor)
-{
-	FAnimNextBlendOverwriteKeyframePerBoneWithScaleTask Task;
-	Task.BlendProfile = BlendProfile.Get();
-	Task.SourceSkeleton = BlendProfile->OwningSkeleton;
-	Task.BlendData = &BlendData;
-	Task.ScaleFactor = ScaleFactor;
-
-	return Task;
-}
-
-FAnimNextBlendOverwriteKeyframePerBoneWithScaleTask FAnimNextBlendOverwriteKeyframePerBoneWithScaleTask::Make(const IInterpolationIndexProvider* BlendProfile, const USkeleton* Skeleton, const FBlendSampleData& BlendData, float ScaleFactor)
+FAnimNextBlendOverwriteKeyframePerBoneWithScaleTask FAnimNextBlendOverwriteKeyframePerBoneWithScaleTask::Make(const IBlendProfileInterface* BlendProfile, const FBlendSampleData& BlendData, float ScaleFactor)
 {
 	FAnimNextBlendOverwriteKeyframePerBoneWithScaleTask Task;
 	Task.BlendProfile = BlendProfile;
-	Task.SourceSkeleton = Skeleton;
+	Task.SourceSkeleton = BlendProfile->GetSkeleton();
 	Task.BlendData = &BlendData;
 	Task.ScaleFactor = ScaleFactor;
 
@@ -38,8 +27,7 @@ void FAnimNextBlendOverwriteKeyframePerBoneWithScaleTask::Execute(UE::AnimNext::
 {
 	using namespace UE::AnimNext;
 
-	const IInterpolationIndexProvider* BlendProfilePtr = BlendProfile;
-	if (BlendProfilePtr == nullptr)
+	if (BlendProfile == nullptr)
 	{
 		// If we don't have a blend profile, we blend the whole pose
 		Super::Execute(VM);
@@ -54,9 +42,9 @@ void FAnimNextBlendOverwriteKeyframePerBoneWithScaleTask::Execute(UE::AnimNext::
 	}
 
 	const USkeleton* TargetSkeleton = Keyframe->Pose.GetSkeletonAsset();
-	TSharedPtr<IInterpolationIndexProvider::FPerBoneInterpolationData> Data = BlendProfilePtr->GetPerBoneInterpolationData(TargetSkeleton);
+	TSharedPtr<IInterpolationIndexProvider::FPerBoneInterpolationData> Data = BlendProfile->GetPerBoneInterpolationData(TargetSkeleton);
 
-	const FSkeletonRemapping& SkeletonRemapping = UE::Anim::FSkeletonRemappingRegistry::Get().GetRemapping(SourceSkeleton, TargetSkeleton);
+	const FSkeletonRemapping& SkeletonRemapping = UE::Anim::FSkeletonRemappingRegistry::Get().GetRemapping(BlendProfile->GetSkeleton(), TargetSkeleton);
 
 	const TArrayView<const FBoneIndexType> LODBoneIndexToSkeletonBoneIndexMap = Keyframe->Pose.GetLODBoneIndexToSkeletonBoneIndexMap();
 	const int32 NumLODBones = LODBoneIndexToSkeletonBoneIndexMap.Num();
@@ -70,7 +58,7 @@ void FAnimNextBlendOverwriteKeyframePerBoneWithScaleTask::Execute(UE::AnimNext::
 		{
 			const FSkeletonPoseBoneIndex TargetSkeletonBoneIndex(LODBoneIndexToSkeletonBoneIndexMap[LODBoneIndex]);
 			const FSkeletonPoseBoneIndex SourceSkeletonBoneIndex(SkeletonRemapping.GetSourceSkeletonBoneIndex(TargetSkeletonBoneIndex.GetInt()));
-			LODBoneIndexToWeightIndexMap[LODBoneIndex] = SourceSkeletonBoneIndex.IsValid() ? BlendProfilePtr->GetPerBoneInterpolationIndex(SourceSkeletonBoneIndex, TargetSkeleton, Data.Get()) : INDEX_NONE;
+			LODBoneIndexToWeightIndexMap[LODBoneIndex] = SourceSkeletonBoneIndex.IsValid() ? BlendProfile->GetPerBoneInterpolationIndex(SourceSkeletonBoneIndex, TargetSkeleton, Data.Get()) : INDEX_NONE;
 		}
 	}
 	else
@@ -78,7 +66,7 @@ void FAnimNextBlendOverwriteKeyframePerBoneWithScaleTask::Execute(UE::AnimNext::
 		for (int32 LODBoneIndex = 0; LODBoneIndex < NumLODBones; ++LODBoneIndex)
 		{
 			const FSkeletonPoseBoneIndex SkeletonBoneIndex(LODBoneIndexToSkeletonBoneIndexMap[LODBoneIndex]);
-			LODBoneIndexToWeightIndexMap[LODBoneIndex] = BlendProfilePtr->GetPerBoneInterpolationIndex(SkeletonBoneIndex, TargetSkeleton, Data.Get());
+			LODBoneIndexToWeightIndexMap[LODBoneIndex] = BlendProfile->GetPerBoneInterpolationIndex(SkeletonBoneIndex, TargetSkeleton, Data.Get());
 		}
 	}
 
@@ -108,7 +96,7 @@ void FAnimNextBlendOverwriteKeyframePerBoneWithScaleTask::Execute(UE::AnimNext::
 	VM.PushValue(KEYFRAME_STACK_NAME, MoveTemp(Keyframe));
 }
 
-FAnimNextBlendAddKeyframePerBoneWithScaleTask FAnimNextBlendAddKeyframePerBoneWithScaleTask::Make(TWeakObjectPtr<const UBlendProfile> BlendProfile, const FBlendSampleData& BlendDataA, const FBlendSampleData& BlendDataB, float ScaleFactor)
+FAnimNextBlendAddKeyframePerBoneWithScaleTask FAnimNextBlendAddKeyframePerBoneWithScaleTask::Make(const IBlendProfileInterface* BlendProfile, const FBlendSampleData& BlendDataA, const FBlendSampleData& BlendDataB, float ScaleFactor)
 {
 	FAnimNextBlendAddKeyframePerBoneWithScaleTask Task;
 	Task.BlendProfile = BlendProfile;
@@ -123,8 +111,7 @@ void FAnimNextBlendAddKeyframePerBoneWithScaleTask::Execute(UE::AnimNext::FEvalu
 {
 	using namespace UE::AnimNext;
 
-	const UBlendProfile* BlendProfilePtr = BlendProfile.Get();
-	if (BlendProfilePtr == nullptr)
+	if (BlendProfile == nullptr)
 	{
 		// If we don't have a blend profile, we blend the whole pose
 		Super::Execute(VM);
@@ -149,10 +136,9 @@ void FAnimNextBlendAddKeyframePerBoneWithScaleTask::Execute(UE::AnimNext::FEvalu
 	}
 
 	const USkeleton* TargetSkeleton = KeyframeB->Pose.GetSkeletonAsset();
-	TSharedPtr<IInterpolationIndexProvider::FPerBoneInterpolationData> Data = BlendProfilePtr->GetPerBoneInterpolationData(TargetSkeleton);
+	TSharedPtr<IInterpolationIndexProvider::FPerBoneInterpolationData> Data = BlendProfile->GetPerBoneInterpolationData(TargetSkeleton);
 
-	const USkeleton* SourceSkeleton = BlendProfilePtr->OwningSkeleton;
-	const FSkeletonRemapping& SkeletonRemapping = UE::Anim::FSkeletonRemappingRegistry::Get().GetRemapping(SourceSkeleton, TargetSkeleton);
+	const FSkeletonRemapping& SkeletonRemapping = UE::Anim::FSkeletonRemappingRegistry::Get().GetRemapping(BlendProfile->GetSkeleton(), TargetSkeleton);
 
 	const TArrayView<const FBoneIndexType> LODBoneIndexToSkeletonBoneIndexMap = KeyframeB->Pose.GetLODBoneIndexToSkeletonBoneIndexMap();
 	const int32 NumLODBones = LODBoneIndexToSkeletonBoneIndexMap.Num();
@@ -166,7 +152,7 @@ void FAnimNextBlendAddKeyframePerBoneWithScaleTask::Execute(UE::AnimNext::FEvalu
 		{
 			const FSkeletonPoseBoneIndex TargetSkeletonBoneIndex(LODBoneIndexToSkeletonBoneIndexMap[LODBoneIndex]);
 			const FSkeletonPoseBoneIndex SourceSkeletonBoneIndex(SkeletonRemapping.GetSourceSkeletonBoneIndex(TargetSkeletonBoneIndex.GetInt()));
-			LODBoneIndexToWeightIndexMap[LODBoneIndex] = SourceSkeletonBoneIndex.IsValid() ? BlendProfilePtr->GetPerBoneInterpolationIndex(SourceSkeletonBoneIndex, TargetSkeleton, Data.Get()) : INDEX_NONE;
+			LODBoneIndexToWeightIndexMap[LODBoneIndex] = SourceSkeletonBoneIndex.IsValid() ? BlendProfile->GetPerBoneInterpolationIndex(SourceSkeletonBoneIndex, TargetSkeleton, Data.Get()) : INDEX_NONE;
 		}
 	}
 	else
@@ -174,7 +160,7 @@ void FAnimNextBlendAddKeyframePerBoneWithScaleTask::Execute(UE::AnimNext::FEvalu
 		for (int32 LODBoneIndex = 0; LODBoneIndex < NumLODBones; ++LODBoneIndex)
 		{
 			const FSkeletonPoseBoneIndex SkeletonBoneIndex(LODBoneIndexToSkeletonBoneIndexMap[LODBoneIndex]);
-			LODBoneIndexToWeightIndexMap[LODBoneIndex] = BlendProfilePtr->GetPerBoneInterpolationIndex(SkeletonBoneIndex, TargetSkeleton, Data.Get());
+			LODBoneIndexToWeightIndexMap[LODBoneIndex] = BlendProfile->GetPerBoneInterpolationIndex(SkeletonBoneIndex, TargetSkeleton, Data.Get());
 		}
 	}
 
diff --git a/Engine/Plugins/Experimental/AnimNextAnimGraph/Source/AnimNextAnimGraph/Private/TraitInterfaces/ISmoothBlendPerBone.cpp b/Engine/Plugins/Experimental/AnimNextAnimGraph/Source/AnimNextAnimGraph/Private/TraitInterfaces/ISmoothBlendPerBone.cpp
new file mode 100644
--- /dev/null
+++ b/Engine/Plugins/Experimental/AnimNextAnimGraph/Source/AnimNextAnimGraph/Private/TraitInterfaces/ISmoothBlendPerBone.cpp
@@ -0,0 +1,34 @@
+// Copyright Epic Games, Inc. All Rights Reserved.
+
+#include "TraitInterfaces/ISmoothBlendPerBone.h"
+
+#include "TraitCore/ExecutionContext.h"
+
+namespace UE::AnimNext
+{
+	AUTO_REGISTER_ANIM_TRAIT_INTERFACE(ISmoothBlendPerBone)
+
+#if WITH_EDITOR
+	const FText& ISmoothBlendPerBone::GetDisplayName() const
+	{
+		static FText InterfaceName = NSLOCTEXT("TraitInterfaces", "TraitInterface_ISmoothBlendPerBone_Name", "Smooth Blend Per Bone");
+		return InterfaceName;
+	}
+	const FText& ISmoothBlendPerBone::GetDisplayShortName() const
+	{
+		static FText InterfaceShortName = NSLOCTEXT("TraitInterfaces", "TraitInterface_ISmoothBlendPerBone_ShortName", "SMBPB");
+		return InterfaceShortName;
+	}
+#endif // WITH_EDITOR
+
+	TSharedPtr<const IBlendProfileInterface> ISmoothBlendPerBone::GetBlendProfile(FExecutionContext& Context, const TTraitBinding<ISmoothBlendPerBone>& Binding, int32 ChildIndex) const
+	{
+		TTraitBinding<ISmoothBlendPerBone> SuperBinding;
+		if (Binding.GetStackInterfaceSuper(SuperBinding))
+		{
+			return SuperBinding.GetBlendProfile(Context, ChildIndex);
+		}
+
+		return nullptr;
+	}
+}
diff --git a/Engine/Plugins/Experimental/AnimNextAnimGraph/Source/AnimNextAnimGraph/Private/Traits/BlendSmootherPerBone.cpp b/Engine/Plugins/Experimental/AnimNextAnimGraph/Source/AnimNextAnimGraph/Private/Traits/BlendSmootherPerBone.cpp
--- a/Engine/Plugins/Experimental/AnimNextAnimGraph/Source/AnimNextAnimGraph/Private/Traits/BlendSmootherPerBone.cpp
+++ b/Engine/Plugins/Experimental/AnimNextAnimGraph/Source/AnimNextAnimGraph/Private/Traits/BlendSmootherPerBone.cpp
@@ -4,6 +4,7 @@
 
 #include "AlphaBlend.h"
 #include "Animation/BlendProfile.h"
+#include "Animation/Skeleton.h"
 #include "TraitCore/ExecutionContext.h"
 #include "TraitInterfaces/IHierarchy.h"
 #include "EvaluationVM/Tasks/BlendKeyframesPerBone.h"
@@ -11,64 +12,126 @@
 
 namespace UE::AnimNext
 {
-	AUTO_REGISTER_ANIM_TRAIT(FBlendSmootherPerBoneTrait)
+	AUTO_REGISTER_ANIM_TRAIT(FBlendSmootherPerBoneCoreTrait)
+	AUTO_REGISTER_ANIM_TRAIT(FBlendProfilePerChildProviderTrait)
+	AUTO_REGISTER_ANIM_TRAIT(FBlendProfileProviderTrait)
 
 	// Trait implementation boilerplate
 	#define TRAIT_INTERFACE_ENUMERATOR(GeneratorMacro) \
 		GeneratorMacro(IDiscreteBlend) \
 		GeneratorMacro(IEvaluate) \
 		GeneratorMacro(IUpdate) \
+	
+	// Trait required interfaces implementation boilerplate
+	#define TRAIT_REQUIRED_INTERFACE_ENUMERATOR(GeneratorMacroRequired) \
+		GeneratorMacroRequired(ISmoothBlend) \
+		GeneratorMacroRequired(ISmoothBlendPerBone) \
 
-	GENERATE_ANIM_TRAIT_IMPLEMENTATION(FBlendSmootherPerBoneTrait, TRAIT_INTERFACE_ENUMERATOR, NULL_ANIM_TRAIT_INTERFACE_ENUMERATOR, NULL_ANIM_TRAIT_EVENT_ENUMERATOR)
+	GENERATE_ANIM_TRAIT_IMPLEMENTATION(FBlendSmootherPerBoneCoreTrait, TRAIT_INTERFACE_ENUMERATOR, TRAIT_REQUIRED_INTERFACE_ENUMERATOR, NULL_ANIM_TRAIT_EVENT_ENUMERATOR)
 	#undef TRAIT_INTERFACE_ENUMERATOR
+	#undef TRAIT_REQUIRED_INTERFACE_ENUMERATOR
 
-	void FBlendSmootherPerBoneTrait::PostEvaluate(FEvaluateTraversalContext& Context, const TTraitBinding<IEvaluate>& Binding) const
+	// Trait implementation boilerplate
+	#define TRAIT_INTERFACE_ENUMERATOR(GeneratorMacro) \
+		GeneratorMacro(ISmoothBlendPerBone) \
+		GeneratorMacro(IUpdate) \
+
+	GENERATE_ANIM_TRAIT_IMPLEMENTATION(FBlendProfilePerChildProviderTrait, TRAIT_INTERFACE_ENUMERATOR, NULL_ANIM_TRAIT_INTERFACE_ENUMERATOR, NULL_ANIM_TRAIT_EVENT_ENUMERATOR)
+	#undef TRAIT_INTERFACE_ENUMERATOR
+	#undef TRAIT_REQUIRED_INTERFACE_ENUMERATOR
+
+		// Trait implementation boilerplate
+	#define TRAIT_INTERFACE_ENUMERATOR(GeneratorMacro) \
+		GeneratorMacro(ISmoothBlendPerBone) \
+		GeneratorMacro(IUpdate) \
+
+		GENERATE_ANIM_TRAIT_IMPLEMENTATION(FBlendProfileProviderTrait, TRAIT_INTERFACE_ENUMERATOR, NULL_ANIM_TRAIT_INTERFACE_ENUMERATOR, NULL_ANIM_TRAIT_EVENT_ENUMERATOR)
+	#undef TRAIT_INTERFACE_ENUMERATOR
+	#undef TRAIT_REQUIRED_INTERFACE_ENUMERATOR
+
+	void FBlendSmootherPerBoneCoreTrait::PostEvaluate(FEvaluateTraversalContext& Context, const TTraitBinding<IEvaluate>& Binding) const
 	{
-		const FSharedData* SharedData = Binding.GetSharedData<FSharedData>();
 		const FInstanceData* InstanceData = Binding.GetInstanceData<FInstanceData>();
 
-		if (!SharedData->BlendProfile)
-		{
-			// No blend profile set, default smooth blend behavior
-			IEvaluate::PostEvaluate(Context, Binding);
-			return;
-		}
+		TTraitBinding<ISmoothBlend> SmoothBlendTrait;
+		Binding.GetStackInterface(SmoothBlendTrait);
 
-		// We override the default behavior since we need to blend per bone
+		TTraitBinding<IDiscreteBlend> DiscreteBlendTrait;
+		Binding.GetStackInterface(DiscreteBlendTrait);
 
 		const int32 NumChildren = InstanceData->PerChildBlendData.Num();
-		if (NumChildren < 2)
+
+		int32 NumBlending = 0;
+		for (int32 ChildIndex = 0; ChildIndex < NumChildren; ++ChildIndex)
 		{
-			return;	// If we don't have at least 2 children, there is nothing to do
+			const FBlendData& ChildBlendData = InstanceData->PerChildBlendData[ChildIndex];
+
+			NumBlending += DiscreteBlendTrait.GetBlendWeight(Context, ChildIndex) > 0.0f ? 1 : 0;
 		}
 
-		// Children are visited depth first, in the order returned
-		// As such, when we evaluate the task program, the keyframe of the last child will be
-		// on top of the keyframe stack
-		// We thus process children in reverse order
+		if (NumBlending < 2)
+		{
+			return;	// If we don't have at least 2 children blending, there is nothing to do
+		}
 
-		// The last child override the top keyframe and scales it
+		int32 ChildIndex = NumChildren - 1;
+		for (; ChildIndex >= 0; --ChildIndex)
 		{
-			const int32 ChildIndex = NumChildren - 1;
+			const float Weight = DiscreteBlendTrait.GetBlendWeight(Context, ChildIndex);
+
+			if (Weight <= 0.0f)
+			{
+				continue;	// Skip inactive child
+			}
+
+			const FBlendData& ChildBlendData = InstanceData->PerChildBlendData[ChildIndex];
 			const FBlendSampleData& PoseSampleData = InstanceData->PerBoneSampleData[ChildIndex];
 
-			Context.AppendTask(FAnimNextBlendOverwriteKeyframePerBoneWithScaleTask::Make(SharedData->BlendProfile, PoseSampleData, PoseSampleData.TotalWeight));
+			TSharedPtr<const IBlendProfileInterface> BlendProfileInterface = ChildBlendData.BlendProfileInterface;
+
+			if (BlendProfileInterface)
+			{
+				Context.AppendTask(FAnimNextBlendOverwriteKeyframePerBoneWithScaleTask::Make(BlendProfileInterface.Get(), PoseSampleData, Weight));
+			}
+			else
+			{
+				Context.AppendTask(FAnimNextBlendOverwriteKeyframeWithScaleTask::Make(Weight));
+			}
+
+			break;
 		}
 
 		// Other children accumulate with scale
-		for (int32 ChildIndex = NumChildren - 2; ChildIndex >= 0; --ChildIndex)
+		ChildIndex--;
+		for (; ChildIndex >= 0; --ChildIndex)
 		{
+			const float Weight = DiscreteBlendTrait.GetBlendWeight(Context, ChildIndex);
+
+			const FBlendData& ChildBlendData = InstanceData->PerChildBlendData[ChildIndex];
+			if (Weight <= 0.0f)
+			{
+				continue;	// Skip inactive child
+			}
+
 			const FBlendSampleData& PoseSampleDataA = InstanceData->PerBoneSampleData[ChildIndex];
 			const FBlendSampleData& PoseSampleDataB = InstanceData->PerBoneSampleData[ChildIndex + 1];	// Above on the keyframe stack
 
-			Context.AppendTask(FAnimNextBlendAddKeyframePerBoneWithScaleTask::Make(SharedData->BlendProfile, PoseSampleDataA, PoseSampleDataB, PoseSampleDataA.TotalWeight));
+			TSharedPtr<const IBlendProfileInterface> BlendProfileInterface = ChildBlendData.BlendProfileInterface;
+
+			if (BlendProfileInterface)
+			{
+				Context.AppendTask(FAnimNextBlendAddKeyframePerBoneWithScaleTask::Make(BlendProfileInterface.Get(), PoseSampleDataA, PoseSampleDataB, Weight));
+			}
+			else
+			{
+				Context.AppendTask(FAnimNextBlendAddKeyframeWithScaleTask::Make(Weight));
+			}
 		}
 
-		// Once we are done, we normalize rotations
 		Context.AppendTask(FAnimNextNormalizeKeyframeRotationsTask());
 	}
 
-	void FBlendSmootherPerBoneTrait::PreUpdate(FUpdateTraversalContext& Context, const TTraitBinding<IUpdate>& Binding, const FTraitUpdateState& TraitState) const
+	void FBlendSmootherPerBoneCoreTrait::PreUpdate(FUpdateTraversalContext& Context, const TTraitBinding<IUpdate>& Binding, const FTraitUpdateState& TraitState) const
 	{
 		const FSharedData* SharedData = Binding.GetSharedData<FSharedData>();
 		FInstanceData* InstanceData = Binding.GetInstanceData<FInstanceData>();
@@ -82,18 +145,14 @@ namespace UE::AnimNext
 		// Update the traits below us, they might trigger a transition
 		IUpdate::PreUpdate(Context, Binding, TraitState);
 
-		if (!SharedData->BlendProfile)
-		{
-			return;	// No blend profile set, nothing to do
-		}
-
 		TTraitBinding<IDiscreteBlend> DiscreteBlendTrait;
 		Binding.GetStackInterface(DiscreteBlendTrait);
 
 		const int32 DestinationChildIndex = DiscreteBlendTrait.GetBlendDestinationChildIndex(Context);
 
-		// If we're using a blend profile, extract the scales and build blend sample data
 		const int32 NumChildren = InstanceData->PerChildBlendData.Num();
+
+		// If we're using a blend profile, extract the scales and build blend sample data
 		for (int32 ChildIndex = 0; ChildIndex < NumChildren; ++ChildIndex)
 		{
 			const float BlendWeight = DiscreteBlendTrait.GetBlendWeight(Context, ChildIndex);
@@ -103,42 +162,52 @@ namespace UE::AnimNext
 			PoseSampleData.TotalWeight = BlendWeight;
 
 			const FBlendData& BlendData = InstanceData->PerChildBlendData[ChildIndex];
-			const bool bInverse = SharedData->BlendProfile->Mode == EBlendProfileMode::WeightFactor ? (DestinationChildIndex != ChildIndex) : false;
-			SharedData->BlendProfile->UpdateBoneWeights(PoseSampleData, *BlendState, BlendData.StartAlpha, BlendWeight, bInverse);
+
+			for (int32 PerBoneIndex = 0; PerBoneIndex < PoseSampleData.PerBoneBlendData.Num(); ++PerBoneIndex)
+			{
+				PoseSampleData.PerBoneBlendData[PerBoneIndex] = UBlendProfile::CalculateBoneWeight(
+					BlendData.BlendProfileInterface->GetBoneBlendScale(PerBoneIndex),
+					EBlendProfileMode::TimeFactor, *BlendState, BlendData.StartAlpha, BlendWeight, false /*bInverse*/);
+			}
 		}
 
-		FBlendSampleData::NormalizeDataWeight(InstanceData->PerBoneSampleData);
+		if (!InstanceData->PerBoneSampleData.IsEmpty())
+		{
+			FBlendSampleData::NormalizeDataWeight(InstanceData->PerBoneSampleData);
+		}
 	}
 
-	void FBlendSmootherPerBoneTrait::OnBlendTransition(FExecutionContext& Context, const TTraitBinding<IDiscreteBlend>& Binding, int32 OldChildIndex, int32 NewChildIndex) const
+	void FBlendSmootherPerBoneCoreTrait::OnBlendTransition(FExecutionContext& Context, const TTraitBinding<IDiscreteBlend>& Binding, int32 OldChildIndex, int32 NewChildIndex) const
 	{
+		IDiscreteBlend::OnBlendTransition(Context, Binding, OldChildIndex, NewChildIndex);
+
 		const FSharedData* SharedData = Binding.GetSharedData<FSharedData>();
 		FInstanceData* InstanceData = Binding.GetInstanceData<FInstanceData>();
 
-		// Trigger the new transition
-		IDiscreteBlend::OnBlendTransition(Context, Binding, OldChildIndex, NewChildIndex);
+		TTraitBinding<ISmoothBlendPerBone> SmoothBlendPerBoneTrait;
+		Binding.GetStackInterface(SmoothBlendPerBoneTrait);
 
-		if (!SharedData->BlendProfile)
-		{
-			return;	// No blend profile set, nothing to do
-		}
+		TTraitBinding<IDiscreteBlend> DiscreteBlendTrait;
+		Binding.GetStackInterface(DiscreteBlendTrait);
 
 		const int32 NumChildren = InstanceData->PerChildBlendData.Num();
 		if (NewChildIndex >= NumChildren)
 		{
 			// We have a new child
-			check(NewChildIndex == NumChildren + 1);
+			check(NewChildIndex == NumChildren);
 
-			InstanceData->PerChildBlendData.AddDefaulted();
-			FBlendSampleData& SampleData = InstanceData->PerBoneSampleData.AddDefaulted_GetRef();
+			TSharedPtr<const IBlendProfileInterface> BlendProfileInterface = SmoothBlendPerBoneTrait.GetBlendProfile(Context, NewChildIndex);
 
-			const uint32 NumBlendEntries = SharedData->BlendProfile->GetNumBlendEntries();
-			SampleData.SampleDataIndex = NewChildIndex;
-			SampleData.PerBoneBlendData.AddZeroed(NumBlendEntries);
-		}
+			const FAlphaBlend* BlendState = DiscreteBlendTrait.GetBlendState(Context, NewChildIndex);
 
-		TTraitBinding<IDiscreteBlend> DiscreteBlendTrait;
-		Binding.GetStackInterface(DiscreteBlendTrait);
+			FBlendData& ChildBlendData = InstanceData->PerChildBlendData.AddDefaulted_GetRef();
+			ChildBlendData.BlendProfileInterface = BlendProfileInterface;
+			ChildBlendData.StartAlpha = BlendState->GetAlpha();
+
+			FBlendSampleData& ChildSampleData = InstanceData->PerBoneSampleData.AddDefaulted_GetRef();
+			ChildSampleData.SampleDataIndex = NewChildIndex;
+			ChildSampleData.PerBoneBlendData.AddZeroed(BlendProfileInterface ? BlendProfileInterface->GetNumBlendEntries() : 0);
+		}
 
 		for (int32 ChildIndex = 0; ChildIndex < NumChildren; ++ChildIndex)
 		{
@@ -149,26 +218,117 @@ namespace UE::AnimNext
 		}
 	}
 
-	void FBlendSmootherPerBoneTrait::InitializeInstanceData(const FExecutionContext& Context, const FTraitBinding& Binding, const FSharedData* SharedData, FInstanceData* InstanceData)
+	void FBlendSmootherPerBoneCoreTrait::InitializeInstanceData(FExecutionContext& Context, const FTraitBinding& Binding, const FSharedData* SharedData, FInstanceData* InstanceData)
 	{
 		check(InstanceData->PerChildBlendData.IsEmpty());
 
-		if (!SharedData->BlendProfile)
-		{
-			return;	// No blend profile set, nothing to do
-		}
+		TTraitBinding<ISmoothBlendPerBone> SmoothBlendPerBoneTrait;
+		Binding.GetStackInterface(SmoothBlendPerBoneTrait);
 
 		const uint32 NumChildren = IHierarchy::GetNumStackChildren(Context, Binding);
 
 		InstanceData->PerChildBlendData.SetNum(NumChildren);
 		InstanceData->PerBoneSampleData.SetNum(NumChildren);
 
-		const uint32 NumBlendEntries = SharedData->BlendProfile->GetNumBlendEntries();
 		for (uint32 ChildIndex = 0; ChildIndex < NumChildren; ++ChildIndex)
 		{
-			FBlendSampleData& SampleData = InstanceData->PerBoneSampleData[ChildIndex];
-			SampleData.SampleDataIndex = ChildIndex;
-			SampleData.PerBoneBlendData.AddZeroed(NumBlendEntries);
+			TSharedPtr<const IBlendProfileInterface> BlendProfileInterface = SmoothBlendPerBoneTrait.GetBlendProfile(Context, ChildIndex);
+
+			FBlendData& ChildBlendData = InstanceData->PerChildBlendData[ChildIndex];
+			ChildBlendData.BlendProfileInterface = BlendProfileInterface;
+			ChildBlendData.StartAlpha = 0.0f;
+
+			if (BlendProfileInterface)
+			{
+				const uint32 NumBlendEntries = BlendProfileInterface->GetNumBlendEntries();
+				FBlendSampleData& SampleData = InstanceData->PerBoneSampleData[ChildIndex];
+				SampleData.SampleDataIndex = ChildIndex;
+				SampleData.PerBoneBlendData.AddZeroed(NumBlendEntries);
+			}
+		}
+	}
+
+	/*
+		FBlendProfilePerChildProviderTrait
+	*/
+
+	void FBlendProfilePerChildProviderTrait::PreUpdate(FUpdateTraversalContext& Context, const TTraitBinding<IUpdate>& Binding, const FTraitUpdateState& TraitState) const
+	{
+		const FSharedData* SharedData = Binding.GetSharedData<FSharedData>();
+		FInstanceData* InstanceData = Binding.GetInstanceData<FInstanceData>();
+
+		if (InstanceData->BlendProfileInterfaces.IsEmpty())
+		{
+			InitializeInstanceData(Context, Binding, SharedData, InstanceData);
 		}
+
+		IUpdate::PreUpdate(Context, Binding, TraitState);
+	}
+
+	TSharedPtr<const IBlendProfileInterface> FBlendProfilePerChildProviderTrait::GetBlendProfile(FExecutionContext& Context, const TTraitBinding<ISmoothBlendPerBone>& Binding, int32 ChildIndex) const
+	{
+		FInstanceData* InstanceData = Binding.GetInstanceData<FInstanceData>();
+
+		// Try find another trait that'll provide a custom blend profile before falling back to ours
+		if (TSharedPtr<const IBlendProfileInterface> BlendProfileInterface = ISmoothBlendPerBone::GetBlendProfile(Context, Binding, ChildIndex))
+		{
+			return BlendProfileInterface;
+		}
+
+		if (InstanceData->BlendProfileInterfaces.IsValidIndex(ChildIndex))
+		{
+			return InstanceData->BlendProfileInterfaces[ChildIndex];
+		}
+
+		return nullptr;
+	}
+
+	void FBlendProfilePerChildProviderTrait::InitializeInstanceData(FExecutionContext& Context, const FTraitBinding& Binding, const FSharedData* SharedData, FInstanceData* InstanceData)
+	{
+		check(InstanceData->BlendProfileInterfaces.IsEmpty());
+
+		for (TObjectPtr<UHierarchyTable> BlendProfile : SharedData->BlendProfiles)
+		{
+			if (BlendProfile)
+			{
+				InstanceData->BlendProfileInterfaces.Add(MakeShared<FHierarchyTableBlendProfile>(BlendProfile));
+			}
+		}
+	}
+
+	/*
+		FBlendProfileProviderTrait
+	*/
+
+	void FBlendProfileProviderTrait::PreUpdate(FUpdateTraversalContext& Context, const TTraitBinding<IUpdate>& Binding, const FTraitUpdateState& TraitState) const
+	{
+		const FSharedData* SharedData = Binding.GetSharedData<FSharedData>();
+		FInstanceData* InstanceData = Binding.GetInstanceData<FInstanceData>();
+
+		if (SharedData->TimeFactorBlendProfile && !InstanceData->BlendProfileInterface)
+		{
+			InitializeInstanceData(Context, Binding, SharedData, InstanceData);
+		}
+
+		IUpdate::PreUpdate(Context, Binding, TraitState);
+	}
+
+	TSharedPtr<const IBlendProfileInterface> FBlendProfileProviderTrait::GetBlendProfile(FExecutionContext& Context, const TTraitBinding<ISmoothBlendPerBone>& Binding, int32 ChildIndex) const
+	{
+		const FSharedData* SharedData = Binding.GetSharedData<FSharedData>();
+		FInstanceData* InstanceData = Binding.GetInstanceData<FInstanceData>();
+
+		// Try find another trait that'll provide a custom blend profile before falling back to ours
+		if (TSharedPtr<const IBlendProfileInterface> BlendProfileInterface = ISmoothBlendPerBone::GetBlendProfile(Context, Binding, ChildIndex))
+		{
+			return BlendProfileInterface;
+		}
+		
+		return InstanceData->BlendProfileInterface;
+	}
+
+	void FBlendProfileProviderTrait::InitializeInstanceData(FExecutionContext& Context, const FTraitBinding& Binding, const FSharedData* SharedData, FInstanceData* InstanceData)
+	{
+		InstanceData->BlendProfileInterface = MakeShared<FHierarchyTableBlendProfile>(SharedData->TimeFactorBlendProfile);
 	}
 }
diff --git a/Engine/Plugins/Experimental/AnimNextAnimGraph/Source/AnimNextAnimGraph/Private/Traits/BlendSmootherPerBone.h b/Engine/Plugins/Experimental/AnimNextAnimGraph/Source/AnimNextAnimGraph/Private/Traits/BlendSmootherPerBone.h
--- a/Engine/Plugins/Experimental/AnimNextAnimGraph/Source/AnimNextAnimGraph/Private/Traits/BlendSmootherPerBone.h
+++ b/Engine/Plugins/Experimental/AnimNextAnimGraph/Source/AnimNextAnimGraph/Private/Traits/BlendSmootherPerBone.h
@@ -9,40 +9,58 @@
 #include "TraitCore/Trait.h"
 #include "TraitInterfaces/IDiscreteBlend.h"
 #include "TraitInterfaces/IEvaluate.h"
+#include "TraitInterfaces/ISmoothBlend.h"
 #include "TraitInterfaces/IUpdate.h"
+#include "TraitInterfaces/ISmoothBlendPerBone.h"
+#include "HierarchyTable.h"
+#include "HierarchyTableBlendProfile.h"
 
 #include "BlendSmootherPerBone.generated.h"
 
-class UBlendProfile;
+USTRUCT(meta = (DisplayName = "Blend Smoother Per Bone Core"))
+struct FAnimNextBlendSmootherPerBoneCoreTraitSharedData : public FAnimNextTraitSharedData
+{
+	GENERATED_BODY()
+};
 
-USTRUCT(meta = (DisplayName = "Blend Smoother Per Bone"))
-struct FAnimNextBlendSmootherPerBoneTraitSharedData : public FAnimNextTraitSharedData
+USTRUCT(meta = (DisplayName = "Blend Profile Per Child Provider"))
+struct FAnimNextBlendProfilePerChildProviderTraitSharedData : public FAnimNextBlendSmootherPerBoneCoreTraitSharedData
 {
 	GENERATED_BODY()
 
 	/** Blend profile that configures how fast to blend each bone. */
-	// TODO: Can't show list of blend profiles, we need to find a skeleton to perform the lookup with
-	UPROPERTY(EditAnywhere, Category = "Default", meta = (Inline, UseAsBlendProfile=true))
-	TObjectPtr<UBlendProfile> BlendProfile = nullptr;
+	UPROPERTY(EditAnywhere, Category = "Default", meta = (Inline))
+	TArray<TObjectPtr<UHierarchyTable>> BlendProfiles;
+};
+
+USTRUCT(meta = (DisplayName = "Blend Profile Provider"))
+struct FAnimNextBlendProfileProviderTraitSharedData : public FAnimNextTraitSharedData
+{
+	GENERATED_BODY()
+
+	/** Blend profile that configures how fast to blend each bone. */
+	UPROPERTY(EditAnywhere, Category = "Default", meta = (Inline))
+	TObjectPtr<UHierarchyTable> TimeFactorBlendProfile;
 };
 
 namespace UE::AnimNext
 {
 	/**
-	 * FBlendSmootherTrait
-	 * 
-	 * A trait that smoothly blends between discrete states over time.
+	 * FBlendSmootherPerBoneCoreTrait
+	 *
+	 * An additive trait that smoothly blends with per-bone weights.
 	 */
-	struct FBlendSmootherPerBoneTrait : FAdditiveTrait, IEvaluate, IUpdate, IDiscreteBlend
+	struct FBlendSmootherPerBoneCoreTrait : FAdditiveTrait, IEvaluate, IUpdate, IDiscreteBlend
 	{
-		DECLARE_ANIM_TRAIT(FBlendSmootherPerBoneTrait, FAdditiveTrait)
+		DECLARE_ANIM_TRAIT(FBlendSmootherPerBoneCoreTrait, FAdditiveTrait)
 
-		using FSharedData = FAnimNextBlendSmootherPerBoneTraitSharedData;
+		using FSharedData = FAnimNextBlendSmootherPerBoneCoreTraitSharedData;
 
 		// Struct for tracking blends for each pose
 		struct FBlendData
 		{
-			// Which blend alpha we started the blend with
+			TSharedPtr<const IBlendProfileInterface> BlendProfileInterface;
+
 			float StartAlpha = 0.0f;
 		};
 
@@ -64,7 +82,75 @@ namespace UE::AnimNext
 		// IDiscreteBlend impl
 		virtual void OnBlendTransition(FExecutionContext& Context, const TTraitBinding<IDiscreteBlend>& Binding, int32 OldChildIndex, int32 NewChildIndex) const override;
 
+#if WITH_EDITOR
+		virtual bool IsHidden() const override { return false; }
+#endif
+
+	private:
+		// Internal impl
+		static void InitializeInstanceData(FExecutionContext& Context, const FTraitBinding& Binding, const FSharedData* SharedData, FInstanceData* InstanceData);
+	};
+
+	/**
+	 * FBlendProfilePerChildProviderTrait
+	 */
+	struct FBlendProfilePerChildProviderTrait : FAdditiveTrait, IUpdate, ISmoothBlendPerBone
+	{
+		DECLARE_ANIM_TRAIT(FBlendProfilePerChildProviderTrait, FAdditiveTrait)
+
+		using FSharedData = FAnimNextBlendProfilePerChildProviderTraitSharedData;
+
+		struct FInstanceData : FTrait::FInstanceData
+		{
+			// Blend state per child
+			TArray<TSharedPtr<const FHierarchyTableBlendProfile>> BlendProfileInterfaces;
+		};
+
+		// IUpdate impl
+		virtual void PreUpdate(FUpdateTraversalContext& Context, const TTraitBinding<IUpdate>& Binding, const FTraitUpdateState& TraitState) const override;
+
+		// ISmoothBlendPerBone impl
+		virtual TSharedPtr<const IBlendProfileInterface> GetBlendProfile(FExecutionContext& Context, const TTraitBinding<ISmoothBlendPerBone>& Binding, int32 ChildIndex) const override;
+
+#if WITH_EDITOR
+		virtual bool IsHidden() const override { return false; } // override the base, as it is hidden
+#endif
+
+	private:
+		// Internal impl
+		static void InitializeInstanceData(FExecutionContext& Context, const FTraitBinding& Binding, const FSharedData* SharedData, FInstanceData* InstanceData);
+	};
+
+	/**
+	 * FBlendProfileProviderTrait
+	 * 
+	 * Serves firstly as a passthrough node that propagates blend profiles of child traits.
+	 * If no child trait provides a blend profile then uses the held blend profile.
+	 */
+	struct FBlendProfileProviderTrait : FAdditiveTrait, IUpdate, ISmoothBlendPerBone
+	{
+		DECLARE_ANIM_TRAIT(FBlendProfileProviderTrait, FAdditiveTrait)
+
+		using FSharedData = FAnimNextBlendProfileProviderTraitSharedData;
+
+		struct FInstanceData : FTrait::FInstanceData
+		{
+			// Blend state per child
+			TSharedPtr<const FHierarchyTableBlendProfile> BlendProfileInterface;
+		};
+
+		// IUpdate impl
+		virtual void PreUpdate(FUpdateTraversalContext& Context, const TTraitBinding<IUpdate>& Binding, const FTraitUpdateState& TraitState) const override;
+
+		// ISmoothBlendPerBone impl
+		virtual TSharedPtr<const IBlendProfileInterface> GetBlendProfile(FExecutionContext& Context, const TTraitBinding<ISmoothBlendPerBone>& Binding, int32 ChildIndex) const override;
+
+#if WITH_EDITOR
+		virtual bool IsHidden() const override { return false; } // override the base, as it is hidden
+#endif
+
+	private:
 		// Internal impl
-		static void InitializeInstanceData(const FExecutionContext& Context, const FTraitBinding& Binding, const FSharedData* SharedData, FInstanceData* InstanceData);
+		static void InitializeInstanceData(FExecutionContext& Context, const FTraitBinding& Binding, const FSharedData* SharedData, FInstanceData* InstanceData);
 	};
 }
diff --git a/Engine/Plugins/Experimental/AnimNextAnimGraph/Source/AnimNextAnimGraph/Private/Traits/BlendStackTrait.cpp b/Engine/Plugins/Experimental/AnimNextAnimGraph/Source/AnimNextAnimGraph/Private/Traits/BlendStackTrait.cpp
--- a/Engine/Plugins/Experimental/AnimNextAnimGraph/Source/AnimNextAnimGraph/Private/Traits/BlendStackTrait.cpp
+++ b/Engine/Plugins/Experimental/AnimNextAnimGraph/Source/AnimNextAnimGraph/Private/Traits/BlendStackTrait.cpp
@@ -14,6 +14,7 @@
 #include "TraitCore/ExecutionContext.h"
 #include "TraitInterfaces/IGraphFactory.h"
 #include "TraitInterfaces/IInertializerBlend.h"
+#include "HierarchyTableBlendProfile.h"
 
 namespace UE::AnimNext
 {
@@ -27,6 +28,7 @@ namespace UE::AnimNext
 		GeneratorMacro(IGarbageCollection) \
 		GeneratorMacro(IHierarchy) \
 		GeneratorMacro(ISmoothBlend) \
+		GeneratorMacro(ISmoothBlendPerBone) \
 		GeneratorMacro(IInertializerBlend) \
 		GeneratorMacro(IAttributeProvider) \
 		GeneratorMacro(ITimeline) \
@@ -463,6 +465,21 @@ namespace UE::AnimNext
 		return &InstanceData->ChildGraphs[InChildIndex].Request;
 	}
 
+	TSharedPtr<const IBlendProfileInterface> FBlendStackCoreTrait::GetBlendProfile(FExecutionContext& Context, const TTraitBinding<ISmoothBlendPerBone>& Binding, int32 ChildIndex) const
+	{
+		const FInstanceData* InstanceData = Binding.GetInstanceData<FInstanceData>();
+		if (InstanceData->ChildGraphs.IsValidIndex(ChildIndex))
+		{
+			const FGraphState& Graph = InstanceData->ChildGraphs[ChildIndex];
+			return Graph.Request.BlendProfile;
+		}
+		else
+		{
+			// Unknown child
+			return nullptr;
+		}
+	}
+
 	void FBlendStackTrait::PreUpdate(FUpdateTraversalContext& Context, const TTraitBinding<IUpdate>& Binding, const FTraitUpdateState& TraitState) const
 	{
 		FBlendStackCoreTrait::PreUpdate(Context, Binding, TraitState);
diff --git a/Engine/Plugins/Experimental/AnimNextAnimGraph/Source/AnimNextAnimGraph/Public/EvaluationVM/Tasks/BlendKeyframesPerBone.h b/Engine/Plugins/Experimental/AnimNextAnimGraph/Source/AnimNextAnimGraph/Public/EvaluationVM/Tasks/BlendKeyframesPerBone.h
--- a/Engine/Plugins/Experimental/AnimNextAnimGraph/Source/AnimNextAnimGraph/Public/EvaluationVM/Tasks/BlendKeyframesPerBone.h
+++ b/Engine/Plugins/Experimental/AnimNextAnimGraph/Source/AnimNextAnimGraph/Public/EvaluationVM/Tasks/BlendKeyframesPerBone.h
@@ -11,7 +11,7 @@
 #include "BlendKeyframesPerBone.generated.h"
 
 struct FBlendSampleData;
-class UBlendProfile;
+class IBlendProfileInterface;
 class IInterpolationIndexProvider;
 class USkeleton;
 
@@ -32,9 +32,7 @@ struct ANIMNEXTANIMGRAPH_API FAnimNextBlendOverwriteKeyframePerBoneWithScaleTask
 
 	DECLARE_ANIM_EVALUATION_TASK(FAnimNextBlendOverwriteKeyframePerBoneWithScaleTask)
 
-	static FAnimNextBlendOverwriteKeyframePerBoneWithScaleTask Make(TWeakObjectPtr<const UBlendProfile> BlendProfile, const FBlendSampleData& BlendData, float ScaleFactor);
-
-	static FAnimNextBlendOverwriteKeyframePerBoneWithScaleTask Make(const IInterpolationIndexProvider* BlendProfile, const USkeleton* Skeleton, const FBlendSampleData& BlendData, float ScaleFactor);
+	static FAnimNextBlendOverwriteKeyframePerBoneWithScaleTask Make(const IBlendProfileInterface* BlendProfile, const FBlendSampleData& BlendData, float ScaleFactor);
 
 	// Task entry point
 	virtual void Execute(UE::AnimNext::FEvaluationVM& VM) const override;
@@ -46,7 +44,7 @@ struct ANIMNEXTANIMGRAPH_API FAnimNextBlendOverwriteKeyframePerBoneWithScaleTask
 
 	// The blend profile to use
 	// An optional blend profile provides a blend weight per bone
-	const IInterpolationIndexProvider* BlendProfile = nullptr;
+	const IBlendProfileInterface* BlendProfile = nullptr;
 
 	const USkeleton* SourceSkeleton = nullptr;
 };
@@ -69,7 +67,7 @@ struct ANIMNEXTANIMGRAPH_API FAnimNextBlendAddKeyframePerBoneWithScaleTask : pub
 
 	DECLARE_ANIM_EVALUATION_TASK(FAnimNextBlendAddKeyframePerBoneWithScaleTask)
 
-	static FAnimNextBlendAddKeyframePerBoneWithScaleTask Make(TWeakObjectPtr<const UBlendProfile> BlendProfile, const FBlendSampleData& BlendDataA, const FBlendSampleData& BlendDataB, float ScaleFactor);
+	static FAnimNextBlendAddKeyframePerBoneWithScaleTask Make(const IBlendProfileInterface* BlendProfile, const FBlendSampleData& BlendDataA, const FBlendSampleData& BlendDataB, float ScaleFactor);
 
 	// Task entry point
 	virtual void Execute(UE::AnimNext::FEvaluationVM& VM) const override;
@@ -86,8 +84,7 @@ struct ANIMNEXTANIMGRAPH_API FAnimNextBlendAddKeyframePerBoneWithScaleTask : pub
 
 	// The blend profile to use
 	// An optional blend profile provides a blend weight per bone
-	UPROPERTY()
-	TWeakObjectPtr<const UBlendProfile> BlendProfile;
+	const IBlendProfileInterface* BlendProfile = nullptr;
 };
 
 
diff --git a/Engine/Plugins/Experimental/AnimNextAnimGraph/Source/AnimNextAnimGraph/Public/TraitInterfaces/IBlendStack.h b/Engine/Plugins/Experimental/AnimNextAnimGraph/Source/AnimNextAnimGraph/Public/TraitInterfaces/IBlendStack.h
--- a/Engine/Plugins/Experimental/AnimNextAnimGraph/Source/AnimNextAnimGraph/Public/TraitInterfaces/IBlendStack.h
+++ b/Engine/Plugins/Experimental/AnimNextAnimGraph/Source/AnimNextAnimGraph/Public/TraitInterfaces/IBlendStack.h
@@ -9,6 +9,7 @@
 #include "AnimNextDataInterfacePayload.h"
 #include "Graph/AnimNextAnimationGraph.h"
 #include "Module/ModuleHandle.h"
+#include "Animation/BlendProfile.h"
 
 namespace UE::AnimNext
 {
@@ -52,6 +53,9 @@ namespace UE::AnimNext
 			// Blend-in duration for this graph request 
 			FAlphaBlendArgs BlendArgs;
 
+			// Optional blend profile for this graph request
+			TSharedPtr<const IBlendProfileInterface> BlendProfile;
+
 			// The template graph to use for the new graph instance
 			TObjectPtr<const UAnimNextAnimationGraph> AnimationGraph;
 
diff --git a/Engine/Plugins/Experimental/AnimNextAnimGraph/Source/AnimNextAnimGraph/Public/TraitInterfaces/ISmoothBlendPerBone.h b/Engine/Plugins/Experimental/AnimNextAnimGraph/Source/AnimNextAnimGraph/Public/TraitInterfaces/ISmoothBlendPerBone.h
new file mode 100644
--- /dev/null
+++ b/Engine/Plugins/Experimental/AnimNextAnimGraph/Source/AnimNextAnimGraph/Public/TraitInterfaces/ISmoothBlendPerBone.h
@@ -0,0 +1,48 @@
+// Copyright Epic Games, Inc. All Rights Reserved.
+
+#pragma once
+
+#include "CoreMinimal.h"
+#include "AlphaBlend.h"
+#include "TraitCore/ITraitInterface.h"
+#include "TraitCore/TraitBinding.h"
+
+class IBlendProfileInterface;
+class UCurveFloat;
+
+namespace UE::AnimNext
+{
+	/**
+	 * ISmoothBlendPerBone
+	 *
+	 * This interface exposes blend smoothing per bone related information.
+	 */
+	struct ANIMNEXTANIMGRAPH_API ISmoothBlendPerBone : ITraitInterface
+	{
+		DECLARE_ANIM_TRAIT_INTERFACE(ISmoothBlendPerBone, 0xb29ee376)
+
+		// Returns the desired blend profile for the specified child
+		virtual TSharedPtr<const IBlendProfileInterface> GetBlendProfile(FExecutionContext& Context, const TTraitBinding<ISmoothBlendPerBone>& Binding, int32 ChildIndex) const;
+
+#if WITH_EDITOR
+		virtual const FText& GetDisplayName() const override;
+		virtual const FText& GetDisplayShortName() const override;
+#endif // WITH_EDITOR
+	};
+
+	/**
+	 * Specialization for trait binding.
+	 */
+	template<>
+	struct TTraitBinding<ISmoothBlendPerBone> : FTraitBinding
+	{
+		// @see ISmoothBlendPerBone::GetBlendProfile
+		const TSharedPtr<const IBlendProfileInterface> GetBlendProfile(FExecutionContext& Context, int32 ChildIndex) const
+		{
+			return GetInterface()->GetBlendProfile(Context, *this, ChildIndex);
+		}
+
+	protected:
+		const ISmoothBlendPerBone* GetInterface() const { return GetInterfaceTyped<ISmoothBlendPerBone>(); }
+	};
+}
diff --git a/Engine/Plugins/Experimental/AnimNextPoseSearch/Source/AnimNextPoseSearch/AnimNextPoseSearch.Build.cs b/Engine/Plugins/Experimental/AnimNextPoseSearch/Source/AnimNextPoseSearch/AnimNextPoseSearch.Build.cs
--- a/Engine/Plugins/Experimental/AnimNextPoseSearch/Source/AnimNextPoseSearch/AnimNextPoseSearch.Build.cs
+++ b/Engine/Plugins/Experimental/AnimNextPoseSearch/Source/AnimNextPoseSearch/AnimNextPoseSearch.Build.cs
@@ -14,8 +14,10 @@ namespace UnrealBuildTool.Rules
 					"Core",
 					"CoreUObject",
 					"Engine",
+					"HierarchyTableRuntime",
+					"HierarchyTableAnimationRuntime",
 					"PoseSearch",
-					"RigVM"
+					"RigVM",
 				}
 			);
 		}
diff --git a/Engine/Plugins/Experimental/Animation/HierarchyTableAnimation/Source/Runtime/Private/HierarchyTableBlendProfile.cpp b/Engine/Plugins/Experimental/Animation/HierarchyTableAnimation/Source/Runtime/Private/HierarchyTableBlendProfile.cpp
new file mode 100644
--- /dev/null
+++ b/Engine/Plugins/Experimental/Animation/HierarchyTableAnimation/Source/Runtime/Private/HierarchyTableBlendProfile.cpp
@@ -0,0 +1,78 @@
+#include "HierarchyTableBlendProfile.h"
+
+#include "SkeletonHierarchyTableType.h"
+#include "HierarchyTableDefaultTypes.h"
+#include "MaskProfile/HierarchyTableTypeMask.h"
+
+FHierarchyTableBlendProfile::FHierarchyTableBlendProfile(TObjectPtr<UHierarchyTable> InHierarchyTable)
+{
+	const FHierarchyTable_TableType_Skeleton& TableMetadata = InHierarchyTable->GetTableMetadata<FHierarchyTable_TableType_Skeleton>();
+	Skeleton = TableMetadata.Skeleton;
+
+	const FReferenceSkeleton RefSkeleton = Skeleton->GetReferenceSkeleton();
+
+	const int32 BoneCount = RefSkeleton.GetNum();
+	BoneBlendWeights.Reserve(BoneCount);
+
+	if (InHierarchyTable->IsElementType<FHierarchyTable_ElementType_Float>())
+	{
+		for (int32 BoneIndex = 0; BoneIndex < BoneCount; ++BoneIndex)
+		{
+			if (const FHierarchyTableEntryData* EntryData = InHierarchyTable->GetTableEntry(BoneIndex))
+			{
+				BoneBlendWeights.Add(EntryData->GetValue<FHierarchyTable_ElementType_Float>()->Value);
+			}
+			else
+			{
+				BoneBlendWeights.Add(0.0f);
+			}
+		}
+	}
+	else if (InHierarchyTable->IsElementType<FHierarchyTable_ElementType_Mask>())
+	{
+		for (int32 BoneIndex = 0; BoneIndex < BoneCount; ++BoneIndex)
+		{
+			if (const FHierarchyTableEntryData* EntryData = InHierarchyTable->GetTableEntry(BoneIndex))
+			{
+				BoneBlendWeights.Add(EntryData->GetValue<FHierarchyTable_ElementType_Mask>()->Value);
+			}
+			else
+			{
+				BoneBlendWeights.Add(0.0f);
+			}
+		}
+	}
+	else
+	{
+		ensureMsgf(false, TEXT("Unsupported hierarchy table type, use Float or Mask element type instead"));
+		for (int32 BoneIndex = 0; BoneIndex < BoneCount; ++BoneIndex)
+		{
+			BoneBlendWeights.Add(1.0f);
+		}
+	}
+}
+
+float FHierarchyTableBlendProfile::GetBoneBlendScale(int32 InBoneIdx) const
+{
+	return BoneBlendWeights[InBoneIdx];
+}
+
+int32 FHierarchyTableBlendProfile::GetNumBlendEntries() const
+{
+	return BoneBlendWeights.Num();
+}
+
+int32 FHierarchyTableBlendProfile::GetPerBoneInterpolationIndex(const FCompactPoseBoneIndex& InCompactPoseBoneIndex, const FBoneContainer& BoneContainer, const IInterpolationIndexProvider::FPerBoneInterpolationData* Data) const
+{
+	return InCompactPoseBoneIndex.GetInt();
+}
+
+int32 FHierarchyTableBlendProfile::GetPerBoneInterpolationIndex(const FSkeletonPoseBoneIndex InSkeletonBoneIndex, const USkeleton* TargetSkeleton, const IInterpolationIndexProvider::FPerBoneInterpolationData* Data) const
+{
+	return InSkeletonBoneIndex.GetInt();
+}
+
+TObjectPtr<USkeleton> FHierarchyTableBlendProfile::GetSkeleton() const
+{
+	return Skeleton;
+}
\ No newline at end of file
diff --git a/Engine/Plugins/Experimental/Animation/HierarchyTableAnimation/Source/Runtime/Public/HierarchyTableBlendProfile.h b/Engine/Plugins/Experimental/Animation/HierarchyTableAnimation/Source/Runtime/Public/HierarchyTableBlendProfile.h
new file mode 100644
--- /dev/null
+++ b/Engine/Plugins/Experimental/Animation/HierarchyTableAnimation/Source/Runtime/Public/HierarchyTableBlendProfile.h
@@ -0,0 +1,27 @@
+// Copyright Epic Games, Inc. All Rights Reserved.
+
+#pragma once
+
+#include "Animation/BlendProfile.h"
+#include "HierarchyTable.h"
+
+class HIERARCHYTABLEANIMATIONRUNTIME_API FHierarchyTableBlendProfile : public IBlendProfileInterface
+{
+public:
+	FHierarchyTableBlendProfile() = default;
+	virtual ~FHierarchyTableBlendProfile() = default;
+
+	FHierarchyTableBlendProfile(TObjectPtr<UHierarchyTable> InHierarchyTable);
+
+	// IBlendProfileInterface
+	virtual float GetBoneBlendScale(int32 InBoneIdx) const override;
+	virtual int32 GetNumBlendEntries() const override;
+	virtual int32 GetPerBoneInterpolationIndex(const FCompactPoseBoneIndex& InCompactPoseBoneIndex, const FBoneContainer& BoneContainer, const IInterpolationIndexProvider::FPerBoneInterpolationData* Data) const override;
+	virtual int32 GetPerBoneInterpolationIndex(const FSkeletonPoseBoneIndex InSkeletonBoneIndex, const USkeleton* TargetSkeleton, const IInterpolationIndexProvider::FPerBoneInterpolationData* Data) const override;
+	virtual TObjectPtr<USkeleton> GetSkeleton() const override;
+
+private:
+	TObjectPtr<USkeleton> Skeleton;
+
+	TArray<float> BoneBlendWeights;
+};
\ No newline at end of file
diff --git a/Engine/Source/Runtime/Engine/Classes/Animation/BlendProfile.h b/Engine/Source/Runtime/Engine/Classes/Animation/BlendProfile.h
--- a/Engine/Source/Runtime/Engine/Classes/Animation/BlendProfile.h
+++ b/Engine/Source/Runtime/Engine/Classes/Animation/BlendProfile.h
@@ -49,13 +49,25 @@ struct FBlendProfileBoneEntry
 	float BlendScale = 0.f;
 };
 
+class IBlendProfileInterface : public IInterpolationIndexProvider
+{
+public:
+	virtual ~IBlendProfileInterface() = default;
+
+	virtual float GetBoneBlendScale(int32 InBoneIdx) const = 0;
+
+	virtual int32 GetNumBlendEntries() const = 0;
+
+	virtual TObjectPtr<USkeleton> GetSkeleton() const = 0;
+};
+
 //////////////////////////////////////////////////////////////////////////
 
 /** A blend profile is a set of per-bone scales that can be used in transitions and blend lists
  *  to tweak the weights of specific bones. The scales are applied to the normal weight for that bone
  */
-UCLASS(Within=Skeleton, MinimalAPI, BlueprintType)
-class UBlendProfile : public UObject, public IInterpolationIndexProvider
+UCLASS(MinimalAPI, BlueprintType)
+class UBlendProfile : public UObject, public IBlendProfileInterface
 {
 public:
 
@@ -98,6 +110,8 @@ public:
 	/** Removes entries with bone references to invalid bones */
 	ENGINE_API void CleanupBoneEntries();
 
+	ENGINE_API TObjectPtr<USkeleton> GetSkeleton() const override { return OwningSkeleton; }
+
 	/**
 	 * Get the bone entry by entry index.
 	 * @param[in] InEntryIdx The index to the bone entry in range [0, GetNumBlendEntries()-1].
@@ -108,7 +122,7 @@ public:
 	/** Get the set blend scale for the specified bone, will return 1.0f if no entry was found (no scale)
 	 *  @param InBoneIdx Index of the bone to retrieve
 	 */
-	ENGINE_API float GetBoneBlendScale(int32 InBoneIdx) const;
+	ENGINE_API float GetBoneBlendScale(int32 InBoneIdx) const override;
 
 	/** Get the set blend scale for the specified bone, will return 1.0f if no entry was found (no scale). The term bone factor and entry or bone scale refer to the same thing.
 	 *  @param InBoneName Name of the bone to retrieve
