Add more detailed comments to various trait interfaces

#tests NonUnity

diff --git a/Engine/Plugins/Experimental/AnimNextAnimGraph/Source/AnimNextAnimGraph/Public/TraitInterfaces/IContinuousBlend.h b/Engine/Plugins/Experimental/AnimNextAnimGraph/Source/AnimNextAnimGraph/Public/TraitInterfaces/IContinuousBlend.h
--- a/Engine/Plugins/Experimental/AnimNextAnimGraph/Source/AnimNextAnimGraph/Public/TraitInterfaces/IContinuousBlend.h
+++ b/Engine/Plugins/Experimental/AnimNextAnimGraph/Source/AnimNextAnimGraph/Public/TraitInterfaces/IContinuousBlend.h
@@ -12,6 +12,10 @@ namespace UE::AnimNext
 	 * IContinuousBlend
 	 *
 	 * This interface exposes continuous blend related information.
+	 * 
+	 * Assuming that each blend source is temporally coherent (Has no snaps or pops in itself).
+	 * Then, a continuous blend's poses should be temporally coherent from one pose to the next.
+	 * Hence these poses will have no need for things like smoothing between each other.
 	 */
 	struct ANIMNEXTANIMGRAPH_API IContinuousBlend : ITraitInterface
 	{
diff --git a/Engine/Plugins/Experimental/AnimNextAnimGraph/Source/AnimNextAnimGraph/Public/TraitInterfaces/IDiscreteBlend.h b/Engine/Plugins/Experimental/AnimNextAnimGraph/Source/AnimNextAnimGraph/Public/TraitInterfaces/IDiscreteBlend.h
--- a/Engine/Plugins/Experimental/AnimNextAnimGraph/Source/AnimNextAnimGraph/Public/TraitInterfaces/IDiscreteBlend.h
+++ b/Engine/Plugins/Experimental/AnimNextAnimGraph/Source/AnimNextAnimGraph/Public/TraitInterfaces/IDiscreteBlend.h
@@ -14,6 +14,8 @@ namespace UE::AnimNext
 	 * IDiscreteBlend
 	 *
 	 * This interface exposes discrete blend related information.
+	 * 
+	 * Discrete blends may snap or pop from one blend source to another, they may be non-continuous.
 	 */
 	struct ANIMNEXTANIMGRAPH_API IDiscreteBlend : ITraitInterface
 	{
diff --git a/Engine/Plugins/Experimental/AnimNextAnimGraph/Source/AnimNextAnimGraph/Public/TraitInterfaces/IInertializerBlend.h b/Engine/Plugins/Experimental/AnimNextAnimGraph/Source/AnimNextAnimGraph/Public/TraitInterfaces/IInertializerBlend.h
--- a/Engine/Plugins/Experimental/AnimNextAnimGraph/Source/AnimNextAnimGraph/Public/TraitInterfaces/IInertializerBlend.h
+++ b/Engine/Plugins/Experimental/AnimNextAnimGraph/Source/AnimNextAnimGraph/Public/TraitInterfaces/IInertializerBlend.h
@@ -12,6 +12,10 @@ namespace UE::AnimNext
 	 * IInertializerBlend
 	 *
 	 * This interface exposes inertializing blend related information.
+	 * 
+	 * Inertialization is a process in which the delta between two blend sources is made to eventually be zero over time.
+	 * It can be used as a form of blend smoothing.
+	 * Examples: Storing pose deltas and fading out over time, dead blending.
 	 */
 	struct ANIMNEXTANIMGRAPH_API IInertializerBlend : ITraitInterface
 	{
diff --git a/Engine/Plugins/Experimental/AnimNextAnimGraph/Source/AnimNextAnimGraph/Public/TraitInterfaces/ISmoothBlend.h b/Engine/Plugins/Experimental/AnimNextAnimGraph/Source/AnimNextAnimGraph/Public/TraitInterfaces/ISmoothBlend.h
--- a/Engine/Plugins/Experimental/AnimNextAnimGraph/Source/AnimNextAnimGraph/Public/TraitInterfaces/ISmoothBlend.h
+++ b/Engine/Plugins/Experimental/AnimNextAnimGraph/Source/AnimNextAnimGraph/Public/TraitInterfaces/ISmoothBlend.h
@@ -15,6 +15,8 @@ namespace UE::AnimNext
 	 * ISmoothBlend
 	 *
 	 * This interface exposes blend smoothing related information.
+	 * 
+	 * Blend smoothing is a process that may be needed for discrete blends to make their pose transition appear continuous (IE: smooth).
 	 */
 	struct ANIMNEXTANIMGRAPH_API ISmoothBlend : ITraitInterface
 	{
