Remove 'b' prefix for native data interface variable

#virtualized

diff --git a/Engine/Plugins/Experimental/AnimNextAnimGraph/Source/AnimNextAnimGraph/Public/GraphInterfaces/AnimNextNativeDataInterface_BlendSpacePlayer.h b/Engine/Plugins/Experimental/AnimNextAnimGraph/Source/AnimNextAnimGraph/Public/GraphInterfaces/AnimNextNativeDataInterface_BlendSpacePlayer.h
--- a/Engine/Plugins/Experimental/AnimNextAnimGraph/Source/AnimNextAnimGraph/Public/GraphInterfaces/AnimNextNativeDataInterface_BlendSpacePlayer.h
+++ b/Engine/Plugins/Experimental/AnimNextAnimGraph/Source/AnimNextAnimGraph/Public/GraphInterfaces/AnimNextNativeDataInterface_BlendSpacePlayer.h
@@ -41,5 +41,5 @@ struct FAnimNextNativeDataInterface_BlendSpacePlayer : public FAnimNextNativeDat
 
 	// Whether to loop the animation
 	UPROPERTY(EditAnywhere, BlueprintReadWrite, Category = "Blend Space", meta = (EnableCategories))
-	bool bLoop = false;
+	bool Loop = false;
 };
diff --git a/Engine/Plugins/Experimental/AnimNextPoseSearch/Source/AnimNextPoseSearch/Private/MotionMatchingTrait.cpp b/Engine/Plugins/Experimental/AnimNextPoseSearch/Source/AnimNextPoseSearch/Private/MotionMatchingTrait.cpp
--- a/Engine/Plugins/Experimental/AnimNextPoseSearch/Source/AnimNextPoseSearch/Private/MotionMatchingTrait.cpp
+++ b/Engine/Plugins/Experimental/AnimNextPoseSearch/Source/AnimNextPoseSearch/Private/MotionMatchingTrait.cpp
@@ -284,7 +284,7 @@ namespace UE::AnimNext
 
 				NewNativeInterface.PlayRate = DesiredPlayRate;
 				NewNativeInterface.StartPosition = MotionMatchingState.CurrentSearchResult.AssetTime;
-				NewNativeInterface.bLoop = bLoop;
+				NewNativeInterface.Loop = bLoop;
 			}
 			else
 			{
