Add explicit control over when to synchronize with the sync group trait

#virtualized

diff --git a/Engine/Plugins/Experimental/AnimNextAnimGraph/Source/AnimNextAnimGraph/Private/Graph/SyncGroup_GraphInstanceComponent.cpp b/Engine/Plugins/Experimental/AnimNextAnimGraph/Source/AnimNextAnimGraph/Private/Graph/SyncGroup_GraphInstanceComponent.cpp
--- a/Engine/Plugins/Experimental/AnimNextAnimGraph/Source/AnimNextAnimGraph/Private/Graph/SyncGroup_GraphInstanceComponent.cpp
+++ b/Engine/Plugins/Experimental/AnimNextAnimGraph/Source/AnimNextAnimGraph/Private/Graph/SyncGroup_GraphInstanceComponent.cpp
@@ -96,6 +96,16 @@ namespace UE::AnimNext
 			TBitArray<TMemStackAllocator<>> MemberSeenMarkers;
 		};
 
+		// A unique group name we've generated
+		struct FSyncGroupUniqueName
+		{
+			FName GroupName;
+
+			FSyncGroupUniqueName* NextFreeEntry = nullptr;
+		};
+
+		static FName NAME_UniqueGroupNamePrefix = TEXT("UE_UNIQUE_GROUP_NAME");
+
 		static void InitGroup(const FSyncGroupState& GroupState, TArray<FSyncGroupMemberContext, TMemStackAllocator<>>& ContextPool, FSyncGroupContext& GroupContext)
 		{
 			const int32 NumMembers = GroupState.Members.Num();
@@ -919,12 +929,17 @@ namespace UE::AnimNext
 				GroupContext.LeaderDeltaTime = DeltaTime;
 
 #if UE_LOG_ENABLE_MARKER_DEBUG
+				const bool bLooped = DeltaTime >= 0.0f
+					? (GroupContext.LeaderTimelineState.GetPosition() < GroupLeaderContext.TimelineState.GetPosition())
+					: (GroupContext.LeaderTimelineState.GetPosition() > GroupLeaderContext.TimelineState.GetPosition());
+
 				UE_LOG(LogAnimMarkerSync, Warning,
-					TEXT("[%p] Leader [Time] [%0.2f + %0.2f -> %0.2f / %0.2f (%0.2f%%)]"),
+					TEXT("[%p] Leader [Time] [%0.2f + %0.2f -> %0.2f / %0.2f (%0.2f%%)]%s"),
 					GroupLeader.TraitPtr.GetNodeInstance(),
 					GroupLeaderContext.TimelineState.GetPosition(), DeltaTime,
 					GroupContext.LeaderTimelineState.GetPosition(), GroupContext.LeaderTimelineState.GetDuration(),
-					GroupContext.LeaderTimelineState.GetPositionRatio());
+					GroupContext.LeaderTimelineState.GetPositionRatio(),
+					bLooped ? TEXT(" (looped)") : TEXT(""));
 #endif
 			}
 		}
@@ -967,6 +982,7 @@ namespace UE::AnimNext
 
 #if UE_LOG_ENABLE_MARKER_DEBUG
 				const TCHAR* SyncMode = TEXT("????");
+				bool bLooped = false;
 #endif
 
 				if (GroupMember.GroupParameters.GroupRole == EAnimGroupSynchronizationRole::ExclusiveAlwaysLeader)
@@ -998,6 +1014,10 @@ namespace UE::AnimNext
 						{
 							// We must have looped around but we still want a positive delta time to match our leader
 							CurrentPosition += GroupMemberContext.TimelineState.GetDuration();
+
+#if UE_LOG_ENABLE_MARKER_DEBUG
+							bLooped = true;
+#endif
 						}
 					}
 					else
@@ -1006,6 +1026,10 @@ namespace UE::AnimNext
 						{
 							// We must have looped around but we still want a negative delta time to match our leader
 							CurrentPosition -= GroupMemberContext.TimelineState.GetDuration();
+
+#if UE_LOG_ENABLE_MARKER_DEBUG
+							bLooped = true;
+#endif
 						}
 					}
 
@@ -1022,11 +1046,12 @@ namespace UE::AnimNext
 				const FTimelineState CurrTimelineState = GroupMemberContext.TimelineTrait.GetState(Context);
 
 				UE_LOG(LogAnimMarkerSync, Warning,
-					TEXT("[%p] Follower [%s] [%0.2f + %0.2f -> %0.2f / %0.2f (%0.2f%%)]"),
+					TEXT("[%p] Follower [%s] [%0.2f + %0.2f -> %0.2f / %0.2f (%0.2f%%)]%s"),
 					GroupMember.TraitPtr.GetNodeInstance(), SyncMode,
 					GroupMemberContext.TimelineState.GetPosition(), DeltaTime,
 					CurrTimelineState.GetPosition(), CurrTimelineState.GetDuration(),
-					CurrTimelineState.GetPositionRatio());
+					CurrTimelineState.GetPositionRatio(),
+					bLooped ? TEXT(" (looped)") : TEXT(""));
 #endif
 			}
 		}
@@ -1037,7 +1062,23 @@ namespace UE::AnimNext
 	{
 	}
 
-	FSyncGroupGraphInstanceComponent::~FSyncGroupGraphInstanceComponent() = default;
+	FSyncGroupGraphInstanceComponent::~FSyncGroupGraphInstanceComponent()
+	{
+		// We should have released all used group names before we are destroyed
+		check(UsedUniqueGroupNames.IsEmpty());
+
+		Private::FSyncGroupUniqueName* Entry = FirstFreeUniqueGroupName;
+		while (Entry != nullptr)
+		{
+			Private::FSyncGroupUniqueName* NextEntry = Entry->NextFreeEntry;
+
+			delete Entry;
+
+			Entry = NextEntry;
+		}
+
+		FirstFreeUniqueGroupName = nullptr;
+	}
 
 	void FSyncGroupGraphInstanceComponent::RegisterWithGroup(const FSyncGroupParameters& GroupParameters, const FWeakTraitPtr& TraitPtr, const FTraitUpdateState& TraitState)
 	{
@@ -1048,6 +1089,50 @@ namespace UE::AnimNext
 		MaxNumMembers = FMath::Max<uint32>(MaxNumMembers, GroupState.Members.Num());
 	}
 
+	FName FSyncGroupGraphInstanceComponent::CreateUniqueGroupName()
+	{
+		// We create a unique group name by incrementing a counter and using it to generate a new FName
+		// We also recycle them to keep the counter as low as possible because in some configurations
+		// FNames are 32-bit as the number is stored in the string table. Recycling the group names
+		// ensures we don't pollute the FName string table.
+		// The counter could be the same across multiple graph instances which is fine since we only
+		// synchronize using groups within a single graph instance.
+
+		Private::FSyncGroupUniqueName* Entry;
+
+		if (FirstFreeUniqueGroupName != nullptr)
+		{
+			// Recycle an old entry
+			Entry = FirstFreeUniqueGroupName;
+			FirstFreeUniqueGroupName = Entry->NextFreeEntry;
+		}
+		else
+		{
+			// Generate a new unique group name
+			FName UniqueGroupName = FName(Private::NAME_UniqueGroupNamePrefix, UniqueGroupNameCounter);
+			UniqueGroupNameCounter++;
+
+			Entry = new Private::FSyncGroupUniqueName();
+			Entry->GroupName = UniqueGroupName;
+		}
+
+		Entry->NextFreeEntry = nullptr;
+
+		UsedUniqueGroupNames.Add(Entry->GroupName, Entry);
+
+		return Entry->GroupName;
+	}
+
+	void FSyncGroupGraphInstanceComponent::ReleaseUniqueGroupName(FName GroupName)
+	{
+		Private::FSyncGroupUniqueName* Entry = nullptr;
+		if (ensure(UsedUniqueGroupNames.RemoveAndCopyValue(GroupName, Entry)))
+		{
+			Entry->NextFreeEntry = FirstFreeUniqueGroupName;
+			FirstFreeUniqueGroupName = Entry;
+		}
+	}
+
 	void FSyncGroupGraphInstanceComponent::PreUpdate(FExecutionContext& Context)
 	{
 		bool bRebuildGroupIndexMap = false;
@@ -1058,7 +1143,7 @@ namespace UE::AnimNext
 			Private::FSyncGroupState& GroupState = SyncGroups[SyncGroupIndex];
 
 			const bool bWasActive = GroupState.Members.IsEmpty();
-			if (!GroupState.bWasActive && bWasActive)
+			if (!GroupState.bWasActive && !bWasActive)
 			{
 				// We have been inactive for two consecutive updates, time to purge this group
 				SyncGroups.RemoveAtSwap(SyncGroupIndex, EAllowShrinking::No);
@@ -1099,6 +1184,14 @@ namespace UE::AnimNext
 		// Now that we have discovered all groups and their memberships, we can perform synchronization
 		for (FSyncGroupState& GroupState : SyncGroups)
 		{
+			if (GroupState.Members.IsEmpty())
+			{
+				// Skip empty groups
+				// A group can be empty if it became inactive in the previous frame
+				check(!GroupState.bWasActive);
+				continue;
+			}
+
 			// Initialize our context for this new group
 			InitGroup(GroupState, MemberContexts, GroupContext);
 
diff --git a/Engine/Plugins/Experimental/AnimNextAnimGraph/Source/AnimNextAnimGraph/Private/Graph/SyncGroup_GraphInstanceComponent.h b/Engine/Plugins/Experimental/AnimNextAnimGraph/Source/AnimNextAnimGraph/Private/Graph/SyncGroup_GraphInstanceComponent.h
--- a/Engine/Plugins/Experimental/AnimNextAnimGraph/Source/AnimNextAnimGraph/Private/Graph/SyncGroup_GraphInstanceComponent.h
+++ b/Engine/Plugins/Experimental/AnimNextAnimGraph/Source/AnimNextAnimGraph/Private/Graph/SyncGroup_GraphInstanceComponent.h
@@ -12,6 +12,7 @@ namespace UE::AnimNext
 	namespace Private
 	{
 		struct FSyncGroupState;
+		struct FSyncGroupUniqueName;
 	}
 
 	struct FSyncGroupParameters;
@@ -28,8 +29,17 @@ namespace UE::AnimNext
 		explicit FSyncGroupGraphInstanceComponent(FAnimNextGraphInstance& InOwnerInstance);
 		~FSyncGroupGraphInstanceComponent();
 
+		// Registers the specified trait with group based synchronization
 		void RegisterWithGroup(const FSyncGroupParameters& GroupParameters, const FWeakTraitPtr& TraitPtr, const FTraitUpdateState& TraitState);
 
+		// Create a unique group name suitable for spawned sub-graphs to self-synchronize
+		// Unique group names are a limited resource, when no longed needed they must be released
+		FName CreateUniqueGroupName();
+
+		// Releases a unique group name that is no longer needed. It will be recycled the
+		// next time one is needed
+		void ReleaseUniqueGroupName(FName GroupName);
+
 		// FGraphInstanceComponent impl
 		virtual void PreUpdate(FExecutionContext& Context) override;
 		virtual void PostUpdate(FExecutionContext& Context) override;
@@ -41,7 +51,16 @@ namespace UE::AnimNext
 		// A list of groups and their data
 		TArray<Private::FSyncGroupState> SyncGroups;
 
+		// The first free unique group name
+		Private::FSyncGroupUniqueName* FirstFreeUniqueGroupName = nullptr;
+
+		// The map of currently used group names
+		TMap<FName, Private::FSyncGroupUniqueName*> UsedUniqueGroupNames;
+
 		// The number of members in the largest group
 		uint32 MaxNumMembers = 0;
+
+		// A counter tracking the next unique group name to allocate
+		int32 UniqueGroupNameCounter = 0;
 	};
 }
diff --git a/Engine/Plugins/Experimental/AnimNextAnimGraph/Source/AnimNextAnimGraph/Private/Traits/SynchronizeUsingGroups.cpp b/Engine/Plugins/Experimental/AnimNextAnimGraph/Source/AnimNextAnimGraph/Private/Traits/SynchronizeUsingGroups.cpp
--- a/Engine/Plugins/Experimental/AnimNextAnimGraph/Source/AnimNextAnimGraph/Private/Traits/SynchronizeUsingGroups.cpp
+++ b/Engine/Plugins/Experimental/AnimNextAnimGraph/Source/AnimNextAnimGraph/Private/Traits/SynchronizeUsingGroups.cpp
@@ -22,10 +22,20 @@ namespace UE::AnimNext
 
 	void FSynchronizeUsingGroupsTrait::FInstanceData::Construct(const FExecutionContext& Context, const FTraitBinding& Binding)
 	{
+		SyncGroupComponent = &Context.GetComponent<FSyncGroupGraphInstanceComponent>();
+
 		TTraitBinding<ITimeline> TimelineTrait;
 		bHasTimeline = Binding.GetStackInterface(TimelineTrait);
 	}
 
+	void FSynchronizeUsingGroupsTrait::FInstanceData::Destruct(const FExecutionContext& Context, const FTraitBinding& Binding)
+	{
+		if (!UniqueGroupName.IsNone())
+		{
+			SyncGroupComponent->ReleaseUniqueGroupName(UniqueGroupName);
+		}
+	}
+
 	void FSynchronizeUsingGroupsTrait::OnBecomeRelevant(FUpdateTraversalContext& Context, const TTraitBinding<IUpdate>& Binding, const FTraitUpdateState& TraitState) const
 	{
 		IUpdate::OnBecomeRelevant(Context, Binding, TraitState);
@@ -67,31 +77,52 @@ namespace UE::AnimNext
 		}
 
 		// Append this trait to our group, we'll need to synchronize it
-		FSyncGroupGraphInstanceComponent& Component = Context.GetComponent<FSyncGroupGraphInstanceComponent>();
-		Component.RegisterWithGroup(GroupParameters, Binding.GetTraitPtr(), TraitState);
+		InstanceData->SyncGroupComponent->RegisterWithGroup(GroupParameters, Binding.GetTraitPtr(), TraitState);
 	}
 
 	FSyncGroupParameters FSynchronizeUsingGroupsTrait::GetGroupParameters(FExecutionContext& Context, const TTraitBinding<IGroupSynchronization>& Binding) const
 	{
 		const FSharedData* SharedData = Binding.GetSharedData<FSharedData>();
+		FInstanceData* InstanceData = Binding.GetInstanceData<FInstanceData>();
 
-		FSyncGroupParameters Parameters;
-		Parameters.GroupName = SharedData->GetGroupName(Binding);
-		Parameters.GroupRole = SharedData->GetGroupRole(Binding);
-		Parameters.bMatchSyncPoint = SharedData->GetbMatchSyncPoint(Binding);
+		FName GroupName;
+		EAnimGroupSynchronizationRole GroupRole = SharedData->GetGroupRole(Binding);
+		EAnimGroupSynchronizationMode SyncMode = SharedData->GetSyncMode(Binding);
+
+		switch (SyncMode)
+		{
+		case EAnimGroupSynchronizationMode::NoSynchronization:
+			GroupName = NAME_None;
+			break;
+		case EAnimGroupSynchronizationMode::SynchronizeUsingGroupName:
+			GroupName = SharedData->GetGroupName(Binding);
+			break;
+		case EAnimGroupSynchronizationMode::SynchronizeUsingUniqueGroupName:
+			if (InstanceData->UniqueGroupName.IsNone())
+			{
+				InstanceData->UniqueGroupName = InstanceData->SyncGroupComponent->CreateUniqueGroupName();
+			}
+
+			GroupName = InstanceData->UniqueGroupName;
+			break;
+		}
 
 		// When the role is set to TransitionLeader or TransitionFollower, we will not be part of a sync group until we are full weight
 		// meaning we ignore the sync group while blending in, but not while blending out
-		if (Parameters.GroupRole == EAnimGroupSynchronizationRole::TransitionLeader || Parameters.GroupRole == EAnimGroupSynchronizationRole::TransitionFollower)
+		if (GroupRole == EAnimGroupSynchronizationRole::TransitionLeader || GroupRole == EAnimGroupSynchronizationRole::TransitionFollower)
 		{
-			const FInstanceData* InstanceData = Binding.GetInstanceData<FInstanceData>();
-
 			if (InstanceData->bHasReachedFullWeight)
 			{
-				Parameters.GroupName = NAME_None;
+				GroupName = NAME_None;
 			}
 		}
 
+		FSyncGroupParameters Parameters;
+		Parameters.GroupName = GroupName;
+		Parameters.GroupRole = GroupRole;
+		Parameters.SyncMode = SyncMode;
+		Parameters.bMatchSyncPoint = SharedData->GetbMatchSyncPoint(Binding);
+
 		return Parameters;
 	}
 
@@ -135,12 +166,17 @@ namespace UE::AnimNext
 
 		const FSyncGroupParameters GroupParameters = GroupSyncTrait.GetGroupParameters(Context);
 
-		FInstancedStruct SyncGroupPayload;
-		auto& NativeInterface = SyncGroupPayload.InitializeAs<FAnimNextNativeDataInterface_SynchronizeUsingGroups>();
-		NativeInterface.GroupName = GroupParameters.GroupName;
-		NativeInterface.GroupRole = GroupParameters.GroupRole;
-		NativeInterface.MatchSyncPoint = GroupParameters.bMatchSyncPoint;
-
-		InOutPayload.AddNative(MoveTemp(SyncGroupPayload));
+		// If synchronization is enabled, we populate the native interface
+		if (GroupParameters.SyncMode != EAnimGroupSynchronizationMode::NoSynchronization)
+		{
+			FInstancedStruct SyncGroupPayload;
+			auto& NativeInterface = SyncGroupPayload.InitializeAs<FAnimNextNativeDataInterface_SynchronizeUsingGroups>();
+			NativeInterface.GroupName = GroupParameters.GroupName;
+			NativeInterface.GroupRole = GroupParameters.GroupRole;
+			NativeInterface.SyncMode = GroupParameters.SyncMode;
+			NativeInterface.MatchSyncPoint = GroupParameters.bMatchSyncPoint;
+
+			InOutPayload.AddNative(MoveTemp(SyncGroupPayload));
+		}
 	}
 }
diff --git a/Engine/Plugins/Experimental/AnimNextAnimGraph/Source/AnimNextAnimGraph/Private/Traits/SynchronizeUsingGroups.h b/Engine/Plugins/Experimental/AnimNextAnimGraph/Source/AnimNextAnimGraph/Private/Traits/SynchronizeUsingGroups.h
--- a/Engine/Plugins/Experimental/AnimNextAnimGraph/Source/AnimNextAnimGraph/Private/Traits/SynchronizeUsingGroups.h
+++ b/Engine/Plugins/Experimental/AnimNextAnimGraph/Source/AnimNextAnimGraph/Private/Traits/SynchronizeUsingGroups.h
@@ -25,6 +25,10 @@ struct FAnimNextSynchronizeUsingGroupsTraitSharedData : public FAnimNextTraitSha
 	UPROPERTY(EditAnywhere, Category = "Default")
 	EAnimGroupSynchronizationRole GroupRole = EAnimGroupSynchronizationRole::CanBeLeader;
 
+	// The synchronization mode
+	UPROPERTY(EditAnywhere, Category = "Default")
+	EAnimGroupSynchronizationMode SyncMode = EAnimGroupSynchronizationMode::NoSynchronization;
+
 	// Whether or not to match the group sync point when joining as leader
 	// When disabled, the leader will snap the position of the group to its current position
 	UPROPERTY(EditAnywhere, Category = "Default")
@@ -34,6 +38,7 @@ struct FAnimNextSynchronizeUsingGroupsTraitSharedData : public FAnimNextTraitSha
 	#define TRAIT_LATENT_PROPERTIES_ENUMERATOR(GeneratorMacro) \
 		GeneratorMacro(GroupName) \
 		GeneratorMacro(GroupRole) \
+		GeneratorMacro(SyncMode) \
 		GeneratorMacro(bMatchSyncPoint) \
 
 	GENERATE_TRAIT_LATENT_PROPERTIES(FAnimNextSynchronizeUsingGroupsTraitSharedData, TRAIT_LATENT_PROPERTIES_ENUMERATOR)
@@ -42,6 +47,8 @@ struct FAnimNextSynchronizeUsingGroupsTraitSharedData : public FAnimNextTraitSha
 
 namespace UE::AnimNext
 {
+	struct FSyncGroupGraphInstanceComponent;
+
 	/**
 	 * FSynchronizeUsingGroupsTrait
 	 * 
@@ -56,6 +63,13 @@ namespace UE::AnimNext
 		struct FInstanceData : FTraitInstanceData
 		{
 			void Construct(const FExecutionContext& Context, const FTraitBinding& Binding);
+			void Destruct(const FExecutionContext& Context, const FTraitBinding& Binding);
+
+			// Cached pointer to our sync group component
+			FSyncGroupGraphInstanceComponent* SyncGroupComponent = nullptr;
+
+			// If our sync mode requires a unique group name, we'll cache it and re-use it as it is unique to our instance
+			FName UniqueGroupName;
 
 			bool bFreezeTimeline = false;
 			bool bHasReachedFullWeight = false;
diff --git a/Engine/Plugins/Experimental/AnimNextAnimGraph/Source/AnimNextAnimGraph/Public/GraphInterfaces/AnimNextNativeDataInterface_SynchronizeUsingGroups.h b/Engine/Plugins/Experimental/AnimNextAnimGraph/Source/AnimNextAnimGraph/Public/GraphInterfaces/AnimNextNativeDataInterface_SynchronizeUsingGroups.h
--- a/Engine/Plugins/Experimental/AnimNextAnimGraph/Source/AnimNextAnimGraph/Public/GraphInterfaces/AnimNextNativeDataInterface_SynchronizeUsingGroups.h
+++ b/Engine/Plugins/Experimental/AnimNextAnimGraph/Source/AnimNextAnimGraph/Public/GraphInterfaces/AnimNextNativeDataInterface_SynchronizeUsingGroups.h
@@ -26,6 +26,10 @@ struct FAnimNextNativeDataInterface_SynchronizeUsingGroups : public FAnimNextNat
 	UPROPERTY(EditAnywhere, BlueprintReadWrite, Category = "Group Synchronization", meta = (EnableCategories))
 	EAnimGroupSynchronizationRole GroupRole = EAnimGroupSynchronizationRole::CanBeLeader;
 
+	// The synchronization mode
+	UPROPERTY(EditAnywhere, BlueprintReadWrite, Category = "Group Synchronization", meta = (EnableCategories))
+	EAnimGroupSynchronizationMode SyncMode = EAnimGroupSynchronizationMode::NoSynchronization;
+
 	// Whether or not to match the group sync point when joining as leader
 	// When disabled, the leader will snap the position of the group to its current position
 	UPROPERTY(EditAnywhere, BlueprintReadWrite, Category = "Group Synchronization", meta = (EnableCategories))
diff --git a/Engine/Plugins/Experimental/AnimNextAnimGraph/Source/AnimNextAnimGraph/Public/TraitInterfaces/IGroupSynchronization.h b/Engine/Plugins/Experimental/AnimNextAnimGraph/Source/AnimNextAnimGraph/Public/TraitInterfaces/IGroupSynchronization.h
--- a/Engine/Plugins/Experimental/AnimNextAnimGraph/Source/AnimNextAnimGraph/Public/TraitInterfaces/IGroupSynchronization.h
+++ b/Engine/Plugins/Experimental/AnimNextAnimGraph/Source/AnimNextAnimGraph/Public/TraitInterfaces/IGroupSynchronization.h
@@ -11,28 +11,65 @@
 
 class FMarkerTickContext;
 
+/**
+ * Synchronization Role
+ * 
+ * Timelines part of a named group will synchronize together. To do so, a leader is first identified.
+ * The remaining members of the group will be followers. The synchronization role for a member of
+ * the group determines which role it can play within that group.
+ */
 UENUM(BlueprintType)
 enum class EAnimGroupSynchronizationRole : uint8
 {
-	/** This node can be the leader, as long as it has a higher blend weight than the previous best leader. */
+	/** This timeline host can be the leader, as long as it has a higher blend weight than the previous best leader. */
 	CanBeLeader,
 		
-	/** This node will always be a follower (unless there are only followers, in which case the first one ticked wins). */
+	/** This timeline host will always be a follower (unless there are only followers, in which case the first one ticked wins). */
 	AlwaysFollower,
 
-	/** This node will always be a leader (if more than one node is AlwaysLeader, the last one ticked wins). */
+	/** This timeline host will always be a leader (if more than one timeline host is AlwaysLeader, the last one ticked wins). */
 	AlwaysLeader,
 
-	/** This node will be excluded from the sync group while blending in. Once blended in it will be the sync group leader until blended out*/
+	/** This timeline host will be excluded from the sync group while blending in. Once blended in it will be the sync group leader until blended out*/
 	TransitionLeader,
 
-	/** This node will be excluded from the sync group while blending in. Once blended in it will be a follower until blended out*/
+	/** This timeline host will be excluded from the sync group while blending in. Once blended in it will be a follower until blended out*/
 	TransitionFollower,
 
-	/** This node will always be a leader. If it fails to be ticked as a leader it will be run as ungrouped asset player (EAnimSyncMethod::DoNotSync) .*/
+	/** This timeline host will always be a leader. If it fails to be ticked as a leader it will be run as ungrouped (EAnimGroupSynchronizationMode::NoSynchronization) .*/
 	ExclusiveAlwaysLeader,
 };
 
+/**
+ * Synchronization Mode
+ * 
+ * Timelines part of a named group can be synchronized together. When synchronization is enabled,
+ * two options are presented:
+ *     - A user specified group name can be used
+ *     - A uniquely generated group name can be used
+ * 
+ * Having an explicitly specified group name is useful to synchronize different things together (e.g. two anim sequences).
+ * However, sometimes we want certain things to synchronize only within themselves (e.g. blend space). This is called
+ * self-synchronization. This arises when we want to specify the synchronization behavior at a higher level than individual
+ * timelines. In such cases, we cannot use a user supplied group because every instance would then attempt to use the same
+ * group name, which can be problematic. Consider a blend space embedded within a sub-graph that can blend with itself (e.g. is re-entrant).
+ * This commonly occurs when blend stacks are used (e.g. through motion matching). We would like the anim sequences within
+ * the blend space to self-synchronize but we do not want those to synchronize with other instances of the same blend space.
+ * To solve this, every instance of the blend space needs a unique group name.
+ */
+UENUM(BlueprintType)
+enum class EAnimGroupSynchronizationMode : uint8
+{
+	/** This timeline host will not synchronize. */
+	NoSynchronization,
+
+	/** This timeline host will synchronize and use the provided group name. */
+	SynchronizeUsingGroupName,
+
+	/** This timeline host will synchronize and use a uniquely generated group name. */
+	SynchronizeUsingUniqueGroupName,
+};
+
 namespace UE::AnimNext
 {
 	// Encapsulates the parameters that control sync group behavior
@@ -44,6 +81,9 @@ namespace UE::AnimNext
 		// The role within the group
 		EAnimGroupSynchronizationRole GroupRole = EAnimGroupSynchronizationRole::CanBeLeader;
 
+		// The synchronization mode
+		EAnimGroupSynchronizationMode SyncMode = EAnimGroupSynchronizationMode::NoSynchronization;
+
 		// Whether or not to match the group sync point when joining as leader
 		// When disabled, the leader will snap the position of the group to its current position
 		bool bMatchSyncPoint = true;
