AnimNext: Fix crash where we dispatch null events when skeletal mesh component is not set

diff --git a/Engine/Plugins/Experimental/AnimNextAnimGraph/Source/AnimNextAnimGraph/Private/Traits/NotifyDispatcher.cpp b/Engine/Plugins/Experimental/AnimNextAnimGraph/Source/AnimNextAnimGraph/Private/Traits/NotifyDispatcher.cpp
--- a/Engine/Plugins/Experimental/AnimNextAnimGraph/Source/AnimNextAnimGraph/Private/Traits/NotifyDispatcher.cpp
+++ b/Engine/Plugins/Experimental/AnimNextAnimGraph/Source/AnimNextAnimGraph/Private/Traits/NotifyDispatcher.cpp
@@ -197,36 +197,39 @@ void FAnimNextNotifyDispatcherComponent::OnEndExecution(float InDeltaTime)
 		}
 	}
 
-	// Send end notification to AnimNotifyState not active anymore.
-	for (int32 Index = 0; Index < ActiveAnimNotifyState.Num(); ++Index)
+	if (NotifyQueueDispatchEvent.IsValid())
 	{
-		const FAnimNotifyEvent* AnimNotifyEvent = ActiveAnimNotifyState[Index];
-		const FAnimNotifyEventReference& EventReference = ActiveAnimNotifyEventReference[Index];
-		if (NotifyQueueDispatchEvent.IsValid() && AnimNotifyEvent->NotifyStateClass)
+		// Send end notification to AnimNotifyState not active anymore.
+		for (int32 Index = 0; Index < ActiveAnimNotifyState.Num(); ++Index)
 		{
+			const FAnimNotifyEvent* AnimNotifyEvent = ActiveAnimNotifyState[Index];
+			const FAnimNotifyEventReference& EventReference = ActiveAnimNotifyEventReference[Index];
+			if (AnimNotifyEvent->NotifyStateClass)
+			{
 #if WITH_EDITOR
-			// Prevent firing notifies in animation editors if requested
-			if(ModuleInstance.GetWorldType() != EWorldType::EditorPreview || AnimNotifyEvent->NotifyStateClass->ShouldFireInEditor())
+				// Prevent firing notifies in animation editors if requested
+				if(ModuleInstance.GetWorldType() != EWorldType::EditorPreview || AnimNotifyEvent->NotifyStateClass->ShouldFireInEditor())
 #endif
-			{
-				NotifyQueueDispatchEvent->EventsToEnd.Add(EventReference);
+				{
+					NotifyQueueDispatchEvent->EventsToEnd.Add(EventReference);
+				}
 			}
 		}
-	}
 
-	check(NotifyStateBeginEventReference.Num() == NotifyStateBeginEvent.Num());
-	for (int32 Index = 0; Index < NotifyStateBeginEvent.Num(); Index++)
-	{
-		const FAnimNotifyEvent* AnimNotifyEvent = NotifyStateBeginEvent[Index];
-		const FAnimNotifyEventReference* EventReference = NotifyStateBeginEventReference[Index];
-		if(NotifyQueueDispatchEvent.IsValid() && AnimNotifyEvent->NotifyStateClass)
+		check(NotifyStateBeginEventReference.Num() == NotifyStateBeginEvent.Num());
+		for (int32 Index = 0; Index < NotifyStateBeginEvent.Num(); Index++)
 		{
+			const FAnimNotifyEvent* AnimNotifyEvent = NotifyStateBeginEvent[Index];
+			const FAnimNotifyEventReference* EventReference = NotifyStateBeginEventReference[Index];
+			if(AnimNotifyEvent->NotifyStateClass)
+			{
 #if WITH_EDITOR
-			// Prevent firing notifies in animation editors if requested 
-			if(ModuleInstance.GetWorldType() != EWorldType::EditorPreview || AnimNotifyEvent->NotifyStateClass->ShouldFireInEditor())
+				// Prevent firing notifies in animation editors if requested 
+				if(ModuleInstance.GetWorldType() != EWorldType::EditorPreview || AnimNotifyEvent->NotifyStateClass->ShouldFireInEditor())
 #endif
-			{
-				NotifyQueueDispatchEvent->EventsToBegin.Add(*EventReference);
+				{
+					NotifyQueueDispatchEvent->EventsToBegin.Add(*EventReference);
+				}
 			}
 		}
 	}
@@ -235,24 +238,27 @@ void FAnimNextNotifyDispatcherComponent::OnEndExecution(float InDeltaTime)
 	ActiveAnimNotifyState = MoveTemp(NewActiveAnimNotifyState);
 	ActiveAnimNotifyEventReference = MoveTemp(NewActiveAnimNotifyEventReference);
 
-	// Tick currently active AnimNotifyState
-	for (int32 Index = 0; Index < ActiveAnimNotifyState.Num(); Index++)
+	if (NotifyQueueDispatchEvent.IsValid())
 	{
-		const FAnimNotifyEvent* AnimNotifyEvent = ActiveAnimNotifyState[Index];
-		const FAnimNotifyEventReference& EventReference = ActiveAnimNotifyEventReference[Index];
-		if(NotifyQueueDispatchEvent.IsValid() && AnimNotifyEvent->NotifyStateClass)
+		// Tick currently active AnimNotifyState
+		for (int32 Index = 0; Index < ActiveAnimNotifyState.Num(); Index++)
 		{
+			const FAnimNotifyEvent* AnimNotifyEvent = ActiveAnimNotifyState[Index];
+			const FAnimNotifyEventReference& EventReference = ActiveAnimNotifyEventReference[Index];
+			if(AnimNotifyEvent->NotifyStateClass)
+			{
 #if WITH_EDITOR
-			// Prevent firing notifies in animation editors if requested 
-			if(ModuleInstance.GetWorldType() != EWorldType::EditorPreview || AnimNotifyEvent->NotifyStateClass->ShouldFireInEditor())
+				// Prevent firing notifies in animation editors if requested 
+				if(ModuleInstance.GetWorldType() != EWorldType::EditorPreview || AnimNotifyEvent->NotifyStateClass->ShouldFireInEditor())
 #endif
-			{
-				NotifyQueueDispatchEvent->EventsToTick.Add(EventReference);
+				{
+					NotifyQueueDispatchEvent->EventsToTick.Add(EventReference);
+				}
 			}
 		}
-	}
 
-	GetModuleInstance().QueueOutputTraitEvent(NotifyQueueDispatchEvent);
+		GetModuleInstance().QueueOutputTraitEvent(NotifyQueueDispatchEvent);
+	}
 	NotifyQueue.AnimNotifies.Reset();
 }
 
