AnimNext: Fix for incorrect interface indices being checked when using cached bindings

Fixes injected graphs with module bindings not being able to bind to their module variables
Cached binding indices should always match the indices of the data interface they match. Updated code & comments to reflect this.

diff --git a/Engine/Plugins/Experimental/AnimNext/Source/AnimNext/Public/DataInterface/AnimNextDataInterfaceHost.h b/Engine/Plugins/Experimental/AnimNext/Source/AnimNext/Public/DataInterface/AnimNextDataInterfaceHost.h
--- a/Engine/Plugins/Experimental/AnimNext/Source/AnimNext/Public/DataInterface/AnimNextDataInterfaceHost.h
+++ b/Engine/Plugins/Experimental/AnimNext/Source/AnimNext/Public/DataInterface/AnimNextDataInterfaceHost.h
@@ -30,7 +30,7 @@ private:
 	virtual const UAnimNextDataInterface* GetDataInterface() const = 0;
 
 	// Get the memory for the supplied variable, at the specified index
-	// @param    InVariableIndex    The index into the data interface of the variable
+	// @param    InVariableIndex    The index into the variables of the host
 	// @param    InVariableName     The name of the variable
 	// @param    InVariableProperty The property of the variable
 	virtual uint8* GetMemoryForVariable(int32 InVariableIndex, FName InVariableName, const FProperty* InVariableProperty) const = 0;
diff --git a/Engine/Plugins/Experimental/AnimNextAnimGraph/Source/AnimNextAnimGraph/Private/Graph/AnimNextGraphInstance.cpp b/Engine/Plugins/Experimental/AnimNextAnimGraph/Source/AnimNextAnimGraph/Private/Graph/AnimNextGraphInstance.cpp
--- a/Engine/Plugins/Experimental/AnimNextAnimGraph/Source/AnimNextAnimGraph/Private/Graph/AnimNextGraphInstance.cpp
+++ b/Engine/Plugins/Experimental/AnimNextAnimGraph/Source/AnimNextAnimGraph/Private/Graph/AnimNextGraphInstance.cpp
@@ -288,28 +288,32 @@ bool FAnimNextGraphInstance::BindToCachedBindings()
 			continue;
 		}
 
+		if (CachedBinding->CachedBindings.Num() != ImplementedInterface.NumVariables)
+		{
+			UE_LOGFMT(LogAnimation, Error, "Interface size mismatch: {CachedInterfaceName} {CachedSize} vs {ImplementedInterfaceName} {ImplementedSize}. Interface's values will not be updated at runtime.", CachedBinding->DataInterface->GetFName(), CachedBinding->CachedBindings.Num(), ImplementedInterface.DataInterface->GetFName(), ImplementedInterface.NumVariables);
+			continue;
+		}
+
+		// Index into local property bag
 		int32 VariableIndex = ImplementedInterface.VariableIndex;
-		int32 InterfaceVariableIndex = 0;
 		const int32 EndVariableIndex = ImplementedInterface.VariableIndex + ImplementedInterface.NumVariables;
-		const int32 BindingStartIndex = CachedBindings.IsEmpty() ? 0 : CachedBinding->CachedBindings[0].VariableIndex;
+
+		// Index into interface
+		int32 InterfaceVariableIndex = 0;
+
 		for(; VariableIndex < EndVariableIndex; ++VariableIndex, ++InterfaceVariableIndex)
 		{
-			if (!CachedBinding->CachedBindings.IsValidIndex(InterfaceVariableIndex))
-			{
-				ensureMsgf(false, TEXT("Bindings indexes should always contain variable indexes"));
-				continue;
-			}
-
 			const FPropertyBagPropertyDesc& Desc = PropertyBag->GetPropertyDescs()[VariableIndex];
 			const FCachedDataInterfaceBinding::FVariable& CachedBindingVariable = CachedBinding->CachedBindings[InterfaceVariableIndex];
 			if(CachedBindingVariable.Memory == nullptr)
 			{
-				UE_LOG(LogAnimation, Error, TEXT("FAnimNextGraphInstance::BindToCachedBindings: Did not find cached binding for %s - variable's values will not be updated at runtime.'"), *Desc.Name.ToString());
+				UE_LOGFMT(LogAnimation, Error, "FAnimNextGraphInstance::BindToCachedBindings: Did not find cached binding for {VariableName} - variable's values will not be updated at runtime.", Desc.Name);
 				continue;
 			}
 
+			// Validate that interface layout is the same
 			if( Desc.Name == CachedBindingVariable.VariableName &&
-				VariableIndex == (CachedBindingVariable.VariableIndex - BindingStartIndex) &&
+				InterfaceVariableIndex == CachedBindingVariable.InterfaceVariableIndex &&
 				Desc.CachedProperty->GetClass() == CachedBindingVariable.Property->GetClass())
 			{
 				ExtendedExecuteContext.ExternalVariableRuntimeData[VariableIndex].Memory = CachedBindingVariable.Memory;
@@ -317,7 +321,7 @@ bool FAnimNextGraphInstance::BindToCachedBindings()
 			}
 			else
 			{
-				UE_LOG(LogAnimation, Error, TEXT("FAnimNextGraphInstance::BindToCachedBindings: Did not find cached binding for %s - variable's values will not be updated at runtime.'"), *Desc.Name.ToString());
+				UE_LOGFMT(LogAnimation, Error, "Interface layout mismatch for {GraphName}. Interface's values will not be updated at runtime. (Have:Need): {HaveName}:{NeedName}, {HaveIndex}:{NeedIndex}, {HaveType}:{NeedType}", AnimationGraph->GetFName(), Desc.Name, CachedBindingVariable.VariableName, InterfaceVariableIndex, CachedBindingVariable.InterfaceVariableIndex, Desc.CachedProperty->GetClass()->GetFName(), CachedBindingVariable.Property->GetClass()->GetFName());
 				continue;
 			}
 		}
@@ -355,7 +359,7 @@ void FAnimNextGraphInstance::UpdateCachedBindingsForHost(const UE::AnimNext::IDa
 template<typename HostType>
 void FAnimNextGraphInstance::UpdateCachedBindingsForHostHelper(const HostType& InHost)
 {
-	auto FindOrAddCachedInterfaceBindings = [this](const UAnimNextDataInterface* InDataInterface) -> FCachedDataInterfaceBinding&
+	auto FindOrAddCachedInterfaceBindings = [this](const UAnimNextDataInterface* InDataInterface, int32 InNumVariables) -> FCachedDataInterfaceBinding&
 	{
 		for(FCachedDataInterfaceBinding& CachedBinding : CachedBindings)
 		{
@@ -367,7 +371,7 @@ void FAnimNextGraphInstance::UpdateCachedBindingsForHostHelper(const HostType& I
 
 		FCachedDataInterfaceBinding& NewCachedBinding = CachedBindings.AddDefaulted_GetRef();
 		NewCachedBinding.DataInterface = InDataInterface;
-		NewCachedBinding.CachedBindings.SetNum(InDataInterface->GetPublicVariableDefaults().GetNumPropertiesInBag());
+		NewCachedBinding.CachedBindings.SetNum(InNumVariables);
 		return NewCachedBinding;
 	};
 
@@ -385,20 +389,26 @@ void FAnimNextGraphInstance::UpdateCachedBindingsForHostHelper(const HostType& I
 
 	for(const FAnimNextImplementedDataInterface& ImplementedInterface : HostDataInterface->GetImplementedInterfaces())
 	{
-		FCachedDataInterfaceBinding& CachedBinding = FindOrAddCachedInterfaceBindings(ImplementedInterface.DataInterface);
+		FCachedDataInterfaceBinding& CachedBinding = FindOrAddCachedInterfaceBindings(ImplementedInterface.DataInterface, ImplementedInterface.NumVariables);
+
+		if (CachedBinding.CachedBindings.Num() != ImplementedInterface.NumVariables)
+		{
+			UE_LOGFMT(LogAnimation, Error, "UpdateCachedBindingsForHostHelper: Interface size mismatch: {CachedInterfaceName} {CachedSize} vs {ImplementedInterfaceName} {ImplementedSize}. Interface's values will not be updated at runtime.", CachedBinding.DataInterface->GetFName(), CachedBinding.CachedBindings.Num(), ImplementedInterface.DataInterface->GetFName(), ImplementedInterface.NumVariables);
+			continue;
+		}
 
 		int32 VariableIndex = ImplementedInterface.VariableIndex;
 		const int32 EndVariableIndex = ImplementedInterface.VariableIndex + ImplementedInterface.NumVariables;
-		for(int32 BindingIndex = 0; VariableIndex < EndVariableIndex && BindingIndex < ImplementedInterface.NumVariables; ++VariableIndex, ++BindingIndex)
+		for(int32 InterfaceVariableIndex = 0; VariableIndex < EndVariableIndex && InterfaceVariableIndex < ImplementedInterface.NumVariables; ++VariableIndex, ++InterfaceVariableIndex)
 		{
 			const FPropertyBagPropertyDesc& Desc = PropertyBag->GetPropertyDescs()[VariableIndex];
 			uint8* HostMemory = InHost.GetMemoryForVariable(VariableIndex, Desc.Name, Desc.CachedProperty);
-			uint8* PrevMemory = CachedBinding.CachedBindings.IsValidIndex(BindingIndex) 
-				? CachedBinding.CachedBindings[BindingIndex].Memory 
+			uint8* PrevMemory = CachedBinding.CachedBindings.IsValidIndex(InterfaceVariableIndex) 
+				? CachedBinding.CachedBindings[InterfaceVariableIndex].Memory 
 				: nullptr;
 			HostMemory = HostMemory ? HostMemory : PrevMemory;
 
-			CachedBinding.CachedBindings[BindingIndex] = FCachedDataInterfaceBinding::FVariable(Desc.Name, VariableIndex, Desc.CachedProperty, HostMemory);
+			CachedBinding.CachedBindings[InterfaceVariableIndex] = FCachedDataInterfaceBinding::FVariable(Desc.Name, InterfaceVariableIndex, Desc.CachedProperty, HostMemory);
 		}
 	}
 }
diff --git a/Engine/Plugins/Experimental/AnimNextAnimGraph/Source/AnimNextAnimGraph/Public/Graph/AnimNextGraphInstance.h b/Engine/Plugins/Experimental/AnimNextAnimGraph/Source/AnimNextAnimGraph/Public/Graph/AnimNextGraphInstance.h
--- a/Engine/Plugins/Experimental/AnimNextAnimGraph/Source/AnimNextAnimGraph/Public/Graph/AnimNextGraphInstance.h
+++ b/Engine/Plugins/Experimental/AnimNextAnimGraph/Source/AnimNextAnimGraph/Public/Graph/AnimNextGraphInstance.h
@@ -186,17 +186,24 @@ private:
 		{
 			FVariable() = default;
 
-			FVariable(const FName& InVariableName, int32 InVariableIndex, const FProperty* InProperty, uint8* InMemory)
+			FVariable(const FName& InVariableName, int32 InInterfaceVariableIndex, const FProperty* InProperty, uint8* InMemory)
 				: VariableName(InVariableName)
-				, VariableIndex(InVariableIndex)
+				, InterfaceVariableIndex(InInterfaceVariableIndex)
 				, Property(InProperty)
 				, Memory(InMemory)
 			{
 			}
 
+			// Name of the variable
 			FName VariableName;
-			int32 VariableIndex = INDEX_NONE;
+
+			// Index of the variable in the public variables property bag for DataInterface
+			int32 InterfaceVariableIndex = INDEX_NONE;
+
+			// Cached property of the variable, used for type validation
 			const FProperty* Property = nullptr;
+
+			// Memory of the variable, bound to the 'outermost' container's property bag that implements DataInterface
 			uint8* Memory = nullptr;
 		};
 
