AnimNext: Fix injection of graphs that do not have a ITimeline-implementing stack at their root

This change now allows ITimeline (and other interfaces) to be optionally queried on the 'forwarding children' of a trait stack.
The 'forwarding child' of a stack is a single child. If there are multiple children or zero children we do not consider this a valid situation where forwarding can occur.

diff --git a/Engine/Plugins/Experimental/AnimNextAnimGraph/Source/AnimNextAnimGraph/Internal/Traits/BlendStackTrait.h b/Engine/Plugins/Experimental/AnimNextAnimGraph/Source/AnimNextAnimGraph/Internal/Traits/BlendStackTrait.h
--- a/Engine/Plugins/Experimental/AnimNextAnimGraph/Source/AnimNextAnimGraph/Internal/Traits/BlendStackTrait.h
+++ b/Engine/Plugins/Experimental/AnimNextAnimGraph/Source/AnimNextAnimGraph/Internal/Traits/BlendStackTrait.h
@@ -109,6 +109,9 @@ namespace UE::AnimNext
 			// Weak handle to the graph instance's root
 			FTraitPtr ChildPtr;
 
+			// The length of time this subgraph has been alive
+			float Lifetime;
+
 			FBlendStackCoreTrait::EGraphState State;
 
 			// Whether or not this graph was previously relevant
diff --git a/Engine/Plugins/Experimental/AnimNextAnimGraph/Source/AnimNextAnimGraph/Private/Injection/InjectionSiteTrait.cpp b/Engine/Plugins/Experimental/AnimNextAnimGraph/Source/AnimNextAnimGraph/Private/Injection/InjectionSiteTrait.cpp
--- a/Engine/Plugins/Experimental/AnimNextAnimGraph/Source/AnimNextAnimGraph/Private/Injection/InjectionSiteTrait.cpp
+++ b/Engine/Plugins/Experimental/AnimNextAnimGraph/Source/AnimNextAnimGraph/Private/Injection/InjectionSiteTrait.cpp
@@ -13,6 +13,7 @@
 #include "Module/AnimNextModuleInstance.h"
 #include "TraitInterfaces/IGraphFactory.h"
 #include "TraitInterfaces/IBlendStack.h"
+#include "TraitInterfaces/IHierarchy.h"
 #include "TraitInterfaces/IInertializerBlend.h"
 #include "TraitInterfaces/ISmoothBlend.h"
 
@@ -264,13 +265,13 @@ namespace UE::AnimNext
 			}
 
 			TTraitBinding<ITimeline> TimelineTrait;
-			if(Binding.GetStackInterface(TimelineTrait))
+			if (Binding.GetStackInterface(TimelineTrait))
 			{
 				const FTimelineState ChildState = TimelineTrait.GetState(Context);
 
 				// Compute how much time is left before the timeline ends (can be negative if we overshoot)
-				// Note when looping we dont have a 'time left', so we leave it at FLT_MAX
-				if (ActiveChildIndex == TrackerIndex && !ChildState.IsLooping())
+				// Note when looping/infinite we dont have a 'time left', so we leave it at FLT_MAX
+				if (ActiveChildIndex == TrackerIndex && !(ChildState.IsLooping() || !ChildState.IsFinite()))
 				{
 					const float ChildCurrentPosition = ChildState.GetPosition();
 					const float ChildNextPosition = ChildCurrentPosition + (TraitState.GetDeltaTime() * ChildState.GetPlayRate());
diff --git a/Engine/Plugins/Experimental/AnimNextAnimGraph/Source/AnimNextAnimGraph/Private/TraitInterfaces/IHierarchy.cpp b/Engine/Plugins/Experimental/AnimNextAnimGraph/Source/AnimNextAnimGraph/Private/TraitInterfaces/IHierarchy.cpp
--- a/Engine/Plugins/Experimental/AnimNextAnimGraph/Source/AnimNextAnimGraph/Private/TraitInterfaces/IHierarchy.cpp
+++ b/Engine/Plugins/Experimental/AnimNextAnimGraph/Source/AnimNextAnimGraph/Private/TraitInterfaces/IHierarchy.cpp
@@ -78,4 +78,40 @@ namespace UE::AnimNext
 
 		return NumChildren;
 	}
+
+	FWeakTraitPtr IHierarchy::GetStackForwardingChild(const FExecutionContext& Context, const FTraitStackBinding& Binding)
+	{
+		if (!Binding.IsValid())
+		{
+			return FWeakTraitPtr();
+		}
+
+		TTraitBinding<IHierarchy> HierarchyTrait;
+
+		// Visit the trait stack and query for a forwarding child
+		FWeakTraitPtr ForwardingChild;
+		Binding.GetInterface(HierarchyTrait);
+		while (HierarchyTrait.IsValid())
+		{
+			// The 'forwarding child' is a singular child for the whole stack
+			FChildrenArray Children;
+			HierarchyTrait.GetChildren(Context, Children);
+			if (Children.Num() == 1 && Children[0].IsValid())
+			{
+				if (ForwardingChild.IsValid())
+				{
+					// Already have a forwarding child, we can only have one so fail early
+					return FWeakTraitPtr();
+				}
+				else
+				{
+					ForwardingChild = Children[0];
+				}
+			}
+
+			Binding.GetInterfaceSuper(HierarchyTrait, HierarchyTrait);
+		}
+
+		return ForwardingChild;
+	}
 }
diff --git a/Engine/Plugins/Experimental/AnimNextAnimGraph/Source/AnimNextAnimGraph/Private/Traits/BlendStackTrait.cpp b/Engine/Plugins/Experimental/AnimNextAnimGraph/Source/AnimNextAnimGraph/Private/Traits/BlendStackTrait.cpp
--- a/Engine/Plugins/Experimental/AnimNextAnimGraph/Source/AnimNextAnimGraph/Private/Traits/BlendStackTrait.cpp
+++ b/Engine/Plugins/Experimental/AnimNextAnimGraph/Source/AnimNextAnimGraph/Private/Traits/BlendStackTrait.cpp
@@ -70,6 +70,7 @@ namespace UE::AnimNext
 	void FBlendStackCoreTrait::FGraphState::Initialize(IBlendStack::FGraphRequest&& GraphRequest)
 	{
 		Request = MoveTemp(GraphRequest);
+		Lifetime = 0.0f;
 		State = FBlendStackCoreTrait::EGraphState::Active;
 		bNewlyCreated = true;
 	}
@@ -78,6 +79,7 @@ namespace UE::AnimNext
 	{
 		Instance.Reset();
 		ChildPtr.Reset();
+		Lifetime = 0.0f;
 		bNewlyCreated = false;
 		State = FBlendStackCoreTrait::EGraphState::Inactive;
 	}
@@ -102,19 +104,21 @@ namespace UE::AnimNext
 	void FBlendStackCoreTrait::PreUpdate(FUpdateTraversalContext& Context, const TTraitBinding<IUpdate>& Binding, const FTraitUpdateState& TraitState) const
 	{
 		IUpdate::PreUpdate(Context, Binding, TraitState);
-#if WITH_EDITOR
-		// In editor we may need to rebind our data if an outer hosting module was recompiled
+
 		FInstanceData* InstanceData = Binding.GetInstanceData<FInstanceData>();
 		for (FGraphState& ChildGraph : InstanceData->ChildGraphs)
 		{
+			// Track lifetime for each child graph
+			ChildGraph.Lifetime += TraitState.GetDeltaTime();
+#if WITH_EDITOR
 			if(ChildGraph.State == EGraphState::Active && ChildGraph.Instance.IsValid() && ChildGraph.Instance->RequiresPublicVariableBinding())
 			{
 				FModuleInjectionDataInterfaceAdapter ModuleAdapter(Context.GetRootGraphInstance().GetModuleInstance(), ChildGraph.Request.BindingModuleHandle);
 				IDataInterfaceHost* DataInterfaceHost = &ModuleAdapter;
 				ChildGraph.Instance->BindPublicVariables(ChildGraph.Request.GraphPayload.Get(), TConstArrayView<IDataInterfaceHost*>(&DataInterfaceHost, 1));
 			}
-		}
 #endif
+		}
 	}
 
 	uint32 FBlendStackCoreTrait::GetNumChildren(const FExecutionContext& Context, const TTraitBinding<IHierarchy>& Binding) const
@@ -357,11 +361,16 @@ namespace UE::AnimNext
 			FTraitStackBinding ChildTraitStack;
 			ensure(Context.GetStack(ActiveGraph.ChildPtr, ChildTraitStack));
 
-			UE::AnimNext::TTraitBinding<UE::AnimNext::ITimeline> Timeline;
-			if (ChildTraitStack.GetInterface<UE::AnimNext::ITimeline>(Timeline))
+			TTraitBinding<ITimeline> Timeline;
+			FTraitStackBinding StackBinding;
+			if (IHierarchy::GetForwardedStackInterface<ITimeline>(Context, ChildTraitStack, StackBinding, Timeline))
 			{
 				return Timeline.GetState(Context);
 			}
+			else
+			{
+				return FTimelineState(ActiveGraph.Lifetime, INFINITY, 1.0f, false);
+			}
 		}
 
 		return FTimelineState();
diff --git a/Engine/Plugins/Experimental/AnimNextAnimGraph/Source/AnimNextAnimGraph/Public/TraitInterfaces/IHierarchy.h b/Engine/Plugins/Experimental/AnimNextAnimGraph/Source/AnimNextAnimGraph/Public/TraitInterfaces/IHierarchy.h
--- a/Engine/Plugins/Experimental/AnimNextAnimGraph/Source/AnimNextAnimGraph/Public/TraitInterfaces/IHierarchy.h
+++ b/Engine/Plugins/Experimental/AnimNextAnimGraph/Source/AnimNextAnimGraph/Public/TraitInterfaces/IHierarchy.h
@@ -4,6 +4,7 @@
 
 #include "CoreMinimal.h"
 #include "Misc/MemStack.h"
+#include "TraitCore/ExecutionContext.h"
 #include "TraitCore/ITraitInterface.h"
 #include "TraitCore/TraitBinding.h"
 
@@ -43,6 +44,44 @@ namespace UE::AnimNext
 		// Queries the trait stack of the specified binding and calls GetNumChildren for each trait, accumulating the result.
 		ANIMNEXTANIMGRAPH_API static uint32 GetNumStackChildren(const FExecutionContext& Context, const FTraitBinding& Binding);
 
+		// Queries the trait stack and then its forwarding children for a trait implementing the specified interface.
+		// The trait child found implementing the interface is returned
+		template<class TraitInterface>
+		static bool GetForwardedStackInterface(const FExecutionContext& Context, const FTraitStackBinding& StackBinding, FTraitStackBinding& OutStackBinding, TTraitBinding<TraitInterface>& OutBinding)
+		{
+			// Start with the supplied stack
+			OutStackBinding = StackBinding;
+			while (true)
+			{
+				if (OutStackBinding.GetInterface<TraitInterface>(OutBinding))
+				{
+					return true;
+				}
+
+				// Forward to child stack, if any
+				FWeakTraitPtr ForwardingChild = GetStackForwardingChild(Context, OutStackBinding);
+				if (!ForwardingChild.IsValid())
+				{
+					return false;
+				}
+				ensure(Context.GetStack(ForwardingChild, OutStackBinding));
+			}
+		}
+
+		// Queries the trait's forwarding children for a trait implementing the specified interface.
+		// The first child found implementing the interface is returned
+		template<class TraitInterface>
+		static bool GetForwardedStackInterface(const FExecutionContext& Context, const FTraitBinding& Binding, FTraitStackBinding& OutStackBinding, TTraitBinding<TraitInterface>& OutBinding)
+		{
+			return Binding.IsValid() ? GetForwardedStackInterface(Context, *Binding.GetStack(), OutStackBinding, OutBinding) : false;
+		}
+
+	private:
+		// Queries the trait stack of the specified binding and calls GetForwardingChild for each trait.
+		// If there is exactly one child, it is returned.
+		ANIMNEXTANIMGRAPH_API static FWeakTraitPtr GetStackForwardingChild(const FExecutionContext& Context, const FTraitStackBinding& Binding);
+
+	public:
 #if WITH_EDITOR
 		ANIMNEXTANIMGRAPH_API virtual const FText& GetDisplayName() const override;
 		ANIMNEXTANIMGRAPH_API virtual const FText& GetDisplayShortName() const override;
@@ -66,7 +105,7 @@ namespace UE::AnimNext
 		{
 			GetInterface()->GetChildren(Context, *this, Children);
 		}
-
+		
 	protected:
 		const IHierarchy* GetInterface() const { return GetInterfaceTyped<IHierarchy>(); }
 	};
diff --git a/Engine/Plugins/Experimental/AnimNextAnimGraph/Source/AnimNextAnimGraph/Public/TraitInterfaces/ITimeline.h b/Engine/Plugins/Experimental/AnimNextAnimGraph/Source/AnimNextAnimGraph/Public/TraitInterfaces/ITimeline.h
--- a/Engine/Plugins/Experimental/AnimNextAnimGraph/Source/AnimNextAnimGraph/Public/TraitInterfaces/ITimeline.h
+++ b/Engine/Plugins/Experimental/AnimNextAnimGraph/Source/AnimNextAnimGraph/Public/TraitInterfaces/ITimeline.h
@@ -72,12 +72,29 @@ namespace UE::AnimNext
 		float GetPosition() const { return Position; }
 
 		// Returns the timeline position as a ratio (0.0 = start of timeline, 1.0 = end of timeline)
-		float GetPositionRatio() const { return Duration != 0.0f ? FMath::Clamp(Position / Duration, 0.0f, 1.0f) : 0.0f; }
+		// For infinite timelines, this will return infinity
+		float GetPositionRatio() const
+		{
+			if (IsFinite())
+			{
+				return Duration != 0.0f ? FMath::Clamp(Position / Duration, 0.0f, 1.0f) : 0.0f;
+			}
+			return INFINITY;
+		}
 
 		// Returns the time left to play in the timeline
-		float GetTimeLeft() const { return PlayRate >= 0.0f ? (Duration - Position) : Position; }
+		// For infinite timelines this will return infinity
+		float GetTimeLeft() const
+		{
+			if (IsFinite())
+			{
+				return PlayRate >= 0.0f ? (Duration - Position) : Position;
+			}
+			return INFINITY;
+		}
 
 		// Returns the timeline duration in seconds
+		// For infinite timelines this will return infinity
 		float GetDuration() const { return Duration; }
 
 		// Returns the play rate of the timeline
@@ -86,6 +103,9 @@ namespace UE::AnimNext
 		// Returns whether or not the timeline is looping
 		bool IsLooping() const { return bIsLooping; }
 
+		// Returns whether or not this timeline is finite
+		bool IsFinite() const { return Duration != INFINITY; }
+
 		// Creates a new instance of the timeline state with the supplied position
 		[[nodiscard]] FTimelineState WithPosition(float InPosition) const
 		{
