Allow AnimNext timeline implementation to be read only or read write

diff --git a/Engine/Plugins/Experimental/AnimNextAnimGraph/Source/AnimNextAnimGraph/Internal/Traits/BlendStackTrait.h b/Engine/Plugins/Experimental/AnimNextAnimGraph/Source/AnimNextAnimGraph/Internal/Traits/BlendStackTrait.h
--- a/Engine/Plugins/Experimental/AnimNextAnimGraph/Source/AnimNextAnimGraph/Internal/Traits/BlendStackTrait.h
+++ b/Engine/Plugins/Experimental/AnimNextAnimGraph/Source/AnimNextAnimGraph/Internal/Traits/BlendStackTrait.h
@@ -168,7 +168,7 @@ namespace UE::AnimNext
 		virtual FOnExtractRootMotionAttribute GetOnExtractRootMotionAttribute(FExecutionContext& Context, const TTraitBinding<IAttributeProvider>& Binding) const override;
 
 		// ITimeline impl
-		virtual FTimelineState GetState(FExecutionContext& Context, const TTraitBinding<ITimeline>& Binding) const override;
+		virtual FTimelineState GetState(const FExecutionContext& Context, const TTraitBinding<ITimeline>& Binding) const override;
 
 		// IGarbageCollection impl
 		virtual void AddReferencedObjects(const FExecutionContext& Context, const TTraitBinding<IGarbageCollection>& Binding, FReferenceCollector& Collector) const override;
diff --git a/Engine/Plugins/Experimental/AnimNextAnimGraph/Source/AnimNextAnimGraph/Private/TraitCore/TraitStackBinding.cpp b/Engine/Plugins/Experimental/AnimNextAnimGraph/Source/AnimNextAnimGraph/Private/TraitCore/TraitStackBinding.cpp
--- a/Engine/Plugins/Experimental/AnimNextAnimGraph/Source/AnimNextAnimGraph/Private/TraitCore/TraitStackBinding.cpp
+++ b/Engine/Plugins/Experimental/AnimNextAnimGraph/Source/AnimNextAnimGraph/Private/TraitCore/TraitStackBinding.cpp
@@ -286,6 +286,72 @@ namespace UE::AnimNext
 		return false;
 	}
 
+	bool FTraitStackBinding::HasInterfaceImpl(FTraitInterfaceUID InterfaceUID) const
+	{
+		if (!IsValid())
+		{
+			return false;
+		}
+
+		// Start searching with the top trait towards our base trait
+		const FTraitTemplate* TraitDescs = NodeTemplate->GetTraits();
+		const FTraitTemplate* StartDesc = TraitDescs + TopTraitIndex;
+		const FTraitTemplate* EndDesc = TraitDescs + BaseTraitIndex - 1;
+
+		for (const FTraitTemplate* TraitDesc = StartDesc; TraitDesc != EndDesc; --TraitDesc)
+		{
+			const FTrait* Trait = Context->GetTrait(*TraitDesc);
+			if (Trait == nullptr)
+			{
+				continue;	// Trait hasn't been loaded or registered, skip it
+			}
+
+			if (Trait->GetTraitInterface(InterfaceUID) != nullptr)
+			{
+				return true;
+			}
+		}
+
+		// We failed to find the specified interface on the trait stack
+		return false;
+	}
+
+	bool FTraitStackBinding::HasInterfaceSuperImpl(FTraitInterfaceUID InterfaceUID, const FTraitBinding& Binding) const
+	{
+		if (!IsValid())
+		{
+			return false;
+		}
+
+		const uint32 CurrentTraitIndex = Binding.TraitIndex;
+		if (CurrentTraitIndex == BaseTraitIndex)
+		{
+			return false;	// We reached the base of the stack, we don't have a super
+		}
+
+		// Start searching with the current trait towards our base trait
+		const FTraitTemplate* TraitDescs = NodeTemplate->GetTraits();
+		const FTraitTemplate* StartDesc = TraitDescs + CurrentTraitIndex - 1;
+		const FTraitTemplate* EndDesc = TraitDescs + BaseTraitIndex - 1;
+
+		for (const FTraitTemplate* TraitDesc = StartDesc; TraitDesc != EndDesc; --TraitDesc)
+		{
+			const FTrait* Trait = Context->GetTrait(*TraitDesc);
+			if (Trait == nullptr)
+			{
+				continue;	// Trait hasn't been loaded or registered, skip it
+			}
+
+			if (Trait->GetTraitInterface(InterfaceUID) != nullptr)
+			{
+				return true;
+			}
+		}
+
+		// We failed to find the specified interface on the trait stack
+		return false;
+	}
+
 	void FTraitStackBinding::SnapshotLatentProperties(bool bIsFrozen)
 	{
 		if (!IsValid())
diff --git a/Engine/Plugins/Experimental/AnimNextAnimGraph/Source/AnimNextAnimGraph/Private/TraitInterfaces/ITimeline.cpp b/Engine/Plugins/Experimental/AnimNextAnimGraph/Source/AnimNextAnimGraph/Private/TraitInterfaces/ITimeline.cpp
--- a/Engine/Plugins/Experimental/AnimNextAnimGraph/Source/AnimNextAnimGraph/Private/TraitInterfaces/ITimeline.cpp
+++ b/Engine/Plugins/Experimental/AnimNextAnimGraph/Source/AnimNextAnimGraph/Private/TraitInterfaces/ITimeline.cpp
@@ -21,7 +21,7 @@ namespace UE::AnimNext
 	}
 #endif // WITH_EDITOR
 
-	void ITimeline::GetSyncMarkers(FExecutionContext& Context, const TTraitBinding<ITimeline>& Binding, FTimelineSyncMarkerArray& OutSyncMarkers) const
+	void ITimeline::GetSyncMarkers(const FExecutionContext& Context, const TTraitBinding<ITimeline>& Binding, FTimelineSyncMarkerArray& OutSyncMarkers) const
 	{
 		TTraitBinding<ITimeline> SuperBinding;
 		if (Binding.GetStackInterfaceSuper(SuperBinding))
@@ -30,7 +30,7 @@ namespace UE::AnimNext
 		}
 	}
 
-	FTimelineState ITimeline::GetState(FExecutionContext& Context, const TTraitBinding<ITimeline>& Binding) const
+	FTimelineState ITimeline::GetState(const FExecutionContext& Context, const TTraitBinding<ITimeline>& Binding) const
 	{
 		TTraitBinding<ITimeline> SuperBinding;
 		if (Binding.GetStackInterfaceSuper(SuperBinding))
@@ -41,16 +41,7 @@ namespace UE::AnimNext
 		return FTimelineState();
 	}
 
-	void ITimeline::AdvanceBy(FExecutionContext& Context, const TTraitBinding<ITimeline>& Binding, float DeltaTime) const
-	{
-		TTraitBinding<ITimeline> SuperBinding;
-		if (Binding.GetStackInterfaceSuper(SuperBinding))
-		{
-			SuperBinding.AdvanceBy(Context, DeltaTime);
-		}
-	}
-
-	FTimelineDelta ITimeline::GetDelta(FExecutionContext& Context, const TTraitBinding<ITimeline>& Binding) const
+	FTimelineDelta ITimeline::GetDelta(const FExecutionContext& Context, const TTraitBinding<ITimeline>& Binding) const
 	{
 		TTraitBinding<ITimeline> SuperBinding;
 		if (Binding.GetStackInterfaceSuper(SuperBinding))
diff --git a/Engine/Plugins/Experimental/AnimNextAnimGraph/Source/AnimNextAnimGraph/Private/TraitInterfaces/ITimelinePlayer.cpp b/Engine/Plugins/Experimental/AnimNextAnimGraph/Source/AnimNextAnimGraph/Private/TraitInterfaces/ITimelinePlayer.cpp
new file mode 100644
--- /dev/null
+++ b/Engine/Plugins/Experimental/AnimNextAnimGraph/Source/AnimNextAnimGraph/Private/TraitInterfaces/ITimelinePlayer.cpp
@@ -0,0 +1,32 @@
+// Copyright Epic Games, Inc. All Rights Reserved.
+
+#include "TraitInterfaces/ITimelinePlayer.h"
+
+#include "TraitCore/ExecutionContext.h"
+
+namespace UE::AnimNext
+{
+	AUTO_REGISTER_ANIM_TRAIT_INTERFACE(ITimelinePlayer)
+
+#if WITH_EDITOR
+	const FText& ITimelinePlayer::GetDisplayName() const
+	{
+		static FText InterfaceName = NSLOCTEXT("TraitInterfaces", "TraitInterface_ITimelinePlayer_Name", "Timeline Player");
+		return InterfaceName;
+	}
+	const FText& ITimelinePlayer::GetDisplayShortName() const
+	{
+		static FText InterfaceShortName = NSLOCTEXT("TraitInterfaces", "TraitInterface_ITimelinePlayer_ShortName", "TIM");
+		return InterfaceShortName;
+	}
+#endif // WITH_EDITOR
+
+	void ITimelinePlayer::AdvanceBy(FExecutionContext& Context, const TTraitBinding<ITimelinePlayer>& Binding, float DeltaTime) const
+	{
+		TTraitBinding<ITimelinePlayer> SuperBinding;
+		if (Binding.GetStackInterfaceSuper(SuperBinding))
+		{
+			SuperBinding.AdvanceBy(Context, DeltaTime);
+		}
+	}
+}
diff --git a/Engine/Plugins/Experimental/AnimNextAnimGraph/Source/AnimNextAnimGraph/Private/Traits/BlendSpacePlayer.cpp b/Engine/Plugins/Experimental/AnimNextAnimGraph/Source/AnimNextAnimGraph/Private/Traits/BlendSpacePlayer.cpp
--- a/Engine/Plugins/Experimental/AnimNextAnimGraph/Source/AnimNextAnimGraph/Private/Traits/BlendSpacePlayer.cpp
+++ b/Engine/Plugins/Experimental/AnimNextAnimGraph/Source/AnimNextAnimGraph/Private/Traits/BlendSpacePlayer.cpp
@@ -262,7 +262,7 @@ namespace UE::AnimNext
 		}
 	}
 
-	FTimelineState FBlendSpacePlayerTrait::GetState(FExecutionContext& Context, const TTraitBinding<ITimeline>& Binding) const
+	FTimelineState FBlendSpacePlayerTrait::GetState(const FExecutionContext& Context, const TTraitBinding<ITimeline>& Binding) const
 	{
 		const FInstanceData* InstanceData = Binding.GetInstanceData<FInstanceData>();
 
diff --git a/Engine/Plugins/Experimental/AnimNextAnimGraph/Source/AnimNextAnimGraph/Private/Traits/BlendSpacePlayer.h b/Engine/Plugins/Experimental/AnimNextAnimGraph/Source/AnimNextAnimGraph/Private/Traits/BlendSpacePlayer.h
--- a/Engine/Plugins/Experimental/AnimNextAnimGraph/Source/AnimNextAnimGraph/Private/Traits/BlendSpacePlayer.h
+++ b/Engine/Plugins/Experimental/AnimNextAnimGraph/Source/AnimNextAnimGraph/Private/Traits/BlendSpacePlayer.h
@@ -132,7 +132,7 @@ namespace UE::AnimNext
 		virtual void QueueChildrenForTraversal(FUpdateTraversalContext& Context, const TTraitBinding<IUpdateTraversal>& Binding, const FTraitUpdateState& TraitState, FUpdateTraversalQueue& TraversalQueue) const override;
 
 		// ITimeline impl
-		virtual FTimelineState GetState(FExecutionContext& Context, const TTraitBinding<ITimeline>& Binding) const override;
+		virtual FTimelineState GetState(const FExecutionContext& Context, const TTraitBinding<ITimeline>& Binding) const override;
 
 		// IGarbageCollection impl
 		virtual void AddReferencedObjects(const FExecutionContext& Context, const TTraitBinding<IGarbageCollection>& Binding, FReferenceCollector& Collector) const override;
diff --git a/Engine/Plugins/Experimental/AnimNextAnimGraph/Source/AnimNextAnimGraph/Private/Traits/BlendStackTrait.cpp b/Engine/Plugins/Experimental/AnimNextAnimGraph/Source/AnimNextAnimGraph/Private/Traits/BlendStackTrait.cpp
--- a/Engine/Plugins/Experimental/AnimNextAnimGraph/Source/AnimNextAnimGraph/Private/Traits/BlendStackTrait.cpp
+++ b/Engine/Plugins/Experimental/AnimNextAnimGraph/Source/AnimNextAnimGraph/Private/Traits/BlendStackTrait.cpp
@@ -350,7 +350,7 @@ namespace UE::AnimNext
 		return FOnExtractRootMotionAttribute();
 	}
 
-	FTimelineState FBlendStackCoreTrait::GetState(FExecutionContext& Context, const TTraitBinding<ITimeline>& Binding) const
+	FTimelineState FBlendStackCoreTrait::GetState(const FExecutionContext& Context, const TTraitBinding<ITimeline>& Binding) const
 	{
 		const FInstanceData* InstanceData = Binding.GetInstanceData<FInstanceData>();
 		const int32 CurrentlyActiveGraphIndex = InstanceData->CurrentlyActiveGraphIndex;
diff --git a/Engine/Plugins/Experimental/AnimNextAnimGraph/Source/AnimNextAnimGraph/Private/Traits/NotifyDispatcher.cpp b/Engine/Plugins/Experimental/AnimNextAnimGraph/Source/AnimNextAnimGraph/Private/Traits/NotifyDispatcher.cpp
--- a/Engine/Plugins/Experimental/AnimNextAnimGraph/Source/AnimNextAnimGraph/Private/Traits/NotifyDispatcher.cpp
+++ b/Engine/Plugins/Experimental/AnimNextAnimGraph/Source/AnimNextAnimGraph/Private/Traits/NotifyDispatcher.cpp
@@ -1,7 +1,8 @@
-﻿// Copyright Epic Games, Inc. All Rights Reserved.
+// Copyright Epic Games, Inc. All Rights Reserved.
 
 #include "Traits/NotifyDispatcher.h"
 #include "TraitInterfaces/INotifySource.h"
+#include "TraitInterfaces/ITimeline.h"
 #include "AnimationRuntime.h"
 #include "Animation/AnimTrace.h"
 #include "Animation/AnimNotifies/AnimNotifyState.h"
@@ -16,10 +17,11 @@ namespace UE::AnimNext
 
 	// Trait implementation boilerplate
 	#define TRAIT_INTERFACE_ENUMERATOR(GeneratorMacro) \
-		GeneratorMacro(ITimeline) \
+		GeneratorMacro(ITimelinePlayer) \
 
 		// Trait required interfaces implementation boilerplate
 	#define TRAIT_REQUIRED_INTERFACE_ENUMERATOR(GeneratorMacroRequired) \
+		GeneratorMacroRequired(ITimelinePlayer) \
 		GeneratorMacroRequired(ITimeline) \
 		GeneratorMacroRequired(INotifySource) \
 
@@ -27,17 +29,20 @@ namespace UE::AnimNext
 	#undef TRAIT_REQUIRED_INTERFACE_ENUMERATOR
 	#undef TRAIT_INTERFACE_ENUMERATOR
 
-	void FNotifyDispatcherTrait::AdvanceBy(FExecutionContext& Context, const TTraitBinding<ITimeline>& Binding, float DeltaTime) const
+	void FNotifyDispatcherTrait::AdvanceBy(FExecutionContext& Context, const TTraitBinding<ITimelinePlayer>& Binding, float DeltaTime) const
 	{
 		TTraitBinding<ITimeline> TimelineTrait;
 		ensure(Binding.GetStackInterfaceSuper(TimelineTrait));
 
+		TTraitBinding<ITimelinePlayer> TimelinePlayerTrait;
+		ensure(Binding.GetStackInterfaceSuper(TimelinePlayerTrait));
+
 		TTraitBinding<INotifySource> NotifySourceTrait;
 		ensure(Binding.GetStackInterfaceSuper(NotifySourceTrait));
 
 		// Get current state from the stack, advance time, then get the delta state
 		const FTimelineState PreAdvanceState = TimelineTrait.GetState(Context);
-		TimelineTrait.AdvanceBy(Context, DeltaTime);
+		TimelinePlayerTrait.AdvanceBy(Context, DeltaTime);
 		const FTimelineDelta Delta = TimelineTrait.GetDelta(Context);
 
 		// Query for notifies
diff --git a/Engine/Plugins/Experimental/AnimNextAnimGraph/Source/AnimNextAnimGraph/Private/Traits/NotifyDispatcher.h b/Engine/Plugins/Experimental/AnimNextAnimGraph/Source/AnimNextAnimGraph/Private/Traits/NotifyDispatcher.h
--- a/Engine/Plugins/Experimental/AnimNextAnimGraph/Source/AnimNextAnimGraph/Private/Traits/NotifyDispatcher.h
+++ b/Engine/Plugins/Experimental/AnimNextAnimGraph/Source/AnimNextAnimGraph/Private/Traits/NotifyDispatcher.h
@@ -5,7 +5,7 @@
 #include "CoreMinimal.h"
 #include "Module/ModuleEvents.h"
 #include "TraitCore/Trait.h"
-#include "TraitInterfaces/ITimeline.h"
+#include "TraitInterfaces/ITimelinePlayer.h"
 #include "TraitInterfaces/IUpdate.h"
 #include "Animation/AnimNotifyQueue.h"
 #include "Module/AnimNextModuleInstanceComponent.h"
@@ -68,7 +68,7 @@ namespace UE::AnimNext
 	 * 
 	 * A trait that dispatches notifies according to a timeline advancing
 	 */
-	struct FNotifyDispatcherTrait : FAdditiveTrait, ITimeline
+	struct FNotifyDispatcherTrait : FAdditiveTrait, ITimelinePlayer
 	{
 		DECLARE_ANIM_TRAIT(FNotifyDispatcherTrait, FAdditiveTrait)
 
@@ -78,8 +78,8 @@ namespace UE::AnimNext
 		{
 		};
 
-		// ITimeline impl
-		virtual void AdvanceBy(FExecutionContext& Context, const TTraitBinding<ITimeline>& Binding, float DeltaTime) const override;
+		// ITimelinePlayer impl
+		virtual void AdvanceBy(FExecutionContext& Context, const TTraitBinding<ITimelinePlayer>& Binding, float DeltaTime) const override;
 	};
 
 	// Event that dispatches to the module
diff --git a/Engine/Plugins/Experimental/AnimNextAnimGraph/Source/AnimNextAnimGraph/Private/Traits/SequencePlayer.cpp b/Engine/Plugins/Experimental/AnimNextAnimGraph/Source/AnimNextAnimGraph/Private/Traits/SequencePlayer.cpp
--- a/Engine/Plugins/Experimental/AnimNextAnimGraph/Source/AnimNextAnimGraph/Private/Traits/SequencePlayer.cpp
+++ b/Engine/Plugins/Experimental/AnimNextAnimGraph/Source/AnimNextAnimGraph/Private/Traits/SequencePlayer.cpp
@@ -16,6 +16,7 @@ namespace UE::AnimNext
 		GeneratorMacro(IEvaluate) \
 		GeneratorMacro(IAttributeProvider) \
 		GeneratorMacro(ITimeline) \
+		GeneratorMacro(ITimelinePlayer) \
 		GeneratorMacro(IUpdate) \
 		GeneratorMacro(IGarbageCollection) \
 		GeneratorMacro(INotifySource) \
@@ -74,7 +75,7 @@ namespace UE::AnimNext
 		return FOnExtractRootMotionAttribute();
 	}
 
-	void FSequencePlayerTrait::GetSyncMarkers(FExecutionContext& Context, const TTraitBinding<ITimeline>& Binding, FTimelineSyncMarkerArray& OutSyncMarkers) const
+	void FSequencePlayerTrait::GetSyncMarkers(const FExecutionContext& Context, const TTraitBinding<ITimeline>& Binding, FTimelineSyncMarkerArray& OutSyncMarkers) const
 	{
 		const FInstanceData* InstanceData = Binding.GetInstanceData<FInstanceData>();
 		if (const UAnimSequence* AnimSeq = InstanceData->AnimSequence.Get())
@@ -88,7 +89,7 @@ namespace UE::AnimNext
 		}
 	}
 
-	FTimelineState FSequencePlayerTrait::GetState(FExecutionContext& Context, const TTraitBinding<ITimeline>& Binding) const
+	FTimelineState FSequencePlayerTrait::GetState(const FExecutionContext& Context, const TTraitBinding<ITimeline>& Binding) const
 	{
 		const FInstanceData* InstanceData = Binding.GetInstanceData<FInstanceData>();
 		if (const UAnimSequence* AnimSeq = InstanceData->AnimSequence.Get())
@@ -105,7 +106,7 @@ namespace UE::AnimNext
 		return FTimelineState();
 	}
 
-	FTimelineDelta FSequencePlayerTrait::GetDelta(FExecutionContext& Context, const TTraitBinding<ITimeline>& Binding) const
+	FTimelineDelta FSequencePlayerTrait::GetDelta(const FExecutionContext& Context, const TTraitBinding<ITimeline>& Binding) const
 	{
 		const FInstanceData* InstanceData = Binding.GetInstanceData<FInstanceData>();
 		if (const UAnimSequence* AnimSeq = InstanceData->AnimSequence.Get())
@@ -118,7 +119,7 @@ namespace UE::AnimNext
 		return FTimelineDelta();
 	}
 
-	void FSequencePlayerTrait::AdvanceBy(FExecutionContext& Context, const TTraitBinding<ITimeline>& Binding, float DeltaTime) const
+	void FSequencePlayerTrait::AdvanceBy(FExecutionContext& Context, const TTraitBinding<ITimelinePlayer>& Binding, float DeltaTime) const
 	{
 		FInstanceData* InstanceData = Binding.GetInstanceData<FInstanceData>();
 
@@ -165,10 +166,10 @@ namespace UE::AnimNext
 	void FSequencePlayerTrait::PreUpdate(FUpdateTraversalContext& Context, const TTraitBinding<IUpdate>& Binding, const FTraitUpdateState& TraitState) const
 	{
 		// We just advance the timeline
-		TTraitBinding<ITimeline> TimelineTrait;
-		Binding.GetStackInterface(TimelineTrait);
+		TTraitBinding<ITimelinePlayer> TimelinePlayerTrait;
+		Binding.GetStackInterface(TimelinePlayerTrait);
 
-		TimelineTrait.AdvanceBy(Context, TraitState.GetDeltaTime());
+		TimelinePlayerTrait.AdvanceBy(Context, TraitState.GetDeltaTime());
 	}
 
 	void FSequencePlayerTrait::AddReferencedObjects(const FExecutionContext& Context, const TTraitBinding<IGarbageCollection>& Binding, FReferenceCollector& Collector) const
diff --git a/Engine/Plugins/Experimental/AnimNextAnimGraph/Source/AnimNextAnimGraph/Private/Traits/SequencePlayer.h b/Engine/Plugins/Experimental/AnimNextAnimGraph/Source/AnimNextAnimGraph/Private/Traits/SequencePlayer.h
--- a/Engine/Plugins/Experimental/AnimNextAnimGraph/Source/AnimNextAnimGraph/Private/Traits/SequencePlayer.h
+++ b/Engine/Plugins/Experimental/AnimNextAnimGraph/Source/AnimNextAnimGraph/Private/Traits/SequencePlayer.h
@@ -8,6 +8,7 @@
 #include "TraitInterfaces/IGarbageCollection.h"
 #include "TraitInterfaces/IAttributeProvider.h"
 #include "TraitInterfaces/ITimeline.h"
+#include "TraitInterfaces/ITimelinePlayer.h"
 #include "TraitInterfaces/IUpdate.h"
 #include "Animation/AnimSequence.h"
 #include "TraitInterfaces/INotifySource.h"
@@ -53,7 +54,7 @@ namespace UE::AnimNext
 	 * 
 	 * A trait that can play an animation sequence.
 	 */
-	struct FSequencePlayerTrait : FBaseTrait, IEvaluate, IAttributeProvider, ITimeline, IUpdate, IGarbageCollection, INotifySource
+	struct FSequencePlayerTrait : FBaseTrait, IEvaluate, IAttributeProvider, ITimeline, ITimelinePlayer, IUpdate, IGarbageCollection, INotifySource
 	{
 		DECLARE_ANIM_TRAIT(FSequencePlayerTrait, FBaseTrait)
 
@@ -84,10 +85,12 @@ namespace UE::AnimNext
 		virtual FOnExtractRootMotionAttribute GetOnExtractRootMotionAttribute(FExecutionContext& Context, const TTraitBinding<IAttributeProvider>& Binding) const override;
 
 		// ITimeline impl
-		virtual void GetSyncMarkers(FExecutionContext& Context, const TTraitBinding<ITimeline>& Binding, FTimelineSyncMarkerArray& OutSyncMarkers) const override;
-		virtual FTimelineState GetState(FExecutionContext& Context, const TTraitBinding<ITimeline>& Binding) const override;
-		virtual void AdvanceBy(FExecutionContext& Context, const TTraitBinding<ITimeline>& Binding, float DeltaTime) const override;
-		virtual FTimelineDelta GetDelta(FExecutionContext& Context, const TTraitBinding<ITimeline>& Binding) const override;
+		virtual void GetSyncMarkers(const FExecutionContext& Context, const TTraitBinding<ITimeline>& Binding, FTimelineSyncMarkerArray& OutSyncMarkers) const override;
+		virtual FTimelineState GetState(const FExecutionContext& Context, const TTraitBinding<ITimeline>& Binding) const override;
+		virtual FTimelineDelta GetDelta(const FExecutionContext& Context, const TTraitBinding<ITimeline>& Binding) const override;
+
+		// ITimelinePlayer impl
+		virtual void AdvanceBy(FExecutionContext& Context, const TTraitBinding<ITimelinePlayer>& Binding, float DeltaTime) const override;
 
 		// IUpdate impl
 		virtual void OnBecomeRelevant(FUpdateTraversalContext& Context, const TTraitBinding<IUpdate>& Binding, const FTraitUpdateState& TraitState) const override;
diff --git a/Engine/Plugins/Experimental/AnimNextAnimGraph/Source/AnimNextAnimGraph/Private/Traits/SynchronizeUsingGroups.cpp b/Engine/Plugins/Experimental/AnimNextAnimGraph/Source/AnimNextAnimGraph/Private/Traits/SynchronizeUsingGroups.cpp
--- a/Engine/Plugins/Experimental/AnimNextAnimGraph/Source/AnimNextAnimGraph/Private/Traits/SynchronizeUsingGroups.cpp
+++ b/Engine/Plugins/Experimental/AnimNextAnimGraph/Source/AnimNextAnimGraph/Private/Traits/SynchronizeUsingGroups.cpp
@@ -14,7 +14,7 @@ namespace UE::AnimNext
 	#define TRAIT_INTERFACE_ENUMERATOR(GeneratorMacro) \
 		GeneratorMacro(IGraphFactory) \
 		GeneratorMacro(IGroupSynchronization) \
-		GeneratorMacro(ITimeline) \
+		GeneratorMacro(ITimelinePlayer) \
 		GeneratorMacro(IUpdate) \
 
 	GENERATE_ANIM_TRAIT_IMPLEMENTATION(FSynchronizeUsingGroupsTrait, TRAIT_INTERFACE_ENUMERATOR, NULL_ANIM_TRAIT_INTERFACE_ENUMERATOR, NULL_ANIM_TRAIT_EVENT_ENUMERATOR)
@@ -24,8 +24,11 @@ namespace UE::AnimNext
 	{
 		SyncGroupComponent = &Context.GetComponent<FSyncGroupGraphInstanceComponent>();
 
-		TTraitBinding<ITimeline> TimelineTrait;
-		bHasTimeline = Binding.GetStackInterface(TimelineTrait);
+		// If our timeline can't advance, we can't synchronize; assume we have no timeline
+		// We only check the base trait since it is the one that would own the timeline
+		FTraitBinding BaseTrait;
+		ensure(Binding.GetStackBaseTrait(BaseTrait));
+		bHasTimelinePlayer = BaseTrait.HasInterface<ITimelinePlayer>();
 	}
 
 	void FSynchronizeUsingGroupsTrait::FInstanceData::Destruct(const FExecutionContext& Context, const FTraitBinding& Binding)
@@ -49,9 +52,9 @@ namespace UE::AnimNext
 	{
 		FInstanceData* InstanceData = Binding.GetInstanceData<FInstanceData>();
 
-		if (!InstanceData->bHasTimeline)
+		if (!InstanceData->bHasTimelinePlayer)
 		{
-			// If we don't have a timeline on the trait stack, there is nothing to synchronize
+			// If we don't have a timeline player beneath us on the trait stack, there is nothing to synchronize
 			IUpdate::PreUpdate(Context, Binding, TraitState);
 			return;
 		}
@@ -130,23 +133,23 @@ namespace UE::AnimNext
 	{
 		FInstanceData* InstanceData = Binding.GetInstanceData<FInstanceData>();
 
-		if (!InstanceData->bHasTimeline)
+		if (!InstanceData->bHasTimelinePlayer)
 		{
-			return;	//  No timeline on the trait stack, nothing to synchronize
+			return;	//  No timeline player beneath us on the trait stack, nothing to synchronize
 		}
 
 		// When the group advances the timeline, we thaw it to advance
 		InstanceData->bFreezeTimeline = false;
 
-		TTraitBinding<ITimeline> TimelineTrait;
-		Binding.GetStackInterface(TimelineTrait);
+		TTraitBinding<ITimelinePlayer> TimelinePlayerTrait;
+		Binding.GetStackInterface(TimelinePlayerTrait);
 
-		TimelineTrait.AdvanceBy(Context, DeltaTime);
+		TimelinePlayerTrait.AdvanceBy(Context, DeltaTime);
 
 		InstanceData->bFreezeTimeline = true;
 	}
 
-	void FSynchronizeUsingGroupsTrait::AdvanceBy(FExecutionContext& Context, const TTraitBinding<ITimeline>& Binding, float DeltaTime) const
+	void FSynchronizeUsingGroupsTrait::AdvanceBy(FExecutionContext& Context, const TTraitBinding<ITimelinePlayer>& Binding, float DeltaTime) const
 	{
 		const FInstanceData* InstanceData = Binding.GetInstanceData<FInstanceData>();
 		if (InstanceData->bFreezeTimeline)
@@ -154,7 +157,7 @@ namespace UE::AnimNext
 			return;	// If the timeline is frozen, we don't advance
 		}
 
-		ITimeline::AdvanceBy(Context, Binding, DeltaTime);
+		ITimelinePlayer::AdvanceBy(Context, Binding, DeltaTime);
 	}
 
 	void FSynchronizeUsingGroupsTrait::CreatePayloadForObject(FExecutionContext& Context, const TTraitBinding<IGraphFactory>& Binding, const UObject* InObject, FAnimNextDataInterfacePayload& InOutPayload) const
diff --git a/Engine/Plugins/Experimental/AnimNextAnimGraph/Source/AnimNextAnimGraph/Private/Traits/SynchronizeUsingGroups.h b/Engine/Plugins/Experimental/AnimNextAnimGraph/Source/AnimNextAnimGraph/Private/Traits/SynchronizeUsingGroups.h
--- a/Engine/Plugins/Experimental/AnimNextAnimGraph/Source/AnimNextAnimGraph/Private/Traits/SynchronizeUsingGroups.h
+++ b/Engine/Plugins/Experimental/AnimNextAnimGraph/Source/AnimNextAnimGraph/Private/Traits/SynchronizeUsingGroups.h
@@ -7,7 +7,7 @@
 #include "TraitInterfaces/IGroupSynchronization.h"
 #include "TraitInterfaces/IGraphFactory.h"
 #include "TraitInterfaces/IUpdate.h"
-#include "TraitInterfaces/ITimeline.h"
+#include "TraitInterfaces/ITimelinePlayer.h"
 
 #include "SynchronizeUsingGroups.generated.h"
 
@@ -54,7 +54,7 @@ namespace UE::AnimNext
 	 * 
 	 * A trait that synchronizes animation sequence playback using named groups.
 	 */
-	struct FSynchronizeUsingGroupsTrait : FAdditiveTrait, IUpdate, IGroupSynchronization, ITimeline, IGraphFactory
+	struct FSynchronizeUsingGroupsTrait : FAdditiveTrait, IUpdate, IGroupSynchronization, ITimelinePlayer, IGraphFactory
 	{
 		DECLARE_ANIM_TRAIT(FSynchronizeUsingGroupsTrait, FAdditiveTrait)
 
@@ -73,9 +73,14 @@ namespace UE::AnimNext
 
 			bool bFreezeTimeline = false;
 			bool bHasReachedFullWeight = false;
-			bool bHasTimeline = false;
+			bool bHasTimelinePlayer = false;
 		};
 
+#if WITH_EDITOR
+		// A trait stack has a single timeline, we can't support multiple instances
+		virtual bool MultipleInstanceSupport() const override { return false; }
+#endif
+
 		// IUpdate impl
 		virtual void OnBecomeRelevant(FUpdateTraversalContext& Context, const TTraitBinding<IUpdate>& Binding, const FTraitUpdateState& TraitState) const override;
 		virtual void PreUpdate(FUpdateTraversalContext& Context, const TTraitBinding<IUpdate>& Binding, const FTraitUpdateState& TraitState) const override;
@@ -84,8 +89,8 @@ namespace UE::AnimNext
 		virtual FSyncGroupParameters GetGroupParameters(FExecutionContext& Context, const TTraitBinding<IGroupSynchronization>& Binding) const override;
 		virtual void AdvanceBy(FExecutionContext& Context, const TTraitBinding<IGroupSynchronization>& Binding, float DeltaTime) const override;
 
-		// ITimeline impl
-		virtual void AdvanceBy(FExecutionContext& Context, const TTraitBinding<ITimeline>& Binding, float DeltaTime) const override;
+		// ITimelinePlayer impl
+		virtual void AdvanceBy(FExecutionContext& Context, const TTraitBinding<ITimelinePlayer>& Binding, float DeltaTime) const override;
 
 		// IGraphFactory impl
 		virtual void CreatePayloadForObject(FExecutionContext& Context, const TTraitBinding<IGraphFactory>& Binding, const UObject* InObject, FAnimNextDataInterfacePayload& InOutPayload) const override;
diff --git a/Engine/Plugins/Experimental/AnimNextAnimGraph/Source/AnimNextAnimGraph/Public/TraitCore/TraitBinding.h b/Engine/Plugins/Experimental/AnimNextAnimGraph/Source/AnimNextAnimGraph/Public/TraitCore/TraitBinding.h
--- a/Engine/Plugins/Experimental/AnimNextAnimGraph/Source/AnimNextAnimGraph/Public/TraitCore/TraitBinding.h
+++ b/Engine/Plugins/Experimental/AnimNextAnimGraph/Source/AnimNextAnimGraph/Public/TraitCore/TraitBinding.h
@@ -147,9 +147,43 @@ namespace UE::AnimNext
 			return AsInterfaceImpl(InterfaceUID, OutBinding);
 		}
 
+		// Queries the current trait for the presence of the specified interface.
+		// If no such interface is found, false it returned.
+		template<class TraitInterface>
+		bool HasInterface() const
+		{
+			static_assert(std::is_base_of<ITraitInterface, TraitInterface>::value, "TraitInterface type must derive from ITraitInterface");
+
+			constexpr FTraitInterfaceUID InterfaceUID = TraitInterface::InterfaceUID;
+			return HasInterfaceImpl(InterfaceUID);
+		}
+
+		// Queries the trait stack for the presence of a trait that implements the specified interface.
+		// If no such trait exists, false is returned.
+		template<class TraitInterface>
+		bool HasStackInterface() const
+		{
+			return IsValid() && Stack->HasInterface<TraitInterface>();
+		}
+
+		// Queries the trait stack for the presence of a trait lower on the stack that implements the specified interface.
+		// If no such trait exists, false is returned.
+		template<class TraitInterface>
+		bool HasStackInterfaceSuper() const
+		{
+			return IsValid() && Stack->HasInterfaceSuper<TraitInterface>(*this);
+		}
+
 		// Returns the trait stack binding backing this trait binding.
 		const FTraitStackBinding* GetStack() const { return Stack; }
 
+		// Returns a trait binding to the base of the stack.
+		// Returns true on success, false otherwise.
+		bool GetStackBaseTrait(FTraitBinding& OutBinding) const
+		{
+			return IsValid() && Stack->GetBaseTrait(OutBinding);
+		}
+
 		// Returns the trait pointer we are bound to.
 		FWeakTraitPtr GetTraitPtr() const { return FWeakTraitPtr(Stack != nullptr ? Stack->NodeInstance : nullptr, TraitIndex); }
 
@@ -202,6 +236,16 @@ namespace UE::AnimNext
 			return false;
 		}
 
+		bool HasInterfaceImpl(FTraitInterfaceUID InterfaceUID) const
+		{
+			if (!IsValid())
+			{
+				return false;
+			}
+
+			return TraitImpl->GetTraitInterface(InterfaceUID) != nullptr;
+		}
+
 		// A pointer to the bound trait stack or nullptr if we are not bound
 		const FTraitStackBinding* Stack = nullptr;
 
diff --git a/Engine/Plugins/Experimental/AnimNextAnimGraph/Source/AnimNextAnimGraph/Public/TraitCore/TraitStackBinding.h b/Engine/Plugins/Experimental/AnimNextAnimGraph/Source/AnimNextAnimGraph/Public/TraitCore/TraitStackBinding.h
--- a/Engine/Plugins/Experimental/AnimNextAnimGraph/Source/AnimNextAnimGraph/Public/TraitCore/TraitStackBinding.h
+++ b/Engine/Plugins/Experimental/AnimNextAnimGraph/Source/AnimNextAnimGraph/Public/TraitCore/TraitStackBinding.h
@@ -98,6 +98,16 @@ namespace UE::AnimNext
 		template<class TraitInterface>
 		bool GetInterfaceSuper(const FTraitBinding& Binding, TTraitBinding<TraitInterface>& OutSuperBinding) const;
 
+		// Queries the trait stack for the presence of a trait that implements the specified interface.
+		// The search begins at the top of the stack.
+		// Returns true on success, false otherwise.
+		template<class TraitInterface>
+		bool HasInterface() const;
+
+		// Queries the trait stack for the presence of a trait lower on the stack that implements the specified interface.
+		// Returns true on success, false otherwise.
+		template<class TraitInterface>
+		bool HasInterfaceSuper(const FTraitBinding& Binding) const;
 
 
 	private:
@@ -106,6 +116,8 @@ namespace UE::AnimNext
 
 		bool GetInterfaceImpl(FTraitInterfaceUID InterfaceUID, FTraitBinding& OutBinding) const;
 		bool GetInterfaceSuperImpl(FTraitInterfaceUID InterfaceUID, const FTraitBinding& Binding, FTraitBinding& OutSuperBinding) const;
+		bool HasInterfaceImpl(FTraitInterfaceUID InterfaceUID) const;
+		bool HasInterfaceSuperImpl(FTraitInterfaceUID InterfaceUID, const FTraitBinding& Binding) const;
 
 		// A pointer to the execution context that created the binding
 		const FExecutionContext* Context = nullptr;
@@ -150,4 +162,22 @@ namespace UE::AnimNext
 		constexpr FTraitInterfaceUID InterfaceUID = TraitInterface::InterfaceUID;
 		return GetInterfaceSuperImpl(InterfaceUID, Binding, OutSuperBinding);
 	}
+
+	template<class TraitInterface>
+	inline bool FTraitStackBinding::HasInterface() const
+	{
+		static_assert(std::is_base_of<ITraitInterface, TraitInterface>::value, "TraitInterface type must derive from ITraitInterface");
+
+		constexpr FTraitInterfaceUID InterfaceUID = TraitInterface::InterfaceUID;
+		return HasInterfaceImpl(InterfaceUID);
+	}
+
+	template<class TraitInterface>
+	inline bool FTraitStackBinding::HasInterfaceSuper(const FTraitBinding& Binding) const
+	{
+		static_assert(std::is_base_of<ITraitInterface, TraitInterface>::value, "TraitInterface type must derive from ITraitInterface");
+
+		constexpr FTraitInterfaceUID InterfaceUID = TraitInterface::InterfaceUID;
+		return HasInterfaceSuperImpl(InterfaceUID, Binding);
+	}
 }
diff --git a/Engine/Plugins/Experimental/AnimNextAnimGraph/Source/AnimNextAnimGraph/Public/TraitInterfaces/ITimeline.h b/Engine/Plugins/Experimental/AnimNextAnimGraph/Source/AnimNextAnimGraph/Public/TraitInterfaces/ITimeline.h
--- a/Engine/Plugins/Experimental/AnimNextAnimGraph/Source/AnimNextAnimGraph/Public/TraitInterfaces/ITimeline.h
+++ b/Engine/Plugins/Experimental/AnimNextAnimGraph/Source/AnimNextAnimGraph/Public/TraitInterfaces/ITimeline.h
@@ -193,16 +193,13 @@ namespace UE::AnimNext
 		DECLARE_ANIM_TRAIT_INTERFACE(ITimeline, 0x53760727)
 
 		// Returns the sync markers of this timeline sorted by time
-		virtual void GetSyncMarkers(FExecutionContext& Context, const TTraitBinding<ITimeline>& Binding, FTimelineSyncMarkerArray& OutSyncMarkers) const;
+		virtual void GetSyncMarkers(const FExecutionContext& Context, const TTraitBinding<ITimeline>& Binding, FTimelineSyncMarkerArray& OutSyncMarkers) const;
 
 		// Returns the state of this timeline
-		virtual FTimelineState GetState(FExecutionContext& Context, const TTraitBinding<ITimeline>& Binding) const;
-
-		// Advances time by the provided delta time (positive or negative) on this timeline
-		virtual void AdvanceBy(FExecutionContext& Context, const TTraitBinding<ITimeline>& Binding, float DeltaTime) const;
+		virtual FTimelineState GetState(const FExecutionContext& Context, const TTraitBinding<ITimeline>& Binding) const;
 
 		// Returns the last change in state of this timeline
-		virtual FTimelineDelta GetDelta(FExecutionContext& Context, const TTraitBinding<ITimeline>& Binding) const;
+		virtual FTimelineDelta GetDelta(const FExecutionContext& Context, const TTraitBinding<ITimeline>& Binding) const;
 
 #if WITH_EDITOR
 		virtual const FText& GetDisplayName() const override;
@@ -217,25 +214,19 @@ namespace UE::AnimNext
 	struct TTraitBinding<ITimeline> : FTraitBinding
 	{
 		// @see ITimeline::GetSyncMarkers
-		void GetSyncMarkers(FExecutionContext& Context, FTimelineSyncMarkerArray& OutSyncMarkers) const
+		void GetSyncMarkers(const FExecutionContext& Context, FTimelineSyncMarkerArray& OutSyncMarkers) const
 		{
 			return GetInterface()->GetSyncMarkers(Context, *this, OutSyncMarkers);
 		}
 
 		// @see ITimeline::GetState
-		FTimelineState GetState(FExecutionContext& Context) const
+		FTimelineState GetState(const FExecutionContext& Context) const
 		{
 			return GetInterface()->GetState(Context, *this);
 		}
 
-		// @see ITimeline::AdvanceBy
-		void AdvanceBy(FExecutionContext& Context, float DeltaTime) const
-		{
-			GetInterface()->AdvanceBy(Context, *this, DeltaTime);
-		}
-
 		// @see ITimeline::GetDelta
-		FTimelineDelta GetDelta(FExecutionContext& Context) const
+		FTimelineDelta GetDelta(const FExecutionContext& Context) const
 		{
 			return GetInterface()->GetDelta(Context, *this);
 		}
diff --git a/Engine/Plugins/Experimental/AnimNextAnimGraph/Source/AnimNextAnimGraph/Public/TraitInterfaces/ITimelinePlayer.h b/Engine/Plugins/Experimental/AnimNextAnimGraph/Source/AnimNextAnimGraph/Public/TraitInterfaces/ITimelinePlayer.h
new file mode 100644
--- /dev/null
+++ b/Engine/Plugins/Experimental/AnimNextAnimGraph/Source/AnimNextAnimGraph/Public/TraitInterfaces/ITimelinePlayer.h
@@ -0,0 +1,44 @@
+// Copyright Epic Games, Inc. All Rights Reserved.
+
+#pragma once
+
+#include "CoreMinimal.h"
+#include "TraitCore/ITraitInterface.h"
+#include "TraitCore/TraitBinding.h"
+
+namespace UE::AnimNext
+{
+	/**
+	 * ITimelinePlayer
+	 *
+	 * This interface exposes timeline playing behavior.
+	 */
+	struct ANIMNEXTANIMGRAPH_API ITimelinePlayer : ITraitInterface
+	{
+		DECLARE_ANIM_TRAIT_INTERFACE(ITimelinePlayer, 0xce8c6566)
+
+		// Advances time by the provided delta time (positive or negative) on this timeline player
+		virtual void AdvanceBy(FExecutionContext& Context, const TTraitBinding<ITimelinePlayer>& Binding, float DeltaTime) const;
+
+#if WITH_EDITOR
+		virtual const FText& GetDisplayName() const override;
+		virtual const FText& GetDisplayShortName() const override;
+#endif // WITH_EDITOR
+	};
+
+	/**
+	 * Specialization for trait binding.
+	 */
+	template<>
+	struct TTraitBinding<ITimelinePlayer> : FTraitBinding
+	{
+		// @see ITimelinePlayer::AdvanceBy
+		void AdvanceBy(FExecutionContext& Context, float DeltaTime) const
+		{
+			GetInterface()->AdvanceBy(Context, *this, DeltaTime);
+		}
+
+	protected:
+		const ITimelinePlayer* GetInterface() const { return GetInterfaceTyped<ITimelinePlayer>(); }
+	};
+}
diff --git a/Engine/Plugins/Experimental/EvaluationNotifies/Source/Runtime/Private/AnimNext/EvaluationNotifiesTrait.cpp b/Engine/Plugins/Experimental/EvaluationNotifies/Source/Runtime/Private/AnimNext/EvaluationNotifiesTrait.cpp
--- a/Engine/Plugins/Experimental/EvaluationNotifies/Source/Runtime/Private/AnimNext/EvaluationNotifiesTrait.cpp
+++ b/Engine/Plugins/Experimental/EvaluationNotifies/Source/Runtime/Private/AnimNext/EvaluationNotifiesTrait.cpp
@@ -15,6 +15,7 @@
 #include "TraitInterfaces/IBlendStack.h"
 #include "TraitInterfaces/IGraphFactory.h"
 #include "TraitInterfaces/INotifySource.h"
+#include "TraitInterfaces/ITimelinePlayer.h"
 #include "TraitInterfaces/ITimeline.h"
 
 static const FName RootBoneTransformName = "RootBoneTransform";
@@ -29,30 +30,34 @@ namespace UE::AnimNext
 	AUTO_REGISTER_ANIM_TRAIT(FEvaluationNotifiesTrait)
 
 	#define TRAIT_INTERFACE_ENUMERATOR(GeneratorMacro) \
-		GeneratorMacro(ITimeline) \
+		GeneratorMacro(ITimelinePlayer) \
 		GeneratorMacro(IUpdate) \
 		GeneratorMacro(IEvaluate) \
 
 	// Trait required interfaces implementation boilerplate
     #define TRAIT_REQUIRED_INTERFACE_ENUMERATOR(GeneratorMacroRequired) \
-    	GeneratorMacroRequired(ITimeline) \
+    	GeneratorMacroRequired(ITimelinePlayer) \
+		GeneratorMacroRequired(ITimeline) \
     	GeneratorMacroRequired(INotifySource) \
 
 	// Trait implementation boilerplate
 	GENERATE_ANIM_TRAIT_IMPLEMENTATION(FEvaluationNotifiesTrait, TRAIT_INTERFACE_ENUMERATOR, NULL_ANIM_TRAIT_INTERFACE_ENUMERATOR, NULL_ANIM_TRAIT_EVENT_ENUMERATOR)
 	#undef TRAIT_INTERFACE_ENUMERATOR
 
-	void FEvaluationNotifiesTrait::AdvanceBy(FExecutionContext& Context, const TTraitBinding<ITimeline>& Binding, float DeltaTime) const
+	void FEvaluationNotifiesTrait::AdvanceBy(FExecutionContext& Context, const TTraitBinding<ITimelinePlayer>& Binding, float DeltaTime) const
 	{
 		TTraitBinding<ITimeline> TimelineTrait;
 		ensure(Binding.GetStackInterfaceSuper(TimelineTrait));
 
+		TTraitBinding<ITimelinePlayer> TimelinePlayerTrait;
+		ensure(Binding.GetStackInterfaceSuper(TimelinePlayerTrait));
+
 		TTraitBinding<INotifySource> NotifySourceTrait;
 		ensure(Binding.GetStackInterfaceSuper(NotifySourceTrait));
 
 		// Get current state from the stack, advance time, then get the delta state
 		const FTimelineState PreAdvanceState = TimelineTrait.GetState(Context);
-		TimelineTrait.AdvanceBy(Context, DeltaTime);
+		TimelinePlayerTrait.AdvanceBy(Context, DeltaTime);
 		const FTimelineDelta Delta = TimelineTrait.GetDelta(Context);
 
 		// Query for notifies
diff --git a/Engine/Plugins/Experimental/EvaluationNotifies/Source/Runtime/Private/AnimNext/EvaluationNotifiesTrait.h b/Engine/Plugins/Experimental/EvaluationNotifies/Source/Runtime/Private/AnimNext/EvaluationNotifiesTrait.h
--- a/Engine/Plugins/Experimental/EvaluationNotifies/Source/Runtime/Private/AnimNext/EvaluationNotifiesTrait.h
+++ b/Engine/Plugins/Experimental/EvaluationNotifies/Source/Runtime/Private/AnimNext/EvaluationNotifiesTrait.h
@@ -9,7 +9,7 @@
 #include "TraitCore/TraitSharedData.h"
 #include "TraitInterfaces/IAttributeProvider.h"
 #include "TraitInterfaces/IEvaluate.h"
-#include "TraitInterfaces/ITimeline.h"
+#include "TraitInterfaces/ITimelinePlayer.h"
 #include "TraitInterfaces/IUpdate.h"
 #include "EvaluationNotifiesTrait.generated.h"
 
@@ -39,7 +39,7 @@ namespace UE::AnimNext
 /**
  * this trait will run evaluation time code for notifies which have a registered handler
  */
-struct EVALUATIONNOTIFIESRUNTIME_API FEvaluationNotifiesTrait : FAdditiveTrait, ITimeline, IUpdate, IEvaluate
+struct EVALUATIONNOTIFIESRUNTIME_API FEvaluationNotifiesTrait : FAdditiveTrait, ITimelinePlayer, IUpdate, IEvaluate
 {
 	DECLARE_ANIM_TRAIT(FEvaluationNotifiesTrait, FAdditiveTrait)
 
@@ -68,8 +68,8 @@ struct EVALUATIONNOTIFIESRUNTIME_API FEvaluationNotifiesTrait : FAdditiveTrait,
 	};
 	
 	
-	// ITimeline impl
-	virtual void AdvanceBy(FExecutionContext& Context, const TTraitBinding<ITimeline>& Binding, float DeltaTime) const override;
+	// ITimelinePlayer impl
+	virtual void AdvanceBy(FExecutionContext& Context, const TTraitBinding<ITimelinePlayer>& Binding, float DeltaTime) const override;
 	
 	// IUpdate impl 
 	virtual void PreUpdate(FUpdateTraversalContext& Context, const TTraitBinding<IUpdate>& Binding, const FTraitUpdateState& TraitState) const override;
