Anim Next: Fix missing reference counting on injection requests.

diff --git a/Engine/Plugins/Experimental/AnimNextAnimGraph/Source/AnimNextAnimGraph/Private/Injection/AnimNextModuleInjectionComponent.cpp b/Engine/Plugins/Experimental/AnimNextAnimGraph/Source/AnimNextAnimGraph/Private/Injection/AnimNextModuleInjectionComponent.cpp
--- a/Engine/Plugins/Experimental/AnimNextAnimGraph/Source/AnimNextAnimGraph/Private/Injection/AnimNextModuleInjectionComponent.cpp
+++ b/Engine/Plugins/Experimental/AnimNextAnimGraph/Source/AnimNextAnimGraph/Private/Injection/AnimNextModuleInjectionComponent.cpp
@@ -1,4 +1,4 @@
-﻿// Copyright Epic Games, Inc. All Rights Reserved.
+// Copyright Epic Games, Inc. All Rights Reserved.
 
 #include "Injection/AnimNextModuleInjectionComponent.h"
 
@@ -194,3 +194,23 @@ uint32 FAnimNextModuleInjectionComponent::IncrementSerialNumber()
 	}
 	return SerialNumber;
 }
+
+void FAnimNextModuleInjectionComponent::AddStructReferencedObjects(FReferenceCollector& Collector)
+{
+	for (TPair<FName, FInjectionRecord>& RequestPair : CurrentRequests)
+	{
+		RequestPair.Value.AddReferencedObjects(Collector);
+	}
+}
+
+void FAnimNextModuleInjectionComponent::FInjectionRecord::AddReferencedObjects(FReferenceCollector& Collector)
+{
+	if (GraphRequest.IsValid())
+	{
+		GraphRequest->AddReferencedObjects(Collector);
+	}
+	if (ModifierRequest.IsValid())
+	{
+		ModifierRequest->AddReferencedObjects(Collector);
+	}
+}
diff --git a/Engine/Plugins/Experimental/AnimNextAnimGraph/Source/AnimNextAnimGraph/Private/Injection/AnimNextModuleInjectionComponent.h b/Engine/Plugins/Experimental/AnimNextAnimGraph/Source/AnimNextAnimGraph/Private/Injection/AnimNextModuleInjectionComponent.h
--- a/Engine/Plugins/Experimental/AnimNextAnimGraph/Source/AnimNextAnimGraph/Private/Injection/AnimNextModuleInjectionComponent.h
+++ b/Engine/Plugins/Experimental/AnimNextAnimGraph/Source/AnimNextAnimGraph/Private/Injection/AnimNextModuleInjectionComponent.h
@@ -1,4 +1,4 @@
-﻿// Copyright Epic Games, Inc. All Rights Reserved.
+// Copyright Epic Games, Inc. All Rights Reserved.
 
 #pragma once
 
@@ -6,6 +6,7 @@
 #include "DataInterface/AnimNextDataInterfaceHost.h"
 #include "Injection/InjectionRequest.h"
 #include "Module/AnimNextModuleInstanceComponent.h"
+#include "UObject/GCObject.h"
 #include "AnimNextModuleInjectionComponent.generated.h"
 
 namespace UE::AnimNext
@@ -26,6 +27,8 @@ struct FAnimNextModuleInjectionComponent : public FAnimNextModuleInstanceCompone
 	// Get the cached injection info for our module 
 	const UE::AnimNext::FInjectionInfo& GetInjectionInfo() const { return InjectionInfo; }
 
+	void AddStructReferencedObjects(class FReferenceCollector& Collector);
+
 private:
 	// FAnimNextModuleInstanceComponent interface
 	virtual void OnInitialize() override;
@@ -49,6 +52,8 @@ private:
 			return GraphRequest.IsValid() || ModifierRequest.IsValid();
 		}
 
+		void AddReferencedObjects(FReferenceCollector& Collector);
+
 		TSharedPtr<UE::AnimNext::FInjectionRequest> GraphRequest;
 		TSharedPtr<UE::AnimNext::FInjectionRequest> ModifierRequest;
 		uint32 SerialNumber = 0;
@@ -60,3 +65,12 @@ private:
 	// Serial number used to identify forwarded requests
 	uint32 SerialNumber = 0;
 };
+
+template<>
+struct TStructOpsTypeTraits<FAnimNextModuleInjectionComponent> : public TStructOpsTypeTraitsBase2<FAnimNextModuleInjectionComponent>
+{
+	enum
+	{
+		WithAddStructReferencedObjects = true,
+	};
+};
