AnimNext: Add blend smoother to state tree

diff --git a/Engine/Plugins/Experimental/AnimNextAnimGraph/Source/AnimNextAnimGraph/Private/Traits/BlendSmoother.h b/Engine/Plugins/Experimental/AnimNextAnimGraph/Source/AnimNextAnimGraph/Private/Traits/BlendSmoother.h
deleted file mode 100644
diff --git a/Engine/Plugins/Experimental/AnimNextAnimGraph/Source/AnimNextAnimGraph/Public/Traits/BlendSmoother.h b/Engine/Plugins/Experimental/AnimNextAnimGraph/Source/AnimNextAnimGraph/Public/Traits/BlendSmoother.h
new file mode 100644
--- /dev/null
+++ b/Engine/Plugins/Experimental/AnimNextAnimGraph/Source/AnimNextAnimGraph/Public/Traits/BlendSmoother.h
@@ -0,0 +1,117 @@
+// Copyright Epic Games, Inc. All Rights Reserved.
+
+#pragma once
+
+#include "CoreMinimal.h"
+
+#include "AlphaBlend.h"
+#include "TraitCore/Trait.h"
+#include "TraitInterfaces/IDiscreteBlend.h"
+#include "TraitInterfaces/IEvaluate.h"
+#include "TraitInterfaces/ISmoothBlend.h"
+#include "TraitInterfaces/IUpdate.h"
+
+#include "BlendSmoother.generated.h"
+
+class UCurveFloat;
+
+USTRUCT(meta = (DisplayName = "Blend Smoother Core"))
+struct FAnimNextBlendSmootherCoreTraitSharedData : public FAnimNextTraitSharedData
+{
+	GENERATED_BODY()
+
+	// This struct is empty but required so that we can give a nice display name to the trait
+};
+
+USTRUCT(meta = (DisplayName = "Blend Smoother"))
+struct FAnimNextBlendSmootherTraitSharedData : public FAnimNextBlendSmootherCoreTraitSharedData
+{
+	GENERATED_BODY()
+
+	/** How long to take when blending into each child. */
+	UPROPERTY(EditAnywhere, Category = "Default", meta = (Inline))
+	TArray<float> BlendTimes;
+
+	/** What type of blend equation to use when converting the time elapsed into a blend weight. */
+	UPROPERTY(EditAnywhere, Category = "Default", meta = (Inline))
+	EAlphaBlendOption BlendType = EAlphaBlendOption::Linear;
+
+	/** Custom curve to use when the Custom blend type is used. */
+	UPROPERTY(EditAnywhere, Category = "Default", meta = (Inline))
+	TObjectPtr<UCurveFloat> CustomBlendCurve;
+};
+
+namespace UE::AnimNext
+{
+	/**
+	 * FBlendSmootherCoreTrait
+	 *
+	 * A trait that smoothly blends between discrete states over time.
+	 * It only implements the smoothing logic, it queries its required arguments using the ISmoothBlend interface.
+	 */
+	struct FBlendSmootherCoreTrait : FAdditiveTrait, IEvaluate, IUpdate, IDiscreteBlend
+	{
+		DECLARE_ANIM_TRAIT(FBlendSmootherCoreTrait, FAdditiveTrait)
+
+		using FSharedData = FAnimNextBlendSmootherCoreTraitSharedData;
+
+		// Struct for tracking blends for each pose
+		struct FBlendData
+		{
+			// Helper struct to update a time based weight
+			FAlphaBlend Blend;
+
+			// Current child weight (normalized with all children)
+			float Weight = 0.0f;
+
+			// Whether or not this child is actively blending
+			bool bIsBlending = false;
+		};
+
+		struct FInstanceData : FTrait::FInstanceData
+		{
+			// Blend state per child
+			TArray<FBlendData> PerChildBlendData;
+		};
+
+		// IEvaluate impl
+		virtual void PostEvaluate(FEvaluateTraversalContext& Context, const TTraitBinding<IEvaluate>& Binding) const override;
+
+		// IUpdate impl
+		virtual void PreUpdate(FUpdateTraversalContext& Context, const TTraitBinding<IUpdate>& Binding, const FTraitUpdateState& TraitState) const override;
+
+		// IDiscreteBlend impl
+		virtual float GetBlendWeight(FExecutionContext& Context, const TTraitBinding<IDiscreteBlend>& Binding, int32 ChildIndex) const override;
+		virtual const FAlphaBlend* GetBlendState(FExecutionContext& Context, const TTraitBinding<IDiscreteBlend>& Binding, int32 ChildIndex) const override;
+		virtual void OnBlendTransition(FExecutionContext& Context, const TTraitBinding<IDiscreteBlend>& Binding, int32 OldChildIndex, int32 NewChildIndex) const override;
+
+		// Internal impl
+		static void InitializeInstanceData(FExecutionContext& Context, const FTraitBinding& Binding, const FSharedData* SharedData, FInstanceData* InstanceData);
+
+#if WITH_EDITOR
+		virtual bool IsHidden() const override { return true; }
+#endif
+	};
+
+	/**
+	 * FBlendSmootherTrait
+	 * 
+	 * A trait that smoothly blends between discrete states over time.
+	 * This trait implements both the logic and contains the arguments necessary.
+	 */
+	struct FBlendSmootherTrait : FBlendSmootherCoreTrait, ISmoothBlend
+	{
+		DECLARE_ANIM_TRAIT(FBlendSmootherTrait, FBlendSmootherCoreTrait)
+
+		using FSharedData = FAnimNextBlendSmootherTraitSharedData;
+
+		// ISmoothBlend impl
+		virtual float GetBlendTime(FExecutionContext& Context, const TTraitBinding<ISmoothBlend>& Binding, int32 ChildIndex) const override;
+		virtual EAlphaBlendOption GetBlendType(FExecutionContext& Context, const TTraitBinding<ISmoothBlend>& Binding, int32 ChildIndex) const override;
+		virtual UCurveFloat* GetCustomBlendCurve(FExecutionContext& Context, const TTraitBinding<ISmoothBlend>& Binding, int32 ChildIndex) const override;
+
+#if WITH_EDITOR
+		virtual bool IsHidden() const override { return false; } // overrdie the base, as it is hidden
+#endif
+	};
+}
diff --git a/Engine/Plugins/Experimental/AnimNextStateTree/Source/AnimNextStateTreeUncookedOnly/Internal/AnimNextStateTree_EditorData.cpp b/Engine/Plugins/Experimental/AnimNextStateTree/Source/AnimNextStateTreeUncookedOnly/Internal/AnimNextStateTree_EditorData.cpp
--- a/Engine/Plugins/Experimental/AnimNextStateTree/Source/AnimNextStateTreeUncookedOnly/Internal/AnimNextStateTree_EditorData.cpp
+++ b/Engine/Plugins/Experimental/AnimNextStateTree/Source/AnimNextStateTreeUncookedOnly/Internal/AnimNextStateTree_EditorData.cpp
@@ -24,6 +24,7 @@
 #include "AnimStateTreeTrait.h"
 #include "Entries/AnimNextDataInterfaceEntry.h"
 #include "Traits/BlendStackTrait.h"
+#include "Traits/BlendSmoother.h"
 
 TSubclassOf<UAssetUserData> UAnimNextStateTree_EditorData::GetAssetUserDataClass() const
 {
@@ -149,6 +150,14 @@ void UAnimNextStateTree_EditorData::OnPreCompileGetProgrammaticGraphs(const FRig
 			return;
 		}
 
+		const FName SmootherTraitName = Controller->AddTraitByName(TraitStackNode->GetFName(), *UE::AnimNext::FTraitRegistry::Get().Find(UE::AnimNext::FBlendSmootherCoreTrait::TraitUID)->GetTraitName(), INDEX_NONE, TEXT(""), false);
+
+		if (SmootherTraitName == NAME_None)
+		{
+			InSettings.ReportError(TEXT("Failed to add Blend Smoother Core trait to node"));
+			return;
+		}
+
 		URigVMPin* StateTreeReferencePin = TraitStackNode->FindTrait(StateTreeTraitName, GET_MEMBER_NAME_STRING_CHECKED(FAnimNextStateTreeTraitSharedData, StateTreeReference));
 		if (StateTreeReferencePin == nullptr)
 		{
