Use constexpr hash generation for AnimNext interface UIDs to reduce boilerplate

diff --git a/Engine/Plugins/Experimental/AnimNextAnimGraph/Source/AnimNextAnimGraph/Public/TraitCore/IScopedTraitInterface.h b/Engine/Plugins/Experimental/AnimNextAnimGraph/Source/AnimNextAnimGraph/Public/TraitCore/IScopedTraitInterface.h
--- a/Engine/Plugins/Experimental/AnimNextAnimGraph/Source/AnimNextAnimGraph/Public/TraitCore/IScopedTraitInterface.h
+++ b/Engine/Plugins/Experimental/AnimNextAnimGraph/Source/AnimNextAnimGraph/Public/TraitCore/IScopedTraitInterface.h
@@ -14,6 +14,6 @@ namespace UE::AnimNext
 	 */
 	struct ANIMNEXTANIMGRAPH_API IScopedTraitInterface : ITraitInterface
 	{
-		DECLARE_ANIM_TRAIT_INTERFACE(IScopedTraitInterface, 0x621b01eb)
+		DECLARE_ANIM_TRAIT_INTERFACE(IScopedTraitInterface)
 	};
 }
diff --git a/Engine/Plugins/Experimental/AnimNextAnimGraph/Source/AnimNextAnimGraph/Public/TraitCore/ITraitInterface.h b/Engine/Plugins/Experimental/AnimNextAnimGraph/Source/AnimNextAnimGraph/Public/TraitCore/ITraitInterface.h
--- a/Engine/Plugins/Experimental/AnimNextAnimGraph/Source/AnimNextAnimGraph/Public/TraitCore/ITraitInterface.h
+++ b/Engine/Plugins/Experimental/AnimNextAnimGraph/Source/AnimNextAnimGraph/Public/TraitCore/ITraitInterface.h
@@ -13,9 +13,9 @@ namespace UE::AnimNext
 // Helper macros
 
 // In a trait interface struct declaration, this macro declares the necessary boilerplate we require
-#define DECLARE_ANIM_TRAIT_INTERFACE(InterfaceName, InterfaceNameHash) \
+#define DECLARE_ANIM_TRAIT_INTERFACE(InterfaceName) \
 	/* Globally unique UID for this interface */ \
-	static constexpr UE::AnimNext::FTraitInterfaceUID InterfaceUID = UE::AnimNext::FTraitInterfaceUID(InterfaceNameHash, TEXT(#InterfaceName)); \
+	static constexpr UE::AnimNext::FTraitInterfaceUID InterfaceUID = UE::AnimNext::FTraitInterfaceUID::MakeUID(TEXT(#InterfaceName)); \
 	virtual UE::AnimNext::FTraitInterfaceUID GetInterfaceUID() const override { return InterfaceUID; }
 
 // Allows a trait interface to auto-register and unregister within the current execution scope
@@ -40,7 +40,7 @@ namespace UE::AnimNext
 		// The globally unique UID for this interface
 		// Derived types will have their own InterfaceUID member that hides/aliases/shadows this one
 		// @see DECLARE_ANIM_TRAIT_INTERFACE
-		static constexpr FTraitInterfaceUID InterfaceUID = FTraitInterfaceUID(0xcdc27733, TEXT("ITraitInterface"));
+		static constexpr FTraitInterfaceUID InterfaceUID = FTraitInterfaceUID::MakeUID(TEXT("ITraitInterface"));
 
 		// Returns the globally unique UID for this interface
 		virtual FTraitInterfaceUID GetInterfaceUID() const { return InterfaceUID; };
diff --git a/Engine/Plugins/Experimental/AnimNextAnimGraph/Source/AnimNextAnimGraph/Public/TraitCore/TraitInterfaceUID.h b/Engine/Plugins/Experimental/AnimNextAnimGraph/Source/AnimNextAnimGraph/Public/TraitCore/TraitInterfaceUID.h
--- a/Engine/Plugins/Experimental/AnimNextAnimGraph/Source/AnimNextAnimGraph/Public/TraitCore/TraitInterfaceUID.h
+++ b/Engine/Plugins/Experimental/AnimNextAnimGraph/Source/AnimNextAnimGraph/Public/TraitCore/TraitInterfaceUID.h
@@ -3,6 +3,7 @@
 #pragma once
 
 #include "CoreMinimal.h"
+#include "Hash/Fnv.h"
 
 namespace UE::AnimNext
 {
@@ -13,32 +14,37 @@ namespace UE::AnimNext
 	 * FTraitInterfaceUID
 	 *
 	 * Encapsulates an interface global UID.
-	 * The string is exposed in non-shipping builds for logging and debugging purposes.
-	 * The UID should be generated from the provided string using FNV1a with 32 bits.
 	 *
-	 * The whole struct is meant to be 'constexpr' to allow inlining in interface queries.
+	 * The whole struct is meant to be 'constexpr' to allow inlining within the assembly as constants.
 	 */
 	struct FTraitInterfaceUID final
 	{
 		// Constructs an invalid UID
 		constexpr FTraitInterfaceUID()
 			: UID(INVALID_UID)
-#if !UE_BUILD_SHIPPING
+#if !UE_BUILD_SHIPPING || WITH_EDITOR
 			, InterfaceName(TEXT("<Invalid Interface UID>"))
 #endif
 		{}
 
-		// Constructs an interface UID
-		explicit constexpr FTraitInterfaceUID(FTraitInterfaceUIDRaw InUID, const TCHAR* InInterfaceName = TEXT("<Unknown Trait Interface Name>"))
+		// Constructs an interface UID from its raw value
+		explicit constexpr FTraitInterfaceUID(FTraitInterfaceUIDRaw InUID)
 			: UID(InUID)
-#if !UE_BUILD_SHIPPING
-			, InterfaceName(InInterfaceName)
+#if !UE_BUILD_SHIPPING || WITH_EDITOR
+			, InterfaceName(TEXT("<Unknown Trait Interface Name>"))
 #endif
 		{
 		}
 
-#if !UE_BUILD_SHIPPING
-		// Returns a literal string to the interface name
+		// Constructs a trait UID from a string literal
+		template<typename CharType, int32 Len>
+		static UE_CONSTEVAL FTraitInterfaceUID MakeUID(const CharType(&InInterfaceName)[Len]) noexcept
+		{
+			return FTraitInterfaceUID(UE::HashStringFNV1a<FTraitInterfaceUIDRaw>(MakeStringView(InInterfaceName, Len - 1)), InInterfaceName);
+		}
+
+#if !UE_BUILD_SHIPPING || WITH_EDITOR
+		// Returns a string literal to the interface name
 		constexpr const TCHAR* GetInterfaceName() const { return InterfaceName; }
 #endif
 
@@ -51,6 +57,15 @@ namespace UE::AnimNext
 	private:
 		static constexpr FTraitInterfaceUIDRaw INVALID_UID = 0;
 
+		// Constructs an interface UID
+		explicit constexpr FTraitInterfaceUID(FTraitInterfaceUIDRaw InUID, const TCHAR* InInterfaceName) noexcept
+			: UID(InUID)
+#if !UE_BUILD_SHIPPING || WITH_EDITOR
+			, InterfaceName(InInterfaceName)
+#endif
+		{
+		}
+
 		FTraitInterfaceUIDRaw		UID;
 
 #if !UE_BUILD_SHIPPING || WITH_EDITOR
diff --git a/Engine/Plugins/Experimental/AnimNextAnimGraph/Source/AnimNextAnimGraph/Public/TraitInterfaces/IAttributeProvider.h b/Engine/Plugins/Experimental/AnimNextAnimGraph/Source/AnimNextAnimGraph/Public/TraitInterfaces/IAttributeProvider.h
--- a/Engine/Plugins/Experimental/AnimNextAnimGraph/Source/AnimNextAnimGraph/Public/TraitInterfaces/IAttributeProvider.h
+++ b/Engine/Plugins/Experimental/AnimNextAnimGraph/Source/AnimNextAnimGraph/Public/TraitInterfaces/IAttributeProvider.h
@@ -23,7 +23,7 @@ namespace UE::AnimNext
 	 */
 	struct ANIMNEXTANIMGRAPH_API IAttributeProvider : ITraitInterface
 	{
-		DECLARE_ANIM_TRAIT_INTERFACE(IAttributeProvider, 0x698139d1)
+		DECLARE_ANIM_TRAIT_INTERFACE(IAttributeProvider)
 
 		// Returns a callback that can be used to extract root motion at a later time
 		virtual FOnExtractRootMotionAttribute GetOnExtractRootMotionAttribute(FExecutionContext& Context, const TTraitBinding<IAttributeProvider>& Binding) const;
diff --git a/Engine/Plugins/Experimental/AnimNextAnimGraph/Source/AnimNextAnimGraph/Public/TraitInterfaces/IBlendStack.h b/Engine/Plugins/Experimental/AnimNextAnimGraph/Source/AnimNextAnimGraph/Public/TraitInterfaces/IBlendStack.h
--- a/Engine/Plugins/Experimental/AnimNextAnimGraph/Source/AnimNextAnimGraph/Public/TraitInterfaces/IBlendStack.h
+++ b/Engine/Plugins/Experimental/AnimNextAnimGraph/Source/AnimNextAnimGraph/Public/TraitInterfaces/IBlendStack.h
@@ -20,7 +20,7 @@ namespace UE::AnimNext
 	 */
 	struct ANIMNEXTANIMGRAPH_API IBlendStack : ITraitInterface
 	{
-		DECLARE_ANIM_TRAIT_INTERFACE(IBlendStack, 0xf3468e64)
+		DECLARE_ANIM_TRAIT_INTERFACE(IBlendStack)
 
 		enum class EBlendMode : uint8
 		{
diff --git a/Engine/Plugins/Experimental/AnimNextAnimGraph/Source/AnimNextAnimGraph/Public/TraitInterfaces/IContinuousBlend.h b/Engine/Plugins/Experimental/AnimNextAnimGraph/Source/AnimNextAnimGraph/Public/TraitInterfaces/IContinuousBlend.h
--- a/Engine/Plugins/Experimental/AnimNextAnimGraph/Source/AnimNextAnimGraph/Public/TraitInterfaces/IContinuousBlend.h
+++ b/Engine/Plugins/Experimental/AnimNextAnimGraph/Source/AnimNextAnimGraph/Public/TraitInterfaces/IContinuousBlend.h
@@ -19,7 +19,7 @@ namespace UE::AnimNext
 	 */
 	struct ANIMNEXTANIMGRAPH_API IContinuousBlend : ITraitInterface
 	{
-		DECLARE_ANIM_TRAIT_INTERFACE(IContinuousBlend, 0xe7d79186)
+		DECLARE_ANIM_TRAIT_INTERFACE(IContinuousBlend)
 
 		// Returns the blend weight for the specified child
 		// Multiple children can have non-zero weight but their sum must be 1.0
diff --git a/Engine/Plugins/Experimental/AnimNextAnimGraph/Source/AnimNextAnimGraph/Public/TraitInterfaces/IDiscreteBlend.h b/Engine/Plugins/Experimental/AnimNextAnimGraph/Source/AnimNextAnimGraph/Public/TraitInterfaces/IDiscreteBlend.h
--- a/Engine/Plugins/Experimental/AnimNextAnimGraph/Source/AnimNextAnimGraph/Public/TraitInterfaces/IDiscreteBlend.h
+++ b/Engine/Plugins/Experimental/AnimNextAnimGraph/Source/AnimNextAnimGraph/Public/TraitInterfaces/IDiscreteBlend.h
@@ -19,7 +19,7 @@ namespace UE::AnimNext
 	 */
 	struct ANIMNEXTANIMGRAPH_API IDiscreteBlend : ITraitInterface
 	{
-		DECLARE_ANIM_TRAIT_INTERFACE(IDiscreteBlend, 0x2d395d56)
+		DECLARE_ANIM_TRAIT_INTERFACE(IDiscreteBlend)
 
 		// Returns the blend weight for the specified child
 		// Multiple children can have non-zero weight but their sum must be 1.0
diff --git a/Engine/Plugins/Experimental/AnimNextAnimGraph/Source/AnimNextAnimGraph/Public/TraitInterfaces/IEvaluate.h b/Engine/Plugins/Experimental/AnimNextAnimGraph/Source/AnimNextAnimGraph/Public/TraitInterfaces/IEvaluate.h
--- a/Engine/Plugins/Experimental/AnimNextAnimGraph/Source/AnimNextAnimGraph/Public/TraitInterfaces/IEvaluate.h
+++ b/Engine/Plugins/Experimental/AnimNextAnimGraph/Source/AnimNextAnimGraph/Public/TraitInterfaces/IEvaluate.h
@@ -54,7 +54,7 @@ namespace UE::AnimNext
 	 */
 	struct ANIMNEXTANIMGRAPH_API IEvaluate : ITraitInterface
 	{
-		DECLARE_ANIM_TRAIT_INTERFACE(IEvaluate, 0xa303e9e7)
+		DECLARE_ANIM_TRAIT_INTERFACE(IEvaluate)
 
 		// Called before a traits children are evaluated
 		virtual void PreEvaluate(FEvaluateTraversalContext& Context, const TTraitBinding<IEvaluate>& Binding) const;
diff --git a/Engine/Plugins/Experimental/AnimNextAnimGraph/Source/AnimNextAnimGraph/Public/TraitInterfaces/IGarbageCollection.h b/Engine/Plugins/Experimental/AnimNextAnimGraph/Source/AnimNextAnimGraph/Public/TraitInterfaces/IGarbageCollection.h
--- a/Engine/Plugins/Experimental/AnimNextAnimGraph/Source/AnimNextAnimGraph/Public/TraitInterfaces/IGarbageCollection.h
+++ b/Engine/Plugins/Experimental/AnimNextAnimGraph/Source/AnimNextAnimGraph/Public/TraitInterfaces/IGarbageCollection.h
@@ -17,7 +17,7 @@ namespace UE::AnimNext
 	 */
 	struct ANIMNEXTANIMGRAPH_API IGarbageCollection : ITraitInterface
 	{
-		DECLARE_ANIM_TRAIT_INTERFACE(IGarbageCollection, 0x231a2017)
+		DECLARE_ANIM_TRAIT_INTERFACE(IGarbageCollection)
 
 
 		// Registers the provided binding for GC callback
diff --git a/Engine/Plugins/Experimental/AnimNextAnimGraph/Source/AnimNextAnimGraph/Public/TraitInterfaces/IGraphFactory.h b/Engine/Plugins/Experimental/AnimNextAnimGraph/Source/AnimNextAnimGraph/Public/TraitInterfaces/IGraphFactory.h
--- a/Engine/Plugins/Experimental/AnimNextAnimGraph/Source/AnimNextAnimGraph/Public/TraitInterfaces/IGraphFactory.h
+++ b/Engine/Plugins/Experimental/AnimNextAnimGraph/Source/AnimNextAnimGraph/Public/TraitInterfaces/IGraphFactory.h
@@ -19,7 +19,7 @@ namespace UE::AnimNext
 	 */
 	struct ANIMNEXTANIMGRAPH_API IGraphFactory : ITraitInterface
 	{
-		DECLARE_ANIM_TRAIT_INTERFACE(IGraphFactory, 0xb5c08dd6)
+		DECLARE_ANIM_TRAIT_INTERFACE(IGraphFactory)
 
 		// Gets a graph and a set of native interfaces from an object
 		// @param	InObject				The object to manufacture a graph for (e.g. AnimSequence). If this is a UAnimNextAnimationGraph, it should be returned verbatim
diff --git a/Engine/Plugins/Experimental/AnimNextAnimGraph/Source/AnimNextAnimGraph/Public/TraitInterfaces/IGroupSynchronization.h b/Engine/Plugins/Experimental/AnimNextAnimGraph/Source/AnimNextAnimGraph/Public/TraitInterfaces/IGroupSynchronization.h
--- a/Engine/Plugins/Experimental/AnimNextAnimGraph/Source/AnimNextAnimGraph/Public/TraitInterfaces/IGroupSynchronization.h
+++ b/Engine/Plugins/Experimental/AnimNextAnimGraph/Source/AnimNextAnimGraph/Public/TraitInterfaces/IGroupSynchronization.h
@@ -96,7 +96,7 @@ namespace UE::AnimNext
 	 */
 	struct ANIMNEXTANIMGRAPH_API IGroupSynchronization : ITraitInterface
 	{
-		DECLARE_ANIM_TRAIT_INTERFACE(IGroupSynchronization, 0xf607d0fd)
+		DECLARE_ANIM_TRAIT_INTERFACE(IGroupSynchronization)
 
 		// Returns the group parameters used for synchronization
 		virtual FSyncGroupParameters GetGroupParameters(FExecutionContext& Context, const TTraitBinding<IGroupSynchronization>& Binding) const;
diff --git a/Engine/Plugins/Experimental/AnimNextAnimGraph/Source/AnimNextAnimGraph/Public/TraitInterfaces/IHierarchy.h b/Engine/Plugins/Experimental/AnimNextAnimGraph/Source/AnimNextAnimGraph/Public/TraitInterfaces/IHierarchy.h
--- a/Engine/Plugins/Experimental/AnimNextAnimGraph/Source/AnimNextAnimGraph/Public/TraitInterfaces/IHierarchy.h
+++ b/Engine/Plugins/Experimental/AnimNextAnimGraph/Source/AnimNextAnimGraph/Public/TraitInterfaces/IHierarchy.h
@@ -21,7 +21,7 @@ namespace UE::AnimNext
 	 */
 	struct IHierarchy : ITraitInterface
 	{
-		DECLARE_ANIM_TRAIT_INTERFACE(IHierarchy, 0x846d8a37)
+		DECLARE_ANIM_TRAIT_INTERFACE(IHierarchy)
 
 		// Returns the number of children of the trait implementation (not the whole stack)
 		// Includes inactive children
diff --git a/Engine/Plugins/Experimental/AnimNextAnimGraph/Source/AnimNextAnimGraph/Public/TraitInterfaces/IInertializerBlend.h b/Engine/Plugins/Experimental/AnimNextAnimGraph/Source/AnimNextAnimGraph/Public/TraitInterfaces/IInertializerBlend.h
--- a/Engine/Plugins/Experimental/AnimNextAnimGraph/Source/AnimNextAnimGraph/Public/TraitInterfaces/IInertializerBlend.h
+++ b/Engine/Plugins/Experimental/AnimNextAnimGraph/Source/AnimNextAnimGraph/Public/TraitInterfaces/IInertializerBlend.h
@@ -19,7 +19,7 @@ namespace UE::AnimNext
 	 */
 	struct ANIMNEXTANIMGRAPH_API IInertializerBlend : ITraitInterface
 	{
-		DECLARE_ANIM_TRAIT_INTERFACE(IInertializerBlend, 0x3856b8e9)
+		DECLARE_ANIM_TRAIT_INTERFACE(IInertializerBlend)
 
 		// Returns the desired blend time for the specified child
 		virtual float GetBlendTime(FExecutionContext& Context, const TTraitBinding<IInertializerBlend>& Binding, int32 ChildIndex) const;
diff --git a/Engine/Plugins/Experimental/AnimNextAnimGraph/Source/AnimNextAnimGraph/Public/TraitInterfaces/INotifySource.h b/Engine/Plugins/Experimental/AnimNextAnimGraph/Source/AnimNextAnimGraph/Public/TraitInterfaces/INotifySource.h
--- a/Engine/Plugins/Experimental/AnimNextAnimGraph/Source/AnimNextAnimGraph/Public/TraitInterfaces/INotifySource.h
+++ b/Engine/Plugins/Experimental/AnimNextAnimGraph/Source/AnimNextAnimGraph/Public/TraitInterfaces/INotifySource.h
@@ -1,4 +1,4 @@
-﻿// Copyright Epic Games, Inc. All Rights Reserved.
+// Copyright Epic Games, Inc. All Rights Reserved.
 
 #pragma once
 
@@ -17,7 +17,7 @@ namespace UE::AnimNext
 	 */
 	struct ANIMNEXTANIMGRAPH_API INotifySource : ITraitInterface
 	{
-		DECLARE_ANIM_TRAIT_INTERFACE(INotifySource, 0x0be015e0)
+		DECLARE_ANIM_TRAIT_INTERFACE(INotifySource)
 
 		// Returns the events of this timeline in the specified time interval
 		virtual void GetNotifies(FExecutionContext& Context, const TTraitBinding<INotifySource>& Binding, float StartPosition, float Duration, bool bLooping, TArray<FAnimNotifyEventReference>& OutNotifies) const;
diff --git a/Engine/Plugins/Experimental/AnimNextAnimGraph/Source/AnimNextAnimGraph/Public/TraitInterfaces/ISmoothBlend.h b/Engine/Plugins/Experimental/AnimNextAnimGraph/Source/AnimNextAnimGraph/Public/TraitInterfaces/ISmoothBlend.h
--- a/Engine/Plugins/Experimental/AnimNextAnimGraph/Source/AnimNextAnimGraph/Public/TraitInterfaces/ISmoothBlend.h
+++ b/Engine/Plugins/Experimental/AnimNextAnimGraph/Source/AnimNextAnimGraph/Public/TraitInterfaces/ISmoothBlend.h
@@ -20,7 +20,7 @@ namespace UE::AnimNext
 	 */
 	struct ANIMNEXTANIMGRAPH_API ISmoothBlend : ITraitInterface
 	{
-		DECLARE_ANIM_TRAIT_INTERFACE(ISmoothBlend, 0x1c2c1739)
+		DECLARE_ANIM_TRAIT_INTERFACE(ISmoothBlend)
 
 		// Returns the desired blend time for the specified child
 		virtual float GetBlendTime(FExecutionContext& Context, const TTraitBinding<ISmoothBlend>& Binding, int32 ChildIndex) const;
diff --git a/Engine/Plugins/Experimental/AnimNextAnimGraph/Source/AnimNextAnimGraph/Public/TraitInterfaces/ISmoothBlendPerBone.h b/Engine/Plugins/Experimental/AnimNextAnimGraph/Source/AnimNextAnimGraph/Public/TraitInterfaces/ISmoothBlendPerBone.h
--- a/Engine/Plugins/Experimental/AnimNextAnimGraph/Source/AnimNextAnimGraph/Public/TraitInterfaces/ISmoothBlendPerBone.h
+++ b/Engine/Plugins/Experimental/AnimNextAnimGraph/Source/AnimNextAnimGraph/Public/TraitInterfaces/ISmoothBlendPerBone.h
@@ -19,7 +19,7 @@ namespace UE::AnimNext
 	 */
 	struct ANIMNEXTANIMGRAPH_API ISmoothBlendPerBone : ITraitInterface
 	{
-		DECLARE_ANIM_TRAIT_INTERFACE(ISmoothBlendPerBone, 0xb29ee376)
+		DECLARE_ANIM_TRAIT_INTERFACE(ISmoothBlendPerBone)
 
 		// Returns the desired blend profile for the specified child
 		virtual TSharedPtr<const IBlendProfileInterface> GetBlendProfile(FExecutionContext& Context, const TTraitBinding<ISmoothBlendPerBone>& Binding, int32 ChildIndex) const;
diff --git a/Engine/Plugins/Experimental/AnimNextAnimGraph/Source/AnimNextAnimGraph/Public/TraitInterfaces/ITimeline.h b/Engine/Plugins/Experimental/AnimNextAnimGraph/Source/AnimNextAnimGraph/Public/TraitInterfaces/ITimeline.h
--- a/Engine/Plugins/Experimental/AnimNextAnimGraph/Source/AnimNextAnimGraph/Public/TraitInterfaces/ITimeline.h
+++ b/Engine/Plugins/Experimental/AnimNextAnimGraph/Source/AnimNextAnimGraph/Public/TraitInterfaces/ITimeline.h
@@ -190,7 +190,7 @@ namespace UE::AnimNext
 	 */
 	struct ANIMNEXTANIMGRAPH_API ITimeline : ITraitInterface
 	{
-		DECLARE_ANIM_TRAIT_INTERFACE(ITimeline, 0x53760727)
+		DECLARE_ANIM_TRAIT_INTERFACE(ITimeline)
 
 		// Returns the sync markers of this timeline sorted by time
 		virtual void GetSyncMarkers(const FExecutionContext& Context, const TTraitBinding<ITimeline>& Binding, FTimelineSyncMarkerArray& OutSyncMarkers) const;
diff --git a/Engine/Plugins/Experimental/AnimNextAnimGraph/Source/AnimNextAnimGraph/Public/TraitInterfaces/ITimelinePlayer.h b/Engine/Plugins/Experimental/AnimNextAnimGraph/Source/AnimNextAnimGraph/Public/TraitInterfaces/ITimelinePlayer.h
--- a/Engine/Plugins/Experimental/AnimNextAnimGraph/Source/AnimNextAnimGraph/Public/TraitInterfaces/ITimelinePlayer.h
+++ b/Engine/Plugins/Experimental/AnimNextAnimGraph/Source/AnimNextAnimGraph/Public/TraitInterfaces/ITimelinePlayer.h
@@ -15,7 +15,7 @@ namespace UE::AnimNext
 	 */
 	struct ANIMNEXTANIMGRAPH_API ITimelinePlayer : ITraitInterface
 	{
-		DECLARE_ANIM_TRAIT_INTERFACE(ITimelinePlayer, 0xce8c6566)
+		DECLARE_ANIM_TRAIT_INTERFACE(ITimelinePlayer)
 
 		// Advances time by the provided delta time (positive or negative) on this timeline player
 		virtual void AdvanceBy(FExecutionContext& Context, const TTraitBinding<ITimelinePlayer>& Binding, float DeltaTime) const;
diff --git a/Engine/Plugins/Experimental/AnimNextAnimGraph/Source/AnimNextAnimGraph/Public/TraitInterfaces/IUpdate.h b/Engine/Plugins/Experimental/AnimNextAnimGraph/Source/AnimNextAnimGraph/Public/TraitInterfaces/IUpdate.h
--- a/Engine/Plugins/Experimental/AnimNextAnimGraph/Source/AnimNextAnimGraph/Public/TraitInterfaces/IUpdate.h
+++ b/Engine/Plugins/Experimental/AnimNextAnimGraph/Source/AnimNextAnimGraph/Public/TraitInterfaces/IUpdate.h
@@ -228,7 +228,7 @@ namespace UE::AnimNext
 	 */
 	struct ANIMNEXTANIMGRAPH_API IUpdate : ITraitInterface
 	{
-		DECLARE_ANIM_TRAIT_INTERFACE(IUpdate, 0x59d24dc5)
+		DECLARE_ANIM_TRAIT_INTERFACE(IUpdate)
 
 		// Called before the first update when a trait stack becomes relevant
 		virtual void OnBecomeRelevant(FUpdateTraversalContext& Context, const TTraitBinding<IUpdate>& Binding, const FTraitUpdateState& TraitState) const;
@@ -259,7 +259,7 @@ namespace UE::AnimNext
 	 */
 	struct ANIMNEXTANIMGRAPH_API IUpdateTraversal : ITraitInterface
 	{
-		DECLARE_ANIM_TRAIT_INTERFACE(IUpdateTraversal, 0x256c21b1)
+		DECLARE_ANIM_TRAIT_INTERFACE(IUpdateTraversal)
 
 		// Called after PreUpdate to request that children be queued with the provided context
 		virtual void QueueChildrenForTraversal(FUpdateTraversalContext& Context, const TTraitBinding<IUpdateTraversal>& Binding, const FTraitUpdateState& TraitState, FUpdateTraversalQueue& TraversalQueue) const;
diff --git a/Engine/Plugins/Experimental/AnimNextAnimGraph/Source/AnimNextAnimGraphTestSuite/Private/AnimNextAnimGraphTraitInterfacesTest.cpp b/Engine/Plugins/Experimental/AnimNextAnimGraph/Source/AnimNextAnimGraphTestSuite/Private/AnimNextAnimGraphTraitInterfacesTest.cpp
--- a/Engine/Plugins/Experimental/AnimNextAnimGraph/Source/AnimNextAnimGraphTestSuite/Private/AnimNextAnimGraphTraitInterfacesTest.cpp
+++ b/Engine/Plugins/Experimental/AnimNextAnimGraph/Source/AnimNextAnimGraphTestSuite/Private/AnimNextAnimGraphTraitInterfacesTest.cpp
@@ -231,7 +231,7 @@ namespace UE::AnimNext
 
 	struct IScopedTagInterface : IScopedTraitInterface
 	{
-		DECLARE_ANIM_TRAIT_INTERFACE(IScopedTagInterface, 0xb8f4b397)
+		DECLARE_ANIM_TRAIT_INTERFACE(IScopedTagInterface)
 
 		virtual FName GetTag(const FExecutionContext& Context, const TTraitBinding<IScopedTagInterface>& Binding) const;
 
@@ -359,7 +359,7 @@ namespace UE::AnimNext
 
 	struct ITraitInterfaceTestA : ITraitInterface
 	{
-		DECLARE_ANIM_TRAIT_INTERFACE(ITraitInterfaceTestA, 0xc6d4bde9)
+		DECLARE_ANIM_TRAIT_INTERFACE(ITraitInterfaceTestA)
 
 	#if WITH_EDITOR
 		virtual const FText& GetDisplayName() const override
@@ -379,7 +379,7 @@ namespace UE::AnimNext
 
 	struct ITraitInterfaceTestB : ITraitInterface
 	{
-		DECLARE_ANIM_TRAIT_INTERFACE(ITraitInterfaceTestB, 0xc0eadb31)
+		DECLARE_ANIM_TRAIT_INTERFACE(ITraitInterfaceTestB)
 
 	#if WITH_EDITOR
 		virtual const FText& GetDisplayName() const override
diff --git a/Engine/Plugins/Experimental/AnimNextAnimGraph/Source/AnimNextAnimGraphTestSuite/Private/AnimNextTraitBaseTest.cpp b/Engine/Plugins/Experimental/AnimNextAnimGraph/Source/AnimNextAnimGraphTestSuite/Private/AnimNextTraitBaseTest.cpp
--- a/Engine/Plugins/Experimental/AnimNextAnimGraph/Source/AnimNextAnimGraphTestSuite/Private/AnimNextTraitBaseTest.cpp
+++ b/Engine/Plugins/Experimental/AnimNextAnimGraph/Source/AnimNextAnimGraphTestSuite/Private/AnimNextTraitBaseTest.cpp
@@ -36,7 +36,7 @@ namespace UE::AnimNext
 
 	struct IInterfaceA : ITraitInterface
 	{
-		DECLARE_ANIM_TRAIT_INTERFACE(IInterfaceA, 0x34cb8e62)
+		DECLARE_ANIM_TRAIT_INTERFACE(IInterfaceA)
 
 		virtual void FuncA(const FExecutionContext& Context, const TTraitBinding<IInterfaceA>& Binding) const;
 	};
@@ -66,7 +66,7 @@ namespace UE::AnimNext
 
 	struct IInterfaceB : ITraitInterface
 	{
-		DECLARE_ANIM_TRAIT_INTERFACE(IInterfaceB, 0x33cb8ccf)
+		DECLARE_ANIM_TRAIT_INTERFACE(IInterfaceB)
 
 		virtual void FuncB(const FExecutionContext& Context, const TTraitBinding<IInterfaceB>& Binding) const;
 	};
@@ -96,7 +96,7 @@ namespace UE::AnimNext
 
 	struct IInterfaceC : ITraitInterface
 	{
-		DECLARE_ANIM_TRAIT_INTERFACE(IInterfaceC, 0x32cb8b3c)
+		DECLARE_ANIM_TRAIT_INTERFACE(IInterfaceC)
 
 		virtual void FuncC(const FExecutionContext& Context, const TTraitBinding<IInterfaceC>& Binding) const;
 	};
diff --git a/Engine/Plugins/Experimental/AnimNextPoseSearch/Source/AnimNextPoseSearch/Private/IPoseHistory.h b/Engine/Plugins/Experimental/AnimNextPoseSearch/Source/AnimNextPoseSearch/Private/IPoseHistory.h
--- a/Engine/Plugins/Experimental/AnimNextPoseSearch/Source/AnimNextPoseSearch/Private/IPoseHistory.h
+++ b/Engine/Plugins/Experimental/AnimNextPoseSearch/Source/AnimNextPoseSearch/Private/IPoseHistory.h
@@ -18,7 +18,7 @@ namespace UE::AnimNext
 	 */
 	struct IPoseHistory : ITraitInterface
 	{
-		DECLARE_ANIM_TRAIT_INTERFACE(IPoseHistory, 0x1c2cffc4)
+		DECLARE_ANIM_TRAIT_INTERFACE(IPoseHistory)
 
 		virtual const UE::PoseSearch::IPoseHistory* GetPoseHistory(FExecutionContext& Context, const TTraitBinding<IPoseHistory>& Binding) const;
 
