Add error checking if latent properties aren't correctly set up

diff --git a/Engine/Plugins/Experimental/AnimNextAnimGraph/Source/AnimNextAnimGraph/Private/Graph/RigDecorator_AnimNextCppTrait.cpp b/Engine/Plugins/Experimental/AnimNextAnimGraph/Source/AnimNextAnimGraph/Private/Graph/RigDecorator_AnimNextCppTrait.cpp
--- a/Engine/Plugins/Experimental/AnimNextAnimGraph/Source/AnimNextAnimGraph/Private/Graph/RigDecorator_AnimNextCppTrait.cpp
+++ b/Engine/Plugins/Experimental/AnimNextAnimGraph/Source/AnimNextAnimGraph/Private/Graph/RigDecorator_AnimNextCppTrait.cpp
@@ -50,7 +50,8 @@ void FRigDecorator_AnimNextCppDecorator::GetProgrammaticPins(URigVMController* I
 	// Generate the shared data struct pins
 	OutPinArray.AddPins(DecoratorSharedDataStruct, InController, ERigVMPinDirection::Invalid, InParentPinIndex, DefaultValueTypeGetter, DefaultValueMemoryScope.GetStructMemory(), true);
 
-	for (int32 PinIndex = StartPinIndex; PinIndex < OutPinArray.Num(); ++PinIndex)
+	const int32 NumPins = OutPinArray.Num();
+	for (int32 PinIndex = StartPinIndex; PinIndex < NumPins; ++PinIndex)
 	{
 		FRigVMPinInfo& PinInfo = OutPinArray[PinIndex];
 
@@ -74,7 +75,20 @@ void FRigDecorator_AnimNextCppDecorator::GetProgrammaticPins(URigVMController* I
 		// Except for trait handles which are never lazy since they just encode graph connectivity
 		if (InParentPinIndex == PinInfo.ParentIndex && !bIsHidden && !bIsInline && !bIsTraitHandle)
 		{
-			PinInfo.bIsLazy = true;
+			if (Trait->IsPropertyLatent(PinInfo.Property->GetFName()))
+			{
+				PinInfo.bIsLazy = true;
+			}
+			else
+			{
+				// This can occur when a property is latent (default) and wasn't included in the GENERATE_TRAIT_LATENT_PROPERTIES enumerator
+				InController->ReportErrorf(TEXT("Shared data property '%s' is latent but is missing from the latent property enumerator"), *PinInfo.Name.ToString());
+				PinInfo.bIsLazy = false;
+			}
+		}
+		else
+		{
+			PinInfo.bIsLazy = false;
 		}
 
 		// Remove our property because we configure the pin explicitly
diff --git a/Engine/Plugins/Experimental/AnimNextAnimGraph/Source/AnimNextAnimGraph/Private/TraitCore/Trait.cpp b/Engine/Plugins/Experimental/AnimNextAnimGraph/Source/AnimNextAnimGraph/Private/TraitCore/Trait.cpp
--- a/Engine/Plugins/Experimental/AnimNextAnimGraph/Source/AnimNextAnimGraph/Private/TraitCore/Trait.cpp
+++ b/Engine/Plugins/Experimental/AnimNextAnimGraph/Source/AnimNextAnimGraph/Private/TraitCore/Trait.cpp
@@ -175,6 +175,14 @@ namespace UE::AnimNext
 					continue;	// Skip non-latent properties
 				}
 
+				if (!IsPropertyLatent(Property->GetFName()))
+				{
+					// Skip properties not part of the latent macro enumerator
+					// This can occur if the property is missing from the enumerator by mistake
+					// in which case we'll warn during compilation
+					continue;
+				}
+
 				FLatentPropertyMetadata Metadata;
 				Metadata.Name = Property->GetFName();
 				Metadata.RigVMIndex = GetTraitLatentPropertyIndex(Property->GetFName());
diff --git a/Engine/Plugins/Experimental/AnimNextAnimGraph/Source/AnimNextAnimGraph/Public/TraitCore/Trait.h b/Engine/Plugins/Experimental/AnimNextAnimGraph/Source/AnimNextAnimGraph/Public/TraitCore/Trait.h
--- a/Engine/Plugins/Experimental/AnimNextAnimGraph/Source/AnimNextAnimGraph/Public/TraitCore/Trait.h
+++ b/Engine/Plugins/Experimental/AnimNextAnimGraph/Source/AnimNextAnimGraph/Public/TraitCore/Trait.h
@@ -43,9 +43,16 @@ struct FRigVMPinInfoArray;
 	virtual UE::AnimNext::ETraitStackPropagation OnTraitEvent(UE::AnimNext::FExecutionContext& Context, UE::AnimNext::FTraitBinding& Binding, FAnimNextTraitEvent& Event) const override; \
 	virtual TConstArrayView<UE::AnimNext::FTraitEventUID> GetTraitEvents() const override;
 
+#if WITH_EDITOR
+#define ANIM_NEXT_IMPL_DECLARE_ANIM_TRAIT_LATENT_PROPERTY_SUPPORT() \
+	virtual uint32 GetNumLatentTraitProperties() const override { return -FSharedData::GetLatentPropertyIndex(~(size_t)0); } \
+	virtual UE::AnimNext::FTraitLatentPropertyMemoryLayout GetLatentPropertyMemoryLayout(const FAnimNextTraitSharedData& SharedData, FName PropertyName, uint32 PropertyIndex) const override; \
+	virtual bool IsPropertyLatent(FName PropertyName) const override;
+#else
 #define ANIM_NEXT_IMPL_DECLARE_ANIM_TRAIT_LATENT_PROPERTY_SUPPORT() \
 	virtual uint32 GetNumLatentTraitProperties() const override { return -FSharedData::GetLatentPropertyIndex(~(size_t)0); } \
 	virtual UE::AnimNext::FTraitLatentPropertyMemoryLayout GetLatentPropertyMemoryLayout(const FAnimNextTraitSharedData& SharedData, FName PropertyName, uint32 PropertyIndex) const override;
+#endif
 
 // Helper macros
 // In the trait class declaration, this macro declares the Super alias and base functions we override
@@ -108,6 +115,19 @@ struct FRigVMPinInfoArray;
 		return GetLatentPropertyMemoryLayoutImpl(PropertyName, PropertyIndex, CachedLatentPropertyMemoryLayouts); \
 	}
 
+#if WITH_EDITOR
+#define ANIM_NEXT_IMPL_DEFINE_ANIM_TRAIT_IS_PROPERTY_LATENT(TraitName) \
+	bool TraitName::IsPropertyLatent(FName PropertyName) const \
+	{ \
+		const UScriptStruct* SharedDataStruct = FSharedData::StaticStruct(); \
+		const FProperty* Property = SharedDataStruct->FindPropertyByName(PropertyName); \
+		check(Property != nullptr); \
+		return FSharedData::GetLatentPropertyIndex(Property->GetOffset_ForInternal()) > 0; \
+	}
+#else
+#define ANIM_NEXT_IMPL_DEFINE_ANIM_TRAIT_IS_PROPERTY_LATENT(TraitName)
+#endif
+
 // Helper that handles the GetTraitInterface() details for each interface specified by the generator macro
 #define ANIM_NEXT_IMPL_GET_INTERFACE_IMPL_FOR_INTERFACE(InterfaceName) \
 	if (InInterfaceUID == InterfaceName::InterfaceUID) \
@@ -245,6 +265,7 @@ namespace UE::AnimNext::Private
 #define GENERATE_ANIM_TRAIT_IMPLEMENTATION(TraitName, InterfaceEnumeratorMacro, RequiredInterfaceEnumeratorMacro, EventEnumeratorMacro) \
 	ANIM_NEXT_IMPL_DEFINE_ANIM_TRAIT(TraitName) \
 	ANIM_NEXT_IMPL_DEFINE_ANIM_TRAIT_GET_LATENT_PROPERTY_MEMORY_LAYOUT(TraitName) \
+	ANIM_NEXT_IMPL_DEFINE_ANIM_TRAIT_IS_PROPERTY_LATENT(TraitName) \
 	ANIM_NEXT_IMPL_DEFINE_ANIM_TRAIT_GET_INTERFACE(TraitName, InterfaceEnumeratorMacro) \
 	ANIM_NEXT_IMPL_DEFINE_ANIM_TRAIT_GET_INTERFACES(TraitName, InterfaceEnumeratorMacro) \
 	ANIM_NEXT_IMPL_DEFINE_ANIM_TRAIT_GET_REQUIRED_INTERFACES(TraitName, RequiredInterfaceEnumeratorMacro) \
@@ -430,9 +451,13 @@ namespace UE::AnimNext
 			bool bFilterEditorOnly,
 			const TFunction<uint16(FName PropertyName)>& GetTraitLatentPropertyIndex) const;
 
+		// Returns whether or not the specified property has been marked as latent
+		virtual bool IsPropertyLatent(FName PropertyName) const { return false; }
+
 		// Makes the Trait Editor only display this Trait in Advanced view
 		virtual bool IsHidden() const { return false; }
 
+		// Whether or not this trait can be placed multiple times on a trait stack
 		virtual bool MultipleInstanceSupport() const { return false; }
 
 		// Enable Traits to generate TraitStack pins
