AnimNext: Fix notify dispatcher adding notifies twice

Mark the event as consumed so we dont process it again

diff --git a/Engine/Plugins/Experimental/AnimNextAnimGraph/Source/AnimNextAnimGraph/Private/Traits/NotifyDispatcher.cpp b/Engine/Plugins/Experimental/AnimNextAnimGraph/Source/AnimNextAnimGraph/Private/Traits/NotifyDispatcher.cpp
--- a/Engine/Plugins/Experimental/AnimNextAnimGraph/Source/AnimNextAnimGraph/Private/Traits/NotifyDispatcher.cpp
+++ b/Engine/Plugins/Experimental/AnimNextAnimGraph/Source/AnimNextAnimGraph/Private/Traits/NotifyDispatcher.cpp
@@ -130,6 +130,7 @@ void FAnimNextNotifyDispatcherComponent::OnTraitEvent(FAnimNextTraitEvent& Event
 	if(UE::AnimNext::FNotifyDispatchEvent* NotifyDispatchEvent = Event.AsType<UE::AnimNext::FNotifyDispatchEvent>())
 	{
 		NotifyQueue.AddAnimNotifies(NotifyDispatchEvent->Notifies, NotifyDispatchEvent->Weight);
+		NotifyDispatchEvent->MarkConsumed();
 	}
 }
 
