Add AlphaInputArgs Additive so traits don't need to handle alpha input.

#tests PIE

diff --git a/Engine/Plugins/Experimental/AnimNextAnimGraph/Source/AnimNextAnimGraph/Private/EvaluationVM/Tasks/ApplyAdditiveKeyframe.cpp b/Engine/Plugins/Experimental/AnimNextAnimGraph/Source/AnimNextAnimGraph/Private/EvaluationVM/Tasks/ApplyAdditiveKeyframe.cpp
--- a/Engine/Plugins/Experimental/AnimNextAnimGraph/Source/AnimNextAnimGraph/Private/EvaluationVM/Tasks/ApplyAdditiveKeyframe.cpp
+++ b/Engine/Plugins/Experimental/AnimNextAnimGraph/Source/AnimNextAnimGraph/Private/EvaluationVM/Tasks/ApplyAdditiveKeyframe.cpp
@@ -14,6 +14,15 @@ FAnimNextApplyAdditiveKeyframeTask FAnimNextApplyAdditiveKeyframeTask::Make(floa
 	return Task;
 }
 
+FAnimNextApplyAdditiveKeyframeTask FAnimNextApplyAdditiveKeyframeTask::Make(const FName& AlphaSourceCurveName, const int8 AlphaCurveInputIndex, TFunction<float(float)> InputScaleBiasClampFn)
+{
+	FAnimNextApplyAdditiveKeyframeTask Task;
+	Task.AlphaSourceCurveName = AlphaSourceCurveName;
+	Task.AlphaCurveInputIndex = AlphaCurveInputIndex;
+	Task.InputScaleBiasClampFn = MoveTemp(InputScaleBiasClampFn);
+	return Task;
+}
+
 void FAnimNextApplyAdditiveKeyframeTask::Execute(UE::AnimNext::FEvaluationVM& VM) const
 {
 	using namespace UE::AnimNext;
@@ -43,6 +52,8 @@ void FAnimNextApplyAdditiveKeyframeTask::Execute(UE::AnimNext::FEvaluationVM& VM
 		return;
 	}
 
+	const float AdditiveWeight = GetInterpolationAlpha(BaseKeyframe.Get(), AdditiveKeyframe.Get());
+
 	if (EnumHasAnyFlags(VM.GetFlags(), EEvaluationFlags::Bones))
 	{
 		check(BaseKeyframe->Pose.GetNumBones() == AdditiveKeyframe->Pose.GetNumBones());
@@ -53,11 +64,11 @@ void FAnimNextApplyAdditiveKeyframeTask::Execute(UE::AnimNext::FEvaluationVM& VM
 		{
 			BlendWithIdentityAndAccumulateMesh(
 				BaseTransformsView, AdditiveKeyframe->Pose.LocalTransforms.GetConstView(), 
-				AdditiveKeyframe->Pose.GetLODBoneIndexToParentLODBoneIndexMap(), BlendWeight);
+				AdditiveKeyframe->Pose.GetLODBoneIndexToParentLODBoneIndexMap(), AdditiveWeight);
 		}
 		else
 		{
-			BlendWithIdentityAndAccumulate(BaseTransformsView, AdditiveKeyframe->Pose.LocalTransforms.GetConstView(), BlendWeight);
+			BlendWithIdentityAndAccumulate(BaseTransformsView, AdditiveKeyframe->Pose.LocalTransforms.GetConstView(), AdditiveWeight);
 		}
 
 		NormalizeRotations(BaseTransformsView);
@@ -65,13 +76,34 @@ void FAnimNextApplyAdditiveKeyframeTask::Execute(UE::AnimNext::FEvaluationVM& VM
 
 	if (EnumHasAnyFlags(VM.GetFlags(), EEvaluationFlags::Curves))
 	{
-		BaseKeyframe->Curves.Accumulate(AdditiveKeyframe->Curves, BlendWeight);
+		BaseKeyframe->Curves.Accumulate(AdditiveKeyframe->Curves, AdditiveWeight);
 	}
 
 	if (EnumHasAnyFlags(VM.GetFlags(), EEvaluationFlags::Attributes))
 	{
-		UE::Anim::Attributes::AccumulateAttributes(AdditiveKeyframe->Attributes, BaseKeyframe->Attributes, BlendWeight, AAT_None);
+		UE::Anim::Attributes::AccumulateAttributes(AdditiveKeyframe->Attributes, BaseKeyframe->Attributes, AdditiveWeight, AAT_None);
 	}
 
 	VM.PushValue(KEYFRAME_STACK_NAME, MoveTemp(BaseKeyframe));
 }
+
+float FAnimNextApplyAdditiveKeyframeTask::GetInterpolationAlpha(const UE::AnimNext::FKeyframeState* KeyframeA, const UE::AnimNext::FKeyframeState* KeyframeB) const
+{
+	float Alpha = BlendWeight;
+
+	if (AlphaSourceCurveName != NAME_None && AlphaCurveInputIndex != INDEX_NONE)
+	{
+		if (ensure(KeyframeA != nullptr && KeyframeB != nullptr))
+		{
+			const FBlendedCurve& Curves = AlphaCurveInputIndex == 0 ? KeyframeA->Curves : KeyframeB->Curves;
+			Alpha = Curves.Get(AlphaSourceCurveName); // if the curve does not exist, it returns 0.f
+
+			if (InputScaleBiasClampFn.IsSet())
+			{
+				Alpha = InputScaleBiasClampFn(Alpha);
+			}
+		}
+	}
+
+	return FMath::Clamp(Alpha, 0.0f, 1.0f);
+}
diff --git a/Engine/Plugins/Experimental/AnimNextAnimGraph/Source/AnimNextAnimGraph/Private/TraitInterfaces/Args/IAlphaInputArgs.cpp b/Engine/Plugins/Experimental/AnimNextAnimGraph/Source/AnimNextAnimGraph/Private/TraitInterfaces/Args/IAlphaInputArgs.cpp
new file mode 100644
--- /dev/null
+++ b/Engine/Plugins/Experimental/AnimNextAnimGraph/Source/AnimNextAnimGraph/Private/TraitInterfaces/Args/IAlphaInputArgs.cpp
@@ -0,0 +1,83 @@
+// Copyright Epic Games, Inc. All Rights Reserved.
+
+#include "TraitInterfaces/Args/IAlphaInputArgs.h"
+
+#include "TraitCore/ExecutionContext.h"
+
+namespace UE::AnimNext
+{
+	AUTO_REGISTER_ANIM_TRAIT_INTERFACE(IAlphaInputArgs)
+
+	FAlphaInputTraitArgs IAlphaInputArgs::Get(const FExecutionContext& Context, const TTraitBinding<IAlphaInputArgs>& Binding) const
+	{
+		TTraitBinding<IAlphaInputArgs> SuperBinding;
+		if (Binding.GetStackInterfaceSuper(SuperBinding))
+		{
+			return SuperBinding.Get(Context);
+		}
+
+		return FAlphaInputTraitArgs();
+	}
+
+	EAnimAlphaInputType IAlphaInputArgs::GetAlphaInputType(const FExecutionContext& Context, const TTraitBinding<IAlphaInputArgs>& Binding) const
+	{
+		TTraitBinding<IAlphaInputArgs> SuperBinding;
+		if (Binding.GetStackInterfaceSuper(SuperBinding))
+		{
+			return SuperBinding.GetAlphaInputType(Context);
+		}
+
+		return EAnimAlphaInputType::Float;
+	}
+
+	FName IAlphaInputArgs::GetAlphaCurveName(const FExecutionContext& Context, const TTraitBinding<IAlphaInputArgs>& Binding) const
+	{
+		TTraitBinding<IAlphaInputArgs> SuperBinding;
+		if (Binding.GetStackInterfaceSuper(SuperBinding))
+		{
+			return SuperBinding.GetAlphaCurveName(Context);
+		}
+
+		return NAME_None;
+	}
+
+	float IAlphaInputArgs::GetCurrentAlphaValue(const FExecutionContext& Context, const TTraitBinding<IAlphaInputArgs>& Binding) const
+	{
+		TTraitBinding<IAlphaInputArgs> SuperBinding;
+		if (Binding.GetStackInterfaceSuper(SuperBinding))
+		{
+			return SuperBinding.GetCurrentAlphaValue(Context);
+		}
+
+		return 1.0f;
+	}
+
+	TFunction<float(float)> IAlphaInputArgs::GetInputScaleBiasClampCallback(const FExecutionContext& Context, const TTraitBinding<IAlphaInputArgs>& Binding) const
+	{
+		TTraitBinding<IAlphaInputArgs> SuperBinding;
+		if (Binding.GetStackInterfaceSuper(SuperBinding))
+		{
+			return SuperBinding.GetInputScaleBiasClampCallback(Context);
+		}
+
+		auto InputBiasClampCallback = [](float Alpha) -> float
+		{
+			return 1.0f;
+		};
+
+		return InputBiasClampCallback;
+	}
+
+#if WITH_EDITOR
+	const FText& IAlphaInputArgs::GetDisplayName() const
+	{
+		static FText InterfaceName = NSLOCTEXT("TraitInterfaces", "TraitInterface_AlphaInputArg_Name", "Alpha Input Args");
+		return InterfaceName;
+	}
+	const FText& IAlphaInputArgs::GetDisplayShortName() const
+	{
+		static FText InterfaceShortName = NSLOCTEXT("TraitInterfaces", "TraitInterface_AlphaInputArg_ShortName", "AIA");
+		return InterfaceShortName;
+	}
+#endif // WITH_EDITOR
+}
diff --git a/Engine/Plugins/Experimental/AnimNextAnimGraph/Source/AnimNextAnimGraph/Private/Traits/ApplyAdditiveTrait.cpp b/Engine/Plugins/Experimental/AnimNextAnimGraph/Source/AnimNextAnimGraph/Private/Traits/ApplyAdditiveTrait.cpp
--- a/Engine/Plugins/Experimental/AnimNextAnimGraph/Source/AnimNextAnimGraph/Private/Traits/ApplyAdditiveTrait.cpp
+++ b/Engine/Plugins/Experimental/AnimNextAnimGraph/Source/AnimNextAnimGraph/Private/Traits/ApplyAdditiveTrait.cpp
@@ -1,6 +1,7 @@
 // Copyright Epic Games, Inc. All Rights Reserved.
 
 #include "Traits/ApplyAdditiveTrait.h"
+#include "TraitInterfaces/Args/IAlphaInputArgs.h"
 
 #include "Animation/AnimTypes.h"
 #include "TraitCore/ExecutionContext.h"
@@ -18,6 +19,7 @@ namespace UE::AnimNext
 		GeneratorMacro(IUpdateTraversal) \
 		GeneratorMacro(IContinuousBlend) \
 
+	// Note: We don't require IAlphaInputArgs, as an IContinuousBlend can override our alpha
 	GENERATE_ANIM_TRAIT_IMPLEMENTATION(FApplyAdditiveTrait, TRAIT_INTERFACE_ENUMERATOR, NULL_ANIM_TRAIT_INTERFACE_ENUMERATOR, NULL_ANIM_TRAIT_EVENT_ENUMERATOR)
 	#undef TRAIT_INTERFACE_ENUMERATOR
 
@@ -45,7 +47,20 @@ namespace UE::AnimNext
 			Binding.GetStackInterface(ContinuousBlendTrait);
 
 			const float Alpha = ContinuousBlendTrait.GetBlendWeight(Context, 1);
-			Context.AppendTask(FAnimNextApplyAdditiveKeyframeTask::Make(Alpha));
+
+			// @TODO: Note in review, we could make this dissapear with some templating, But that approach doesn't scale well.
+			// Whenver the task may be very custom and have multiple considerations other than alpha. Expect consumers to perform this logic.
+			TTraitBinding<IAlphaInputArgs> AlphaInputArgsTrait;
+			if (Binding.GetStackInterface(AlphaInputArgsTrait) && AlphaInputArgsTrait.GetAlphaInputType(Context) == EAnimAlphaInputType::Curve)
+			{
+				// Assume weight for additive is on the base, not additive itself
+				static constexpr int32 AlphaSourceInputKeyframeIndex = 0;
+				Context.AppendTask(FAnimNextApplyAdditiveKeyframeTask::Make(AlphaInputArgsTrait.GetAlphaCurveName(Context), AlphaSourceInputKeyframeIndex, AlphaInputArgsTrait.GetInputScaleBiasClampCallback(Context)));
+			}
+			else
+			{
+				Context.AppendTask(FAnimNextApplyAdditiveKeyframeTask::Make(Alpha));
+			}
 		}
 		else
 		{
@@ -121,7 +136,20 @@ namespace UE::AnimNext
 		const FSharedData* SharedData = Binding.GetSharedData<FSharedData>();
 		check(SharedData);
 
-		const float Alpha = SharedData->GetAlpha(Binding);
+		FInstanceData* InstanceData = Binding.GetInstanceData<FInstanceData>();
+		check(InstanceData);
+
+		float Alpha = 1.0f;
+		TTraitBinding<IAlphaInputArgs> AlphaInputArgsTrait;
+		if (Binding.GetStackInterface(AlphaInputArgsTrait))
+		{
+			Alpha = AlphaInputArgsTrait.GetCurrentAlphaValue(Context);
+		}
+		else
+		{
+			// @TODO: Fallback path, to be removed
+			Alpha = SharedData->GetAlpha(Binding);
+		}
 
 		if (ChildIndex == 0)
 		{
diff --git a/Engine/Plugins/Experimental/AnimNextAnimGraph/Source/AnimNextAnimGraph/Private/Traits/ApplyAdditiveTrait.h b/Engine/Plugins/Experimental/AnimNextAnimGraph/Source/AnimNextAnimGraph/Private/Traits/ApplyAdditiveTrait.h
--- a/Engine/Plugins/Experimental/AnimNextAnimGraph/Source/AnimNextAnimGraph/Private/Traits/ApplyAdditiveTrait.h
+++ b/Engine/Plugins/Experimental/AnimNextAnimGraph/Source/AnimNextAnimGraph/Private/Traits/ApplyAdditiveTrait.h
@@ -26,7 +26,12 @@ struct FAnimNextApplyAdditiveTraitSharedData : public FAnimNextTraitSharedData
 	UPROPERTY()
 	FAnimNextTraitHandle Additive;
 
-	/** How much to apply our addtive, default is 1. */
+	/** 
+	 * Deprecated. Please add an IAlphaInputArgs addtive trait to set alpha.
+	 * @TODO: Remove pre 5.6 once removing latents doesn't cause crash
+	 * 
+	 * How much to apply our addtive, default is 1. 
+	 */
 	UPROPERTY(EditAnywhere, Category = "Default")
 	float Alpha = 1.0f;
 
diff --git a/Engine/Plugins/Experimental/AnimNextAnimGraph/Source/AnimNextAnimGraph/Private/Traits/Args/AlphaInputArgCoreTrait.cpp b/Engine/Plugins/Experimental/AnimNextAnimGraph/Source/AnimNextAnimGraph/Private/Traits/Args/AlphaInputArgCoreTrait.cpp
new file mode 100644
--- /dev/null
+++ b/Engine/Plugins/Experimental/AnimNextAnimGraph/Source/AnimNextAnimGraph/Private/Traits/Args/AlphaInputArgCoreTrait.cpp
@@ -0,0 +1,179 @@
+// Copyright Epic Games, Inc. All Rights Reserved.
+
+#include "Traits/Args/AlphaInputArgCoreTrait.h"
+
+#include "Animation/AnimTypes.h"
+#include "TraitCore/ExecutionContext.h"
+#include "EvaluationVM/Tasks/BlendKeyframes.h"
+#include "EvaluationVM/Tasks/StoreKeyframe.h"
+
+namespace UE::AnimNext
+{
+
+namespace AlphaInput
+{
+
+float ComputeAlphaValueForFloat(FInputScaleBiasClamp& InAlphaScaleBiasClamp
+	, const FInputScaleBias& InAlphaScaleBias
+	, const float InBaseAlpha
+	, const float InDeltaTime)
+{
+	float Result = InAlphaScaleBias.ApplyTo(InAlphaScaleBiasClamp.ApplyTo(InBaseAlpha, InDeltaTime));
+	return FMath::Clamp(Result, 0.0f, 1.0f);
+}
+
+float ComputeAlphaValueForBool(FInputAlphaBoolBlend& InAlphaBoolBlend
+	, const bool bInAlphaBoolEnabled
+	, const float InDeltaTime)
+{
+	float Result = InAlphaBoolBlend.ApplyTo(bInAlphaBoolEnabled, InDeltaTime);
+	return FMath::Clamp(Result, 0.0f, 1.0f);
+}
+
+float ComputeAlphaValueForType(const EAnimAlphaInputType InAlphaInputType
+	, FInputScaleBiasClamp& InAlphaScaleBiasClamp
+	, const FInputScaleBias& InAlphaScaleBias
+	, const float InBaseAlpha
+	, FInputAlphaBoolBlend& InAlphaBoolBlend
+	, const bool bInAlphaBoolEnabled
+	, const float InDeltaTime)
+{
+	float Result = 0.f;
+
+	switch (InAlphaInputType)
+	{
+		case EAnimAlphaInputType::Float:
+		{
+			Result = AlphaInput::ComputeAlphaValueForFloat(IN OUT InAlphaScaleBiasClamp, InAlphaScaleBias, InBaseAlpha, InDeltaTime);
+			break;
+		}
+		case EAnimAlphaInputType::Bool:
+		{
+			Result = AlphaInput::ComputeAlphaValueForBool(IN OUT InAlphaBoolBlend, bInAlphaBoolEnabled, InDeltaTime);
+			break;
+		}
+		case EAnimAlphaInputType::Curve:
+		{
+			Result = 1.0; // This will be calculated on the task when reading the curve, but we need full weight for branches
+			break;
+		}
+	};
+
+	return Result;
+}
+
+} // namespace AlphaInput
+
+
+////////////////////////////////////////////////////////////////////////////////////////////////////////////////////////
+// FAlphaInputArgCoreTrait
+
+AUTO_REGISTER_ANIM_TRAIT(FAlphaInputArgCoreTrait)
+
+#define TRAIT_INTERFACE_ENUMERATOR(GeneratorMacro) \
+	GeneratorMacro(IAlphaInputArgs) \
+	GeneratorMacro(IUpdate) \
+
+// Trait implementation boilerplate
+GENERATE_ANIM_TRAIT_IMPLEMENTATION(FAlphaInputArgCoreTrait, TRAIT_INTERFACE_ENUMERATOR, NULL_ANIM_TRAIT_INTERFACE_ENUMERATOR, NULL_ANIM_TRAIT_EVENT_ENUMERATOR)
+#undef TRAIT_INTERFACE_ENUMERATOR
+
+void FAlphaInputArgCoreTrait::FInstanceData::Construct(const FExecutionContext& Context, const FTraitBinding& Binding)
+{
+	const FSharedData* SharedData = Binding.GetSharedData<FSharedData>();
+
+	ComputedAlphaValue = SharedData->GetAlpha(Binding);
+	AlphaBoolBlend = SharedData->GetAlphaBoolBlend(Binding);
+	AlphaScaleBiasClamp = SharedData->AlphaScaleBiasClamp;
+}
+
+void FAlphaInputArgCoreTrait::OnBecomeRelevant(FUpdateTraversalContext& Context, const TTraitBinding<IUpdate>& Binding, const FTraitUpdateState& TraitState) const
+{
+	IUpdate::OnBecomeRelevant(Context, Binding, TraitState);
+
+	const FSharedData* SharedData = Binding.GetSharedData<FSharedData>();
+	FInstanceData* InstanceData = Binding.GetInstanceData<FInstanceData>();
+
+	InstanceData->ComputedAlphaValue = AlphaInput::ComputeAlphaValueForType(SharedData->GetAlphaInputType(Binding)
+		, IN OUT InstanceData->AlphaScaleBiasClamp
+		, SharedData->AlphaScaleBias
+		, SharedData->GetAlpha(Binding)
+		, IN OUT InstanceData->AlphaBoolBlend
+		, SharedData->GetbAlphaBoolEnabled(Binding)
+		, TraitState.GetDeltaTime());
+
+	InstanceData->AlphaBoolBlend.Reinitialize();
+	InstanceData->AlphaScaleBiasClamp.Reinitialize();
+}
+
+void FAlphaInputArgCoreTrait::PreUpdate(FUpdateTraversalContext& Context, const TTraitBinding<IUpdate>& Binding, const FTraitUpdateState& TraitState) const
+{
+	IUpdate::PreUpdate(Context, Binding, TraitState);
+
+	const FSharedData* SharedData = Binding.GetSharedData<FSharedData>();
+	FInstanceData* InstanceData = Binding.GetInstanceData<FInstanceData>();
+
+	// Store DeltaTime for deferred curve sampling
+	InstanceData->DeltaTime = TraitState.GetDeltaTime();
+	InstanceData->ComputedAlphaValue = AlphaInput::ComputeAlphaValueForType(SharedData->GetAlphaInputType(Binding)
+		, IN OUT InstanceData->AlphaScaleBiasClamp
+		, SharedData->AlphaScaleBias
+		, SharedData->GetAlpha(Binding)
+		, IN OUT InstanceData->AlphaBoolBlend
+		, SharedData->GetbAlphaBoolEnabled(Binding)
+		, TraitState.GetDeltaTime());
+}
+
+FAlphaInputTraitArgs FAlphaInputArgCoreTrait::Get(const FExecutionContext& Context, const TTraitBinding<IAlphaInputArgs>& Binding) const
+{
+	const FSharedData* SharedData = Binding.GetSharedData<FSharedData>();
+	FInstanceData* InstanceData = Binding.GetInstanceData<FInstanceData>();
+
+	FAlphaInputTraitArgs Result = FAlphaInputTraitArgs();
+	Result.Alpha = SharedData->GetAlpha(Binding);
+	Result.AlphaScaleBias = SharedData->AlphaScaleBias;
+	Result.AlphaScaleBiasClamp = InstanceData->AlphaScaleBiasClamp;
+	Result.bAlphaBoolEnabled = SharedData->GetbAlphaBoolEnabled(Binding);
+	Result.AlphaBoolBlend = InstanceData->AlphaBoolBlend;
+	Result.AlphaCurveName = SharedData->GetAlphaCurveName(Binding);
+	Result.AlphaInputType = SharedData->GetAlphaInputType(Binding);
+
+	return Result;
+}
+
+EAnimAlphaInputType FAlphaInputArgCoreTrait::GetAlphaInputType(const FExecutionContext& Context, const TTraitBinding<IAlphaInputArgs>& Binding) const
+{
+	const FSharedData* SharedData = Binding.GetSharedData<FSharedData>();
+
+	return SharedData->GetAlphaInputType(Binding);
+}
+
+FName FAlphaInputArgCoreTrait::GetAlphaCurveName(const FExecutionContext& Context, const TTraitBinding<IAlphaInputArgs>& Binding) const
+{
+	const FSharedData* SharedData = Binding.GetSharedData<FSharedData>();
+
+	return SharedData->GetAlphaCurveName(Binding);
+}
+
+float FAlphaInputArgCoreTrait::GetCurrentAlphaValue(const FExecutionContext& Context, const TTraitBinding<IAlphaInputArgs>& Binding) const
+{
+	FInstanceData* InstanceData = Binding.GetInstanceData<FInstanceData>();
+
+	return InstanceData->ComputedAlphaValue;
+}
+
+TFunction<float(float)> FAlphaInputArgCoreTrait::GetInputScaleBiasClampCallback(const FExecutionContext& Context, const TTraitBinding<IAlphaInputArgs>& Binding) const
+{
+	FInstanceData* InstanceData = Binding.GetInstanceData<FInstanceData>();
+
+	auto InputBiasClampCallback = [DeltaTime = InstanceData->DeltaTime, InputScaleBiasClamp = &InstanceData->AlphaScaleBiasClamp](float Alpha) -> float
+	{
+		return InputScaleBiasClamp->ApplyTo(Alpha, DeltaTime);
+	};
+
+	return InputBiasClampCallback;
+}
+
+
+} // namespace UE::AnimNext
+
diff --git a/Engine/Plugins/Experimental/AnimNextAnimGraph/Source/AnimNextAnimGraph/Private/Traits/Args/AlphaInputArgCoreTrait.h b/Engine/Plugins/Experimental/AnimNextAnimGraph/Source/AnimNextAnimGraph/Private/Traits/Args/AlphaInputArgCoreTrait.h
new file mode 100644
--- /dev/null
+++ b/Engine/Plugins/Experimental/AnimNextAnimGraph/Source/AnimNextAnimGraph/Private/Traits/Args/AlphaInputArgCoreTrait.h
@@ -0,0 +1,75 @@
+// Copyright Epic Games, Inc. All Rights Reserved.
+
+#pragma once
+
+#include "CoreMinimal.h"
+#include "TraitCore/Trait.h"
+#include "TraitCore/TraitSharedData.h"
+#include "TraitInterfaces/IUpdate.h"
+#include "TraitInterfaces/Args/IAlphaInputArgs.h"
+
+namespace UE::AnimNext
+{
+	namespace AlphaInput
+	{
+		extern ANIMNEXTANIMGRAPH_API float ComputeAlphaValueForFloat(FInputScaleBiasClamp& InAlphaScaleBiasClamp
+			, const FInputScaleBias& InAlphaScaleBias
+			, const float InBaseAlpha
+			, const float InDeltaTime);
+
+		extern ANIMNEXTANIMGRAPH_API float ComputeAlphaValueForBool(FInputAlphaBoolBlend& InAlphaBoolBlend
+			, const bool bInAlphaBoolEnabled
+			, const float InDeltaTime);
+
+		extern ANIMNEXTANIMGRAPH_API float ComputeAlphaValueForType(const EAnimAlphaInputType InAlphaInputType
+			, FInputScaleBiasClamp& InAlphaScaleBiasClamp
+			, const FInputScaleBias& InAlphaScaleBias
+			, const float InBaseAlpha
+			, FInputAlphaBoolBlend& InAlphaBoolBlend
+			, const bool bInAlphaBoolEnabled
+			, const float InDeltaTime);
+
+	} // namespace AlphaInput
+
+	/**
+	 * FAlphaInputArgCoreTrait
+	 *
+	 * Additive trait that provides configurable alpha input args for another trait to use. Ex: as a weight.
+	 * This trait's API calls mutate internal state, so it should not be used by multiple traits concurrently.
+	 * 
+	 * In particular, curve evaluation will not make sense as curves may differ per trait keyframe & DeltaTime
+	 * will be updated multiple per consumer trait.
+	 * 
+	 * Note: Does not implement IContinousBlend as cannot resolve trait-specific child index handling.
+	 */
+	USTRUCT(meta = (DisplayName = "Alpha Input Arg Core"))
+	struct FAlphaInputArgCoreTrait : FAdditiveTrait, IUpdate, IAlphaInputArgs
+	{
+		DECLARE_ANIM_TRAIT(FAlphaInputArgCoreTrait, FAdditiveTrait)
+
+		using FSharedData = FAlphaInputTraitArgs;
+		
+		struct FInstanceData : FTrait::FInstanceData
+		{
+			FInputAlphaBoolBlend AlphaBoolBlend;
+			FInputScaleBiasClamp AlphaScaleBiasClamp;
+
+			float DeltaTime = 0.f;
+			float ComputedAlphaValue = 0.f;
+
+			void Construct(const FExecutionContext& Context, const FTraitBinding& Binding);
+		};
+
+		// IUpdate
+		virtual void OnBecomeRelevant(FUpdateTraversalContext& Context, const TTraitBinding<IUpdate>& Binding, const FTraitUpdateState& TraitState) const override;
+		virtual void PreUpdate(FUpdateTraversalContext& Context, const TTraitBinding<IUpdate>& Binding, const FTraitUpdateState& TraitState) const override;
+
+		// IAlphaInputArgs
+		virtual FAlphaInputTraitArgs Get(const FExecutionContext& Context, const TTraitBinding<IAlphaInputArgs>& Binding) const override;
+		virtual EAnimAlphaInputType GetAlphaInputType(const FExecutionContext& Context, const TTraitBinding<IAlphaInputArgs>& Binding) const override;
+		virtual FName GetAlphaCurveName(const FExecutionContext& Context, const TTraitBinding<IAlphaInputArgs>& Binding) const override;
+		virtual float GetCurrentAlphaValue(const FExecutionContext& Context, const TTraitBinding<IAlphaInputArgs>& Binding) const override;
+		virtual TFunction<float(float)> GetInputScaleBiasClampCallback(const FExecutionContext& Context, const TTraitBinding<IAlphaInputArgs>& Binding) const override;
+	};
+
+} // namespace UE::AnimNext
diff --git a/Engine/Plugins/Experimental/AnimNextAnimGraph/Source/AnimNextAnimGraph/Private/Traits/Args/BoolAlphaInputArgTrait.cpp b/Engine/Plugins/Experimental/AnimNextAnimGraph/Source/AnimNextAnimGraph/Private/Traits/Args/BoolAlphaInputArgTrait.cpp
new file mode 100644
--- /dev/null
+++ b/Engine/Plugins/Experimental/AnimNextAnimGraph/Source/AnimNextAnimGraph/Private/Traits/Args/BoolAlphaInputArgTrait.cpp
@@ -0,0 +1,84 @@
+// Copyright Epic Games, Inc. All Rights Reserved.
+
+#include "Traits/Args/BoolAlphaInputArgTrait.h"
+#include "Traits/Args/AlphaInputArgCoreTrait.h"
+
+#include "Animation/AnimTypes.h"
+#include "TraitCore/ExecutionContext.h"
+#include "EvaluationVM/Tasks/BlendKeyframes.h"
+#include "EvaluationVM/Tasks/StoreKeyframe.h"
+
+namespace UE::AnimNext
+{
+
+////////////////////////////////////////////////////////////////////////////////////////////////////////////////////////
+// FBoolAlphaInputArgTrait
+
+AUTO_REGISTER_ANIM_TRAIT(FBoolAlphaInputArgTrait)
+
+#define TRAIT_INTERFACE_ENUMERATOR(GeneratorMacro) \
+	GeneratorMacro(IAlphaInputArgs) \
+	GeneratorMacro(IUpdate) \
+
+// Trait implementation boilerplate
+GENERATE_ANIM_TRAIT_IMPLEMENTATION(FBoolAlphaInputArgTrait, TRAIT_INTERFACE_ENUMERATOR, NULL_ANIM_TRAIT_INTERFACE_ENUMERATOR, NULL_ANIM_TRAIT_EVENT_ENUMERATOR)
+#undef TRAIT_INTERFACE_ENUMERATOR
+
+void FBoolAlphaInputArgTrait::FInstanceData::Construct(const FExecutionContext& Context, const FTraitBinding& Binding)
+{
+	const FSharedData* SharedData = Binding.GetSharedData<FSharedData>();
+
+	AlphaBoolBlend = SharedData->GetAlphaBoolBlend(Binding);
+}
+
+void FBoolAlphaInputArgTrait::OnBecomeRelevant(FUpdateTraversalContext& Context, const TTraitBinding<IUpdate>& Binding, const FTraitUpdateState& TraitState) const
+{
+	IUpdate::OnBecomeRelevant(Context, Binding, TraitState);
+
+	const FSharedData* SharedData = Binding.GetSharedData<FSharedData>();
+	FInstanceData* InstanceData = Binding.GetInstanceData<FInstanceData>();
+
+	InstanceData->ComputedAlphaValue = AlphaInput::ComputeAlphaValueForBool(IN OUT InstanceData->AlphaBoolBlend, SharedData->GetbAlphaBoolEnabled(Binding), TraitState.GetDeltaTime());
+	
+	InstanceData->AlphaBoolBlend.Reinitialize();
+}
+
+void FBoolAlphaInputArgTrait::PreUpdate(FUpdateTraversalContext& Context, const TTraitBinding<IUpdate>& Binding, const FTraitUpdateState& TraitState) const
+{
+	const FSharedData* SharedData = Binding.GetSharedData<FSharedData>();
+	FInstanceData* InstanceData = Binding.GetInstanceData<FInstanceData>();
+
+	InstanceData->ComputedAlphaValue = AlphaInput::ComputeAlphaValueForBool(IN OUT InstanceData->AlphaBoolBlend, SharedData->GetbAlphaBoolEnabled(Binding), TraitState.GetDeltaTime());
+}
+
+FAlphaInputTraitArgs FBoolAlphaInputArgTrait::Get(const FExecutionContext& Context, const TTraitBinding<IAlphaInputArgs>& Binding) const
+{
+	const FSharedData* SharedData = Binding.GetSharedData<FSharedData>();
+	FInstanceData* InstanceData = Binding.GetInstanceData<FInstanceData>();
+
+	FAlphaInputTraitArgs Result = FAlphaInputTraitArgs();
+	// Result.Alpha = InstanceData->InternalAlpha;
+	// Result.AlphaScaleBias = SharedData->AlphaScaleBias;
+	// Result.AlphaScaleBiasClamp = InstanceData->AlphaScaleBiasClamp;
+	Result.bAlphaBoolEnabled = SharedData->GetbAlphaBoolEnabled(Binding);
+	Result.AlphaBoolBlend = InstanceData->AlphaBoolBlend;
+	// Result.AlphaCurveName = SharedData->GetAlphaCurveName(Binding);
+	Result.AlphaInputType = EAnimAlphaInputType::Bool;
+
+	return Result;
+}
+
+EAnimAlphaInputType FBoolAlphaInputArgTrait::GetAlphaInputType(const FExecutionContext& Context, const TTraitBinding<IAlphaInputArgs>& Binding) const
+{
+	return EAnimAlphaInputType::Bool;
+}
+
+float FBoolAlphaInputArgTrait::GetCurrentAlphaValue(const FExecutionContext& Context, const TTraitBinding<IAlphaInputArgs>& Binding) const
+{
+	FInstanceData* InstanceData = Binding.GetInstanceData<FInstanceData>();
+
+	return InstanceData->ComputedAlphaValue;
+}
+
+} // namespace UE::AnimNext
+
diff --git a/Engine/Plugins/Experimental/AnimNextAnimGraph/Source/AnimNextAnimGraph/Private/Traits/Args/BoolAlphaInputArgTrait.h b/Engine/Plugins/Experimental/AnimNextAnimGraph/Source/AnimNextAnimGraph/Private/Traits/Args/BoolAlphaInputArgTrait.h
new file mode 100644
--- /dev/null
+++ b/Engine/Plugins/Experimental/AnimNextAnimGraph/Source/AnimNextAnimGraph/Private/Traits/Args/BoolAlphaInputArgTrait.h
@@ -0,0 +1,66 @@
+// Copyright Epic Games, Inc. All Rights Reserved.
+
+#pragma once
+
+#include "CoreMinimal.h"
+#include "TraitCore/Trait.h"
+#include "TraitCore/TraitSharedData.h"
+#include "TraitInterfaces/IUpdate.h"
+#include "TraitInterfaces/Args/IAlphaInputArgs.h"
+
+#include "BoolAlphaInputArgTrait.generated.h"
+
+USTRUCT(meta = (DisplayName = "Bool Alpha Input Args"))
+struct FBoolAlphaInputTraitArgs : public FAnimNextTraitSharedData
+{
+	GENERATED_BODY()
+
+	UPROPERTY(EditAnywhere, Category = Settings, meta = (DisplayName = "bEnabled"))
+	bool bAlphaBoolEnabled = true;
+
+	UPROPERTY(EditAnywhere, Category = Settings, meta = (Hidden, DisplayName = "Blend Settings"))
+	FInputAlphaBoolBlend AlphaBoolBlend;
+
+	// Latent pin support boilerplate
+#define TRAIT_LATENT_PROPERTIES_ENUMERATOR(GeneratorMacro) \
+		GeneratorMacro(bAlphaBoolEnabled) \
+		GeneratorMacro(AlphaBoolBlend) \
+
+	GENERATE_TRAIT_LATENT_PROPERTIES(FBoolAlphaInputTraitArgs, TRAIT_LATENT_PROPERTIES_ENUMERATOR)
+#undef TRAIT_LATENT_PROPERTIES_ENUMERATOR
+};
+
+namespace UE::AnimNext
+{
+	/**
+	 * FBoolAlphaInputArgTrait
+	 *
+	 * See 'FAlphaInputArgCoreTrait'. Bool input only version.
+	 */
+	USTRUCT(meta = (DisplayName = "Bool Alpha Input Arg"))
+	struct FBoolAlphaInputArgTrait : FAdditiveTrait, IUpdate, IAlphaInputArgs
+	{
+		DECLARE_ANIM_TRAIT(FBoolAlphaInputArgTrait, FAdditiveTrait)
+
+		using FSharedData = FBoolAlphaInputTraitArgs;
+
+		struct FInstanceData : FTrait::FInstanceData
+		{
+			FInputAlphaBoolBlend AlphaBoolBlend;
+
+			float ComputedAlphaValue = 0.f;
+
+			void Construct(const FExecutionContext& Context, const FTraitBinding& Binding);
+		};
+
+		// IUpdate
+		virtual void OnBecomeRelevant(FUpdateTraversalContext& Context, const TTraitBinding<IUpdate>& Binding, const FTraitUpdateState& TraitState) const override;
+		virtual void PreUpdate(FUpdateTraversalContext& Context, const TTraitBinding<IUpdate>& Binding, const FTraitUpdateState& TraitState) const override;
+
+		// IAlphaInputArgs
+		virtual FAlphaInputTraitArgs Get(const FExecutionContext& Context, const TTraitBinding<IAlphaInputArgs>& Binding) const override;
+		virtual EAnimAlphaInputType GetAlphaInputType(const FExecutionContext& Context, const TTraitBinding<IAlphaInputArgs>& Binding) const override;
+		virtual float GetCurrentAlphaValue(const FExecutionContext& Context, const TTraitBinding<IAlphaInputArgs>& Binding) const override;
+	};
+
+} // namespace UE::AnimNext
diff --git a/Engine/Plugins/Experimental/AnimNextAnimGraph/Source/AnimNextAnimGraph/Private/Traits/Args/CurveAlphaInputArgTrait.cpp b/Engine/Plugins/Experimental/AnimNextAnimGraph/Source/AnimNextAnimGraph/Private/Traits/Args/CurveAlphaInputArgTrait.cpp
new file mode 100644
--- /dev/null
+++ b/Engine/Plugins/Experimental/AnimNextAnimGraph/Source/AnimNextAnimGraph/Private/Traits/Args/CurveAlphaInputArgTrait.cpp
@@ -0,0 +1,94 @@
+// Copyright Epic Games, Inc. All Rights Reserved.
+
+#include "Traits/Args/CurveAlphaInputArgTrait.h"
+#include "Traits/Args/AlphaInputArgCoreTrait.h"
+
+#include "Animation/AnimTypes.h"
+#include "TraitCore/ExecutionContext.h"
+#include "EvaluationVM/Tasks/BlendKeyframes.h"
+#include "EvaluationVM/Tasks/StoreKeyframe.h"
+
+namespace UE::AnimNext
+{
+
+////////////////////////////////////////////////////////////////////////////////////////////////////////////////////////
+// FCurveAlphaInputArgTrait
+
+AUTO_REGISTER_ANIM_TRAIT(FCurveAlphaInputArgTrait)
+
+#define TRAIT_INTERFACE_ENUMERATOR(GeneratorMacro) \
+	GeneratorMacro(IAlphaInputArgs) \
+	GeneratorMacro(IUpdate) \
+
+// Trait implementation boilerplate
+GENERATE_ANIM_TRAIT_IMPLEMENTATION(FCurveAlphaInputArgTrait, TRAIT_INTERFACE_ENUMERATOR, NULL_ANIM_TRAIT_INTERFACE_ENUMERATOR, NULL_ANIM_TRAIT_EVENT_ENUMERATOR)
+#undef TRAIT_INTERFACE_ENUMERATOR
+
+void FCurveAlphaInputArgTrait::FInstanceData::Construct(const FExecutionContext& Context, const FTraitBinding& Binding)
+{
+	const FSharedData* SharedData = Binding.GetSharedData<FSharedData>();
+
+	AlphaScaleBiasClamp = SharedData->AlphaScaleBiasClamp;
+}
+
+void FCurveAlphaInputArgTrait::OnBecomeRelevant(FUpdateTraversalContext& Context, const TTraitBinding<IUpdate>& Binding, const FTraitUpdateState& TraitState) const
+{
+	IUpdate::OnBecomeRelevant(Context, Binding, TraitState);
+
+	FInstanceData* InstanceData = Binding.GetInstanceData<FInstanceData>();
+
+	InstanceData->AlphaScaleBiasClamp.Reinitialize();
+}
+
+void FCurveAlphaInputArgTrait::PreUpdate(FUpdateTraversalContext& Context, const TTraitBinding<IUpdate>& Binding, const FTraitUpdateState& TraitState) const
+{
+	IUpdate::PreUpdate(Context, Binding, TraitState);
+
+	FInstanceData* InstanceData = Binding.GetInstanceData<FInstanceData>();
+
+	InstanceData->DeltaTime = TraitState.GetDeltaTime();
+}
+
+FAlphaInputTraitArgs FCurveAlphaInputArgTrait::Get(const FExecutionContext& Context, const TTraitBinding<IAlphaInputArgs>& Binding) const
+{
+	const FSharedData* SharedData = Binding.GetSharedData<FSharedData>();
+	FInstanceData* InstanceData = Binding.GetInstanceData<FInstanceData>();
+
+	FAlphaInputTraitArgs Result = FAlphaInputTraitArgs();
+	// Alpha = InstanceData->InternalAlpha;
+	// Result.AlphaScaleBias = SharedData->AlphaScaleBias;
+	Result.AlphaScaleBiasClamp = InstanceData->AlphaScaleBiasClamp;
+	// Result.bAlphaBoolEnabled = SharedData->GetbAlphaBoolEnabled(Binding);
+	// Result.AlphaBoolBlend = InstanceData->AlphaBoolBlend;
+	Result.AlphaCurveName = SharedData->GetAlphaCurveName(Binding);
+	Result.AlphaInputType = EAnimAlphaInputType::Curve;
+
+	return Result;
+}
+
+EAnimAlphaInputType FCurveAlphaInputArgTrait::GetAlphaInputType(const FExecutionContext& Context, const TTraitBinding<IAlphaInputArgs>& Binding) const
+{
+	return EAnimAlphaInputType::Curve;
+}
+
+FName FCurveAlphaInputArgTrait::GetAlphaCurveName(const FExecutionContext& Context, const TTraitBinding<IAlphaInputArgs>& Binding) const
+{
+	const FSharedData* SharedData = Binding.GetSharedData<FSharedData>();
+
+	return SharedData->GetAlphaCurveName(Binding);
+}
+
+TFunction<float(float)> FCurveAlphaInputArgTrait::GetInputScaleBiasClampCallback(const FExecutionContext& Context, const TTraitBinding<IAlphaInputArgs>& Binding) const
+{
+	FInstanceData* InstanceData = Binding.GetInstanceData<FInstanceData>();
+	
+	auto InputBiasClampCallback = [DeltaTime = InstanceData->DeltaTime, InputScaleBiasClamp = &InstanceData->AlphaScaleBiasClamp](float Alpha) -> float
+	{
+		return InputScaleBiasClamp->ApplyTo(Alpha, DeltaTime);
+	};
+
+	return InputBiasClampCallback;
+}
+
+} // namespace UE::AnimNext
+
diff --git a/Engine/Plugins/Experimental/AnimNextAnimGraph/Source/AnimNextAnimGraph/Private/Traits/Args/CurveAlphaInputArgTrait.h b/Engine/Plugins/Experimental/AnimNextAnimGraph/Source/AnimNextAnimGraph/Private/Traits/Args/CurveAlphaInputArgTrait.h
new file mode 100644
--- /dev/null
+++ b/Engine/Plugins/Experimental/AnimNextAnimGraph/Source/AnimNextAnimGraph/Private/Traits/Args/CurveAlphaInputArgTrait.h
@@ -0,0 +1,68 @@
+// Copyright Epic Games, Inc. All Rights Reserved.
+
+#pragma once
+
+#include "CoreMinimal.h"
+#include "TraitCore/Trait.h"
+#include "TraitCore/TraitSharedData.h"
+#include "TraitInterfaces/IUpdate.h"
+#include "TraitInterfaces/Args/IAlphaInputArgs.h"
+
+#include "CurveAlphaInputArgTrait.generated.h"
+
+USTRUCT(meta = (DisplayName = "Curve Alpha Input Args"))
+struct FCurveAlphaInputTraitArgs : public FAnimNextTraitSharedData
+{
+	GENERATED_BODY()
+
+	UPROPERTY(EditAnywhere, Category = Settings, meta = (Hidden, DisplayAfter = "AlphaCurveName"))
+	FInputScaleBiasClamp AlphaScaleBiasClamp;
+
+	/** Note: For curve types, the additive branch will always be evaluated. Curve weight blending will occur at task evaluation time. */
+	UPROPERTY(EditAnywhere, Category = Settings)
+	FName AlphaCurveName = NAME_None;
+
+	// Latent pin support boilerplate
+#define TRAIT_LATENT_PROPERTIES_ENUMERATOR(GeneratorMacro) \
+		GeneratorMacro(AlphaCurveName) \
+
+	GENERATE_TRAIT_LATENT_PROPERTIES(FCurveAlphaInputTraitArgs, TRAIT_LATENT_PROPERTIES_ENUMERATOR)
+#undef TRAIT_LATENT_PROPERTIES_ENUMERATOR
+};
+
+namespace UE::AnimNext
+{
+
+	/**
+	 * FCurveAlphaInputArgTrait
+	 *
+	 * See 'FAlphaInputArgCoreTrait'. Curve input only version.
+	 */
+	USTRUCT(meta = (DisplayName = "Curve Alpha Input Arg"))
+	struct FCurveAlphaInputArgTrait : FAdditiveTrait, IUpdate, IAlphaInputArgs
+	{
+		DECLARE_ANIM_TRAIT(FCurveAlphaInputArgTrait, FAdditiveTrait)
+
+		using FSharedData = FCurveAlphaInputTraitArgs;
+		
+		struct FInstanceData : FTrait::FInstanceData
+		{
+			FInputScaleBiasClamp AlphaScaleBiasClamp;
+
+			float DeltaTime = 0.f;
+
+			void Construct(const FExecutionContext& Context, const FTraitBinding& Binding);
+		};
+
+		// IUpdate
+		virtual void OnBecomeRelevant(FUpdateTraversalContext& Context, const TTraitBinding<IUpdate>& Binding, const FTraitUpdateState& TraitState) const override;
+		virtual void PreUpdate(FUpdateTraversalContext& Context, const TTraitBinding<IUpdate>& Binding, const FTraitUpdateState& TraitState) const override;
+
+		// IAlphaInputArgs
+		virtual FAlphaInputTraitArgs Get(const FExecutionContext& Context, const TTraitBinding<IAlphaInputArgs>& Binding) const override;
+		virtual EAnimAlphaInputType GetAlphaInputType(const FExecutionContext& Context, const TTraitBinding<IAlphaInputArgs>& Binding) const override;
+		virtual FName GetAlphaCurveName(const FExecutionContext& Context, const TTraitBinding<IAlphaInputArgs>& Binding) const override;
+		virtual TFunction<float(float)> GetInputScaleBiasClampCallback(const FExecutionContext& Context, const TTraitBinding<IAlphaInputArgs>& Binding) const override;
+	};
+
+} // namespace UE::AnimNext
diff --git a/Engine/Plugins/Experimental/AnimNextAnimGraph/Source/AnimNextAnimGraph/Private/Traits/Args/FloatAlphaInputArgTrait.cpp b/Engine/Plugins/Experimental/AnimNextAnimGraph/Source/AnimNextAnimGraph/Private/Traits/Args/FloatAlphaInputArgTrait.cpp
new file mode 100644
--- /dev/null
+++ b/Engine/Plugins/Experimental/AnimNextAnimGraph/Source/AnimNextAnimGraph/Private/Traits/Args/FloatAlphaInputArgTrait.cpp
@@ -0,0 +1,87 @@
+// Copyright Epic Games, Inc. All Rights Reserved.
+
+#include "Traits/Args/FloatAlphaInputArgTrait.h"
+#include "Traits/Args/AlphaInputArgCoreTrait.h"
+
+#include "Animation/AnimTypes.h"
+#include "TraitCore/ExecutionContext.h"
+#include "EvaluationVM/Tasks/BlendKeyframes.h"
+#include "EvaluationVM/Tasks/StoreKeyframe.h"
+
+namespace UE::AnimNext
+{
+
+////////////////////////////////////////////////////////////////////////////////////////////////////////////////////////
+// FFloatAlphaInputArgTrait
+
+AUTO_REGISTER_ANIM_TRAIT(FFloatAlphaInputArgTrait)
+
+#define TRAIT_INTERFACE_ENUMERATOR(GeneratorMacro) \
+	GeneratorMacro(IAlphaInputArgs) \
+	GeneratorMacro(IUpdate) \
+
+// Trait implementation boilerplate
+GENERATE_ANIM_TRAIT_IMPLEMENTATION(FFloatAlphaInputArgTrait, TRAIT_INTERFACE_ENUMERATOR, NULL_ANIM_TRAIT_INTERFACE_ENUMERATOR, NULL_ANIM_TRAIT_EVENT_ENUMERATOR)
+#undef TRAIT_INTERFACE_ENUMERATOR
+
+void FFloatAlphaInputArgTrait::FInstanceData::Construct(const FExecutionContext& Context, const FTraitBinding& Binding)
+{
+	const FSharedData* SharedData = Binding.GetSharedData<FSharedData>();
+
+	ComputedAlphaValue = SharedData->GetAlpha(Binding);
+	AlphaScaleBiasClamp = SharedData->AlphaScaleBiasClamp;
+}
+
+void FFloatAlphaInputArgTrait::OnBecomeRelevant(FUpdateTraversalContext& Context, const TTraitBinding<IUpdate>& Binding, const FTraitUpdateState& TraitState) const
+{
+	IUpdate::OnBecomeRelevant(Context, Binding, TraitState);
+
+	const FSharedData* SharedData = Binding.GetSharedData<FSharedData>();
+	FInstanceData* InstanceData = Binding.GetInstanceData<FInstanceData>();
+
+	InstanceData->ComputedAlphaValue = AlphaInput::ComputeAlphaValueForFloat(IN OUT InstanceData->AlphaScaleBiasClamp, SharedData->AlphaScaleBias, SharedData->GetAlpha(Binding), TraitState.GetDeltaTime());
+
+	InstanceData->AlphaScaleBiasClamp.Reinitialize();
+}
+
+void FFloatAlphaInputArgTrait::PreUpdate(FUpdateTraversalContext& Context, const TTraitBinding<IUpdate>& Binding, const FTraitUpdateState& TraitState) const
+{
+	IUpdate::PreUpdate(Context, Binding, TraitState);
+
+	const FSharedData* SharedData = Binding.GetSharedData<FSharedData>();
+	FInstanceData* InstanceData = Binding.GetInstanceData<FInstanceData>();
+
+	InstanceData->ComputedAlphaValue = AlphaInput::ComputeAlphaValueForFloat(IN OUT InstanceData->AlphaScaleBiasClamp, SharedData->AlphaScaleBias, SharedData->GetAlpha(Binding), TraitState.GetDeltaTime());
+}
+
+FAlphaInputTraitArgs FFloatAlphaInputArgTrait::Get(const FExecutionContext& Context, const TTraitBinding<IAlphaInputArgs>& Binding) const
+{
+	const FSharedData* SharedData = Binding.GetSharedData<FSharedData>();
+	FInstanceData* InstanceData = Binding.GetInstanceData<FInstanceData>();
+
+	FAlphaInputTraitArgs Result = FAlphaInputTraitArgs();
+	Result.Alpha = SharedData->GetAlpha(Binding);
+	Result.AlphaScaleBias = SharedData->AlphaScaleBias;
+	Result.AlphaScaleBiasClamp = InstanceData->AlphaScaleBiasClamp;
+	// Result.bAlphaBoolEnabled = SharedData->GetbAlphaBoolEnabled(Binding);
+	// Result.AlphaBoolBlend = InstanceData->AlphaBoolBlend;
+	// Result.AlphaCurveName = SharedData->GetAlphaCurveName(Binding);
+	Result.AlphaInputType = EAnimAlphaInputType::Float;
+
+	return Result;
+}
+
+EAnimAlphaInputType FFloatAlphaInputArgTrait::GetAlphaInputType(const FExecutionContext& Context, const TTraitBinding<IAlphaInputArgs>& Binding) const
+{
+	return EAnimAlphaInputType::Float;
+}
+
+float FFloatAlphaInputArgTrait::GetCurrentAlphaValue(const FExecutionContext& Context, const TTraitBinding<IAlphaInputArgs>& Binding) const
+{
+	FInstanceData* InstanceData = Binding.GetInstanceData<FInstanceData>();
+
+	return InstanceData->ComputedAlphaValue;
+}
+
+} // namespace UE::AnimNext
+
diff --git a/Engine/Plugins/Experimental/AnimNextAnimGraph/Source/AnimNextAnimGraph/Private/Traits/Args/FloatAlphaInputArgTrait.h b/Engine/Plugins/Experimental/AnimNextAnimGraph/Source/AnimNextAnimGraph/Private/Traits/Args/FloatAlphaInputArgTrait.h
new file mode 100644
--- /dev/null
+++ b/Engine/Plugins/Experimental/AnimNextAnimGraph/Source/AnimNextAnimGraph/Private/Traits/Args/FloatAlphaInputArgTrait.h
@@ -0,0 +1,69 @@
+// Copyright Epic Games, Inc. All Rights Reserved.
+
+#pragma once
+
+#include "CoreMinimal.h"
+#include "TraitCore/Trait.h"
+#include "TraitCore/TraitSharedData.h"
+#include "TraitInterfaces/IUpdate.h"
+#include "TraitInterfaces/Args/IAlphaInputArgs.h"
+
+#include "FloatAlphaInputArgTrait.generated.h"
+
+USTRUCT(meta = (DisplayName = "Float Alpha Input Args"))
+struct FFloatAlphaInputTraitArgs : public FAnimNextTraitSharedData
+{
+	GENERATED_BODY()
+
+	UPROPERTY(EditAnywhere, Category = Settings)
+	float Alpha = 1.f;
+
+	UPROPERTY(EditAnywhere, Category = Settings, meta = (Hidden))
+	FInputScaleBias AlphaScaleBias;
+
+	UPROPERTY(EditAnywhere, Category = Settings, meta = (Hidden))
+	FInputScaleBiasClamp AlphaScaleBiasClamp;
+
+	// Latent pin support boilerplate
+#define TRAIT_LATENT_PROPERTIES_ENUMERATOR(GeneratorMacro) \
+		GeneratorMacro(Alpha) \
+
+	GENERATE_TRAIT_LATENT_PROPERTIES(FFloatAlphaInputTraitArgs, TRAIT_LATENT_PROPERTIES_ENUMERATOR)
+#undef TRAIT_LATENT_PROPERTIES_ENUMERATOR
+};
+
+namespace UE::AnimNext
+{
+
+	/**
+	 * FFloatAlphaInputArgTrait
+	 *
+	 * See 'FAlphaInputArgCoreTrait'. Float input only version.
+	 */
+	USTRUCT(meta = (DisplayName = "Float Alpha Input Arg"))
+	struct FFloatAlphaInputArgTrait : FAdditiveTrait, IUpdate, IAlphaInputArgs
+	{
+		DECLARE_ANIM_TRAIT(FFloatAlphaInputArgTrait, FAdditiveTrait)
+
+		using FSharedData = FFloatAlphaInputTraitArgs;
+		
+		struct FInstanceData : FTrait::FInstanceData
+		{
+			FInputScaleBiasClamp AlphaScaleBiasClamp;
+
+			float ComputedAlphaValue = 0.f;
+
+			void Construct(const FExecutionContext& Context, const FTraitBinding& Binding);
+		};
+
+		// IUpdate
+		virtual void OnBecomeRelevant(FUpdateTraversalContext& Context, const TTraitBinding<IUpdate>& Binding, const FTraitUpdateState& TraitState) const override;
+		virtual void PreUpdate(FUpdateTraversalContext& Context, const TTraitBinding<IUpdate>& Binding, const FTraitUpdateState& TraitState) const override;
+
+		// IAlphaInputArgs
+		virtual FAlphaInputTraitArgs Get(const FExecutionContext& Context, const TTraitBinding<IAlphaInputArgs>& Binding) const override;
+		virtual EAnimAlphaInputType GetAlphaInputType(const FExecutionContext& Context, const TTraitBinding<IAlphaInputArgs>& Binding) const override;
+		virtual float GetCurrentAlphaValue(const FExecutionContext& Context, const TTraitBinding<IAlphaInputArgs>& Binding) const override;
+	};
+
+} // namespace UE::AnimNext
diff --git a/Engine/Plugins/Experimental/AnimNextAnimGraph/Source/AnimNextAnimGraph/Private/Traits/Args/SimpleFloatAlphaInputArgTrait.cpp b/Engine/Plugins/Experimental/AnimNextAnimGraph/Source/AnimNextAnimGraph/Private/Traits/Args/SimpleFloatAlphaInputArgTrait.cpp
new file mode 100644
--- /dev/null
+++ b/Engine/Plugins/Experimental/AnimNextAnimGraph/Source/AnimNextAnimGraph/Private/Traits/Args/SimpleFloatAlphaInputArgTrait.cpp
@@ -0,0 +1,54 @@
+// Copyright Epic Games, Inc. All Rights Reserved.
+
+#include "Traits/Args/SimpleFloatAlphaInputArgTrait.h"
+#include "Traits/Args/AlphaInputArgCoreTrait.h"
+
+#include "Animation/AnimTypes.h"
+#include "TraitCore/ExecutionContext.h"
+#include "EvaluationVM/Tasks/BlendKeyframes.h"
+#include "EvaluationVM/Tasks/StoreKeyframe.h"
+
+namespace UE::AnimNext
+{
+
+////////////////////////////////////////////////////////////////////////////////////////////////////////////////////////
+// FSimpleFloatAlphaInputArgTrait
+
+AUTO_REGISTER_ANIM_TRAIT(FSimpleFloatAlphaInputArgTrait)
+
+#define TRAIT_INTERFACE_ENUMERATOR(GeneratorMacro) \
+	GeneratorMacro(IAlphaInputArgs) \
+
+// Trait implementation boilerplate
+GENERATE_ANIM_TRAIT_IMPLEMENTATION(FSimpleFloatAlphaInputArgTrait, TRAIT_INTERFACE_ENUMERATOR, NULL_ANIM_TRAIT_INTERFACE_ENUMERATOR, NULL_ANIM_TRAIT_EVENT_ENUMERATOR)
+#undef TRAIT_INTERFACE_ENUMERATOR
+
+FAlphaInputTraitArgs FSimpleFloatAlphaInputArgTrait::Get(const FExecutionContext& Context, const TTraitBinding<IAlphaInputArgs>& Binding) const
+{
+	const FSharedData* SharedData = Binding.GetSharedData<FSharedData>();
+
+	FAlphaInputTraitArgs Result = FAlphaInputTraitArgs();
+	Result.Alpha = SharedData->GetAlpha(Binding);
+	// Result.AlphaScaleBias = SharedData->AlphaScaleBias;
+	// Result.AlphaScaleBiasClamp = InstanceData->AlphaScaleBiasClamp;
+	// Result.bAlphaBoolEnabled = SharedData->GetbAlphaBoolEnabled(Binding);
+	// Result.AlphaBoolBlend = InstanceData->AlphaBoolBlend;
+	// Result.AlphaCurveName = SharedData->GetAlphaCurveName(Binding);
+	Result.AlphaInputType = EAnimAlphaInputType::Float;
+
+	return Result;
+}
+
+EAnimAlphaInputType FSimpleFloatAlphaInputArgTrait::GetAlphaInputType(const FExecutionContext& Context, const TTraitBinding<IAlphaInputArgs>& Binding) const
+{
+	return EAnimAlphaInputType::Float;
+}
+
+float FSimpleFloatAlphaInputArgTrait::GetCurrentAlphaValue(const FExecutionContext& Context, const TTraitBinding<IAlphaInputArgs>& Binding) const
+{
+	const FSharedData* SharedData = Binding.GetSharedData<FSharedData>();
+	return SharedData->GetAlpha(Binding);
+}
+
+} // namespace UE::AnimNext
+
diff --git a/Engine/Plugins/Experimental/AnimNextAnimGraph/Source/AnimNextAnimGraph/Private/Traits/Args/SimpleFloatAlphaInputArgTrait.h b/Engine/Plugins/Experimental/AnimNextAnimGraph/Source/AnimNextAnimGraph/Private/Traits/Args/SimpleFloatAlphaInputArgTrait.h
new file mode 100644
--- /dev/null
+++ b/Engine/Plugins/Experimental/AnimNextAnimGraph/Source/AnimNextAnimGraph/Private/Traits/Args/SimpleFloatAlphaInputArgTrait.h
@@ -0,0 +1,54 @@
+// Copyright Epic Games, Inc. All Rights Reserved.
+
+#pragma once
+
+#include "CoreMinimal.h"
+#include "TraitCore/Trait.h"
+#include "TraitCore/TraitSharedData.h"
+#include "TraitInterfaces/IUpdate.h"
+#include "TraitInterfaces/Args/IAlphaInputArgs.h"
+
+#include "SimpleFloatAlphaInputArgTrait.generated.h"
+
+USTRUCT(meta = (DisplayName = "Simple Float Alpha Input Args"))
+struct FSimpleFloatAlphaInputTraitArgs : public FAnimNextTraitSharedData
+{
+	GENERATED_BODY()
+
+	UPROPERTY(EditAnywhere, Category = Settings)
+	float Alpha = 1.f;
+
+	// Latent pin support boilerplate
+#define TRAIT_LATENT_PROPERTIES_ENUMERATOR(GeneratorMacro) \
+		GeneratorMacro(Alpha) \
+
+	GENERATE_TRAIT_LATENT_PROPERTIES(FSimpleFloatAlphaInputTraitArgs, TRAIT_LATENT_PROPERTIES_ENUMERATOR)
+#undef TRAIT_LATENT_PROPERTIES_ENUMERATOR
+};
+
+namespace UE::AnimNext
+{
+
+	/**
+	 * FSimpleFloatAlphaInputArgTrait
+	 *
+	 * See 'FAlphaInputArgCoreTrait'. Simple Float input only version. Alpha is always updated.
+	 */
+	USTRUCT(meta = (DisplayName = "Simple Float Alpha Input Arg"))
+	struct FSimpleFloatAlphaInputArgTrait : FAdditiveTrait, IAlphaInputArgs
+	{
+		DECLARE_ANIM_TRAIT(FSimpleFloatAlphaInputArgTrait, FAdditiveTrait)
+
+		using FSharedData = FSimpleFloatAlphaInputTraitArgs;
+		
+		struct FInstanceData : FTrait::FInstanceData
+		{
+		};
+
+		// IAlphaInputArgs
+		virtual FAlphaInputTraitArgs Get(const FExecutionContext& Context, const TTraitBinding<IAlphaInputArgs>& Binding) const override;
+		virtual EAnimAlphaInputType GetAlphaInputType(const FExecutionContext& Context, const TTraitBinding<IAlphaInputArgs>& Binding) const override;
+		virtual float GetCurrentAlphaValue(const FExecutionContext& Context, const TTraitBinding<IAlphaInputArgs>& Binding) const override;
+	};
+
+} // namespace UE::AnimNext
diff --git a/Engine/Plugins/Experimental/AnimNextAnimGraph/Source/AnimNextAnimGraph/Private/Traits/PassthroughBlendTrait.cpp b/Engine/Plugins/Experimental/AnimNextAnimGraph/Source/AnimNextAnimGraph/Private/Traits/PassthroughBlendTrait.cpp
--- a/Engine/Plugins/Experimental/AnimNextAnimGraph/Source/AnimNextAnimGraph/Private/Traits/PassthroughBlendTrait.cpp
+++ b/Engine/Plugins/Experimental/AnimNextAnimGraph/Source/AnimNextAnimGraph/Private/Traits/PassthroughBlendTrait.cpp
@@ -25,8 +25,7 @@ void FPassthroughBlendTrait::FInstanceData::Construct(const FExecutionContext& C
 {
 	const FSharedData* SharedData = Binding.GetSharedData<FSharedData>();
 
-	InternalAlpha = SharedData->GetAlpha(Binding);
-	ComputedAlphaValue = InternalAlpha;
+	ComputedAlphaValue = SharedData->GetAlpha(Binding);;
 	AlphaBoolBlend = SharedData->GetAlphaBoolBlend(Binding);
 	AlphaScaleBiasClamp = SharedData->AlphaScaleBiasClamp;
 }
@@ -38,8 +37,7 @@ void FPassthroughBlendTrait::OnBecomeRelevant(FUpdateTraversalContext& Context,
 	const FSharedData* SharedData = Binding.GetSharedData<FSharedData>();
 	FInstanceData* InstanceData = Binding.GetInstanceData<FInstanceData>();
 
-	InstanceData->InternalAlpha = SharedData->GetAlpha(Binding);
-	InstanceData->ComputedAlphaValue = ComputeAlphaValue(SharedData->GetAlphaInputType(Binding), SharedData, InstanceData, TraitState.GetDeltaTime());
+	InstanceData->ComputedAlphaValue = ComputeAlphaValue(SharedData->GetAlphaInputType(Binding), SharedData, InstanceData, Binding, TraitState.GetDeltaTime());
 
 	InstanceData->AlphaBoolBlend.Reinitialize();
 	InstanceData->AlphaScaleBiasClamp.Reinitialize();
@@ -54,7 +52,8 @@ void FPassthroughBlendTrait::PreUpdate(FUpdateTraversalContext& Context, const T
 
 	const float DeltaTime = TraitState.GetDeltaTime();
 	InstanceData->DeltaTime = DeltaTime;
-	InstanceData->ComputedAlphaValue = ComputeAlphaValue(SharedData->GetAlphaInputType(Binding), SharedData, InstanceData, DeltaTime);
+
+	InstanceData->ComputedAlphaValue = ComputeAlphaValue(SharedData->GetAlphaInputType(Binding), SharedData, InstanceData, Binding, DeltaTime);
 }
 
 // IEvaluate impl
@@ -91,7 +90,7 @@ void FPassthroughBlendTrait::PostEvaluate(FEvaluateTraversalContext& Context, co
 }
 
 
-float FPassthroughBlendTrait::ComputeAlphaValue(const EAnimAlphaInputType AlphaInputType, const FPassthroughBlendTraitSharedData* SharedData, FInstanceData* InstanceData, float DeltaTime)
+float FPassthroughBlendTrait::ComputeAlphaValue(const EAnimAlphaInputType AlphaInputType, const FPassthroughBlendTraitSharedData* SharedData, FInstanceData* InstanceData, const TTraitBinding<IUpdate>& Binding, float DeltaTime)
 {
 	// alpha handlers
 	float CurrentAlpha = 0.f;
@@ -100,12 +99,12 @@ float FPassthroughBlendTrait::ComputeAlphaValue(const EAnimAlphaInputType AlphaI
 	{
 		case EAnimAlphaInputType::Float:
 		{
-			CurrentAlpha = SharedData->AlphaScaleBias.ApplyTo(InstanceData->AlphaScaleBiasClamp.ApplyTo(InstanceData->InternalAlpha, DeltaTime));
+			CurrentAlpha = SharedData->AlphaScaleBias.ApplyTo(InstanceData->AlphaScaleBiasClamp.ApplyTo(SharedData->GetAlpha(Binding), DeltaTime));
 			break;
 		}
 		case EAnimAlphaInputType::Bool:
 		{
-			CurrentAlpha = InstanceData->AlphaBoolBlend.ApplyTo(SharedData->bAlphaBoolEnabled, DeltaTime);
+			CurrentAlpha = InstanceData->AlphaBoolBlend.ApplyTo(SharedData->GetbAlphaBoolEnabled(Binding), DeltaTime);
 			break;
 		}
 		case EAnimAlphaInputType::Curve:
diff --git a/Engine/Plugins/Experimental/AnimNextAnimGraph/Source/AnimNextAnimGraph/Public/EvaluationVM/Tasks/ApplyAdditiveKeyframe.h b/Engine/Plugins/Experimental/AnimNextAnimGraph/Source/AnimNextAnimGraph/Public/EvaluationVM/Tasks/ApplyAdditiveKeyframe.h
--- a/Engine/Plugins/Experimental/AnimNextAnimGraph/Source/AnimNextAnimGraph/Public/EvaluationVM/Tasks/ApplyAdditiveKeyframe.h
+++ b/Engine/Plugins/Experimental/AnimNextAnimGraph/Source/AnimNextAnimGraph/Public/EvaluationVM/Tasks/ApplyAdditiveKeyframe.h
@@ -7,6 +7,13 @@
 
 #include "ApplyAdditiveKeyframe.generated.h"
 
+struct FInputScaleBiasClamp;
+
+namespace UE::AnimNext
+{
+	struct FKeyframeState;
+}
+
 /*
  * Apply Additive Keyframe Task
  *
@@ -22,11 +29,23 @@ struct ANIMNEXTANIMGRAPH_API FAnimNextApplyAdditiveKeyframeTask : public FAnimNe
 	DECLARE_ANIM_EVALUATION_TASK(FAnimNextApplyAdditiveKeyframeTask)
 
 	static FAnimNextApplyAdditiveKeyframeTask Make(float BlendWeight);
+	static FAnimNextApplyAdditiveKeyframeTask Make(const FName& AlphaSourceCurveName, const int8 AlphaCurveInputIndex, TFunction<float(float)> InputScaleBiasClampFn);
 
 	// Task entry point
 	virtual void Execute(UE::AnimNext::FEvaluationVM& VM) const override;
 
+	// @TODO: Move this into shared method
+	// Compute the current alpha value for curves or return the interpolation alpha otherwise
+	float GetInterpolationAlpha(const UE::AnimNext::FKeyframeState* KeyframeA, const UE::AnimNext::FKeyframeState* KeyframeB) const;
+
 	// How much weight between the additive identity and the additive pose to apply: lerp(identity, additive, weight)
 	UPROPERTY()
 	float BlendWeight = 0.0f;
+
+	// The curve to evaluate and extract the interpolation alpha between the two input keyframes
+	FName AlphaSourceCurveName = NAME_None;
+
+	int8 AlphaCurveInputIndex = INDEX_NONE;
+
+	TFunction<float(float)> InputScaleBiasClampFn = nullptr;
 };
diff --git a/Engine/Plugins/Experimental/AnimNextAnimGraph/Source/AnimNextAnimGraph/Public/TraitInterfaces/Args/IAlphaInputArgs.h b/Engine/Plugins/Experimental/AnimNextAnimGraph/Source/AnimNextAnimGraph/Public/TraitInterfaces/Args/IAlphaInputArgs.h
new file mode 100644
--- /dev/null
+++ b/Engine/Plugins/Experimental/AnimNextAnimGraph/Source/AnimNextAnimGraph/Public/TraitInterfaces/Args/IAlphaInputArgs.h
@@ -0,0 +1,141 @@
+// Copyright Epic Games, Inc. All Rights Reserved.
+
+#pragma once
+
+#include "CoreMinimal.h"
+#include "TraitCore/ITraitInterface.h"
+#include "TraitCore/TraitBinding.h"
+#include "TraitCore/TraitSharedData.h"
+#include "Animation/InputScaleBias.h"
+
+#include "IAlphaInputArgs.generated.h"
+
+USTRUCT(meta = (DisplayName = "Alpha Input Args"))
+struct FAlphaInputTraitArgs : public FAnimNextTraitSharedData
+{
+	GENERATED_BODY()
+
+	UPROPERTY(EditAnywhere, Category = Settings, meta = (EditConditionHides, EditCondition = "AlphaInputType == EAnimAlphaInputType::Float"))
+	float Alpha = 1.f;
+
+	UPROPERTY(EditAnywhere, Category = Settings, meta = (Hidden, EditConditionHides, EditCondition = "AlphaInputType == EAnimAlphaInputType::Float"))
+	FInputScaleBias AlphaScaleBias;
+
+	UPROPERTY(EditAnywhere, Category = Settings, meta = (Hidden, DisplayAfter = "AlphaCurveName", EditConditionHides, EditCondition = "AlphaInputType == EAnimAlphaInputType::Float || AlphaInputType == EAnimAlphaInputType::Curve"))
+	FInputScaleBiasClamp AlphaScaleBiasClamp;
+
+	UPROPERTY(EditAnywhere, Category = Settings, meta = (DisplayName = "bEnabled", EditConditionHides, EditCondition = "AlphaInputType == EAnimAlphaInputType::Bool"))
+	bool bAlphaBoolEnabled = true;
+
+	UPROPERTY(EditAnywhere, Category = Settings, meta = (Hidden, DisplayName = "Blend Settings", EditConditionHides, EditCondition = "AlphaInputType == EAnimAlphaInputType::Bool"))
+	FInputAlphaBoolBlend AlphaBoolBlend;
+
+	/** Note: For curve types, the additive branch will always be evaluated. Curve weight blending will occur at task evaluation time. */
+	UPROPERTY(EditAnywhere, Category = Settings, meta = (EditConditionHides, EditCondition = "AlphaInputType == EAnimAlphaInputType::Curve"))
+	FName AlphaCurveName = NAME_None;
+
+	UPROPERTY(EditAnywhere, Category = Settings)
+	EAnimAlphaInputType AlphaInputType = EAnimAlphaInputType::Float;
+
+	// Latent pin support boilerplate
+#define TRAIT_LATENT_PROPERTIES_ENUMERATOR(GeneratorMacro) \
+		GeneratorMacro(Alpha) \
+		GeneratorMacro(bAlphaBoolEnabled) \
+		GeneratorMacro(AlphaBoolBlend) \
+		GeneratorMacro(AlphaCurveName) \
+		GeneratorMacro(AlphaInputType) \
+
+	GENERATE_TRAIT_LATENT_PROPERTIES(FAlphaInputTraitArgs, TRAIT_LATENT_PROPERTIES_ENUMERATOR)
+#undef TRAIT_LATENT_PROPERTIES_ENUMERATOR
+};
+
+
+namespace UE::AnimNext
+{
+	/**
+	 * IAlphaInputArgs
+	 *
+	 * This interface exposes alpha input args & convience methods for operating on those args
+	 * 
+	 * Alpha input is typically a curve / float / bool that controls the strength of some other trait.
+	 * Ex: 0 - no effect, 1 - full effect.
+	 */
+	struct ANIMNEXTANIMGRAPH_API IAlphaInputArgs : ITraitInterface
+	{
+		DECLARE_ANIM_TRAIT_INTERFACE(IAlphaInputArgs)
+
+		/**
+		 * Returns the args captured by this trait interface.
+		 * Note: We return a copy as some implementations may return a temp struct & our return type is relatively cheap to copy.
+		 */
+		virtual FAlphaInputTraitArgs Get(const FExecutionContext& Context, const TTraitBinding<IAlphaInputArgs>& Binding) const;
+
+		/**
+		 * Returns current alpha input type
+		 */
+		virtual EAnimAlphaInputType GetAlphaInputType(const FExecutionContext& Context, const TTraitBinding<IAlphaInputArgs>& Binding) const;
+
+		/**
+		 * Returns current alpha curve name, if any
+		 */
+		virtual FName GetAlphaCurveName(const FExecutionContext& Context, const TTraitBinding<IAlphaInputArgs>& Binding) const;
+
+		/**
+		 * Returns alpha value for current input. For curve types this will be 1.0, so curves should use trait args to resolve alpha in a task.
+		 *
+		 * @param DeltaTime - Time since last API call using DeltaTime (Ex: Time since last 'ComputeAlphaValue' call).
+		 */
+		virtual float GetCurrentAlphaValue(const FExecutionContext& Context, const TTraitBinding<IAlphaInputArgs>& Binding) const;
+
+		/**
+		 * Returns callback to clamp alpha values, typically used for curve sampling
+		 * Note: Callback may perform a deferred write to alpha input instance data.
+		 */
+		virtual TFunction<float(float)> GetInputScaleBiasClampCallback(const FExecutionContext& Context, const TTraitBinding<IAlphaInputArgs>& Binding) const;
+
+#if WITH_EDITOR
+		virtual const FText& GetDisplayName() const override;
+		virtual const FText& GetDisplayShortName() const override;
+#endif // WITH_EDITOR
+	};
+
+	/**
+	 * Specialization for trait binding.
+	 */
+	template<>
+	struct TTraitBinding<IAlphaInputArgs> : FTraitBinding
+	{
+		// @see IAlphaInputArgs::Get
+		FAlphaInputTraitArgs Get(const FExecutionContext& Context) const
+		{
+			return GetInterface()->Get(Context, *this);
+		}
+
+		// @see IAlphaInputArgs::GetAlphaInputType
+		EAnimAlphaInputType GetAlphaInputType(const FExecutionContext& Context) const
+		{
+			return GetInterface()->GetAlphaInputType(Context, *this);
+		}
+
+		// @see IAlphaInputArgs::GetAlphaCurveName
+		FName GetAlphaCurveName(const FExecutionContext& Context) const
+		{
+			return GetInterface()->GetAlphaCurveName(Context, *this);
+		}
+
+		// @see IAlphaInputArgs::GetCurrentAlphaValue
+		float GetCurrentAlphaValue(const FExecutionContext& Context) const
+		{
+			return GetInterface()->GetCurrentAlphaValue(Context, *this);
+		}
+
+		// @see IAlphaInputArgs::GetInputScaleBiasClampCallback
+		TFunction<float(float)> GetInputScaleBiasClampCallback(const FExecutionContext& Context) const
+		{
+			return GetInterface()->GetInputScaleBiasClampCallback(Context, *this);
+		}
+
+	protected:
+		const IAlphaInputArgs* GetInterface() const { return GetInterfaceTyped<IAlphaInputArgs>(); }
+	};
+}
diff --git a/Engine/Plugins/Experimental/AnimNextAnimGraph/Source/AnimNextAnimGraph/Public/Traits/PassthroughBlendTrait.h b/Engine/Plugins/Experimental/AnimNextAnimGraph/Source/AnimNextAnimGraph/Public/Traits/PassthroughBlendTrait.h
--- a/Engine/Plugins/Experimental/AnimNextAnimGraph/Source/AnimNextAnimGraph/Public/Traits/PassthroughBlendTrait.h
+++ b/Engine/Plugins/Experimental/AnimNextAnimGraph/Source/AnimNextAnimGraph/Public/Traits/PassthroughBlendTrait.h
@@ -68,7 +68,6 @@ struct FPassthroughBlendTrait : FAdditiveTrait, IUpdate, IEvaluate
 		FInputScaleBiasClamp AlphaScaleBiasClamp;
 
 		float DeltaTime = 0.f;
-		float InternalAlpha = 1.f;
 		float ComputedAlphaValue = 0.f;
 
 		void Construct(const FExecutionContext& Context, const FTraitBinding& Binding);
@@ -81,7 +80,7 @@ struct FPassthroughBlendTrait : FAdditiveTrait, IUpdate, IEvaluate
 	// IEvaluate impl
 	virtual void PostEvaluate(FEvaluateTraversalContext& Context, const TTraitBinding<IEvaluate>& Binding) const override;
 
-	static float ComputeAlphaValue(const EAnimAlphaInputType AlphaInputType, const FPassthroughBlendTraitSharedData* SharedData, FInstanceData* InstanceData, float DeltaTime);
+	static float ComputeAlphaValue(const EAnimAlphaInputType AlphaInputType, const FPassthroughBlendTraitSharedData* SharedData, FInstanceData* InstanceData, const TTraitBinding<IUpdate>& Binding, float DeltaTime);
 };
 
 } // namespace UE::AnimNext
