Added shared data as a parameter to trait function IsPropertyLatent
Implemented IsPropertyLatent at ControlRigTrait to fix cook issues

diff --git a/Engine/Plugins/Experimental/AnimNextAnimGraph/Source/AnimNextAnimGraph/Private/Graph/RigDecorator_AnimNextCppTrait.cpp b/Engine/Plugins/Experimental/AnimNextAnimGraph/Source/AnimNextAnimGraph/Private/Graph/RigDecorator_AnimNextCppTrait.cpp
--- a/Engine/Plugins/Experimental/AnimNextAnimGraph/Source/AnimNextAnimGraph/Private/Graph/RigDecorator_AnimNextCppTrait.cpp
+++ b/Engine/Plugins/Experimental/AnimNextAnimGraph/Source/AnimNextAnimGraph/Private/Graph/RigDecorator_AnimNextCppTrait.cpp
@@ -75,7 +75,7 @@ void FRigDecorator_AnimNextCppDecorator::GetProgrammaticPins(URigVMController* I
 		// Except for trait handles which are never lazy since they just encode graph connectivity
 		if (InParentPinIndex == PinInfo.ParentIndex && !bIsHidden && !bIsInline && !bIsTraitHandle)
 		{
-			if (Trait->IsPropertyLatent(PinInfo.Property->GetFName()))
+			if (Trait->IsPropertyLatent(*TraitSharedData, PinInfo.Property->GetFName()))
 			{
 				PinInfo.bIsLazy = true;
 			}
diff --git a/Engine/Plugins/Experimental/AnimNextAnimGraph/Source/AnimNextAnimGraph/Private/TraitCore/Trait.cpp b/Engine/Plugins/Experimental/AnimNextAnimGraph/Source/AnimNextAnimGraph/Private/TraitCore/Trait.cpp
--- a/Engine/Plugins/Experimental/AnimNextAnimGraph/Source/AnimNextAnimGraph/Private/TraitCore/Trait.cpp
+++ b/Engine/Plugins/Experimental/AnimNextAnimGraph/Source/AnimNextAnimGraph/Private/TraitCore/Trait.cpp
@@ -175,7 +175,7 @@ namespace UE::AnimNext
 					continue;	// Skip non-latent properties
 				}
 
-				if (!IsPropertyLatent(Property->GetFName()))
+				if (!IsPropertyLatent(*InSharedData, Property->GetFName()))
 				{
 					// Skip properties not part of the latent macro enumerator
 					// This can occur if the property is missing from the enumerator by mistake
diff --git a/Engine/Plugins/Experimental/AnimNextAnimGraph/Source/AnimNextAnimGraph/Public/TraitCore/Trait.h b/Engine/Plugins/Experimental/AnimNextAnimGraph/Source/AnimNextAnimGraph/Public/TraitCore/Trait.h
--- a/Engine/Plugins/Experimental/AnimNextAnimGraph/Source/AnimNextAnimGraph/Public/TraitCore/Trait.h
+++ b/Engine/Plugins/Experimental/AnimNextAnimGraph/Source/AnimNextAnimGraph/Public/TraitCore/Trait.h
@@ -47,7 +47,7 @@ struct FRigVMPinInfoArray;
 #define ANIM_NEXT_IMPL_DECLARE_ANIM_TRAIT_LATENT_PROPERTY_SUPPORT() \
 	virtual uint32 GetNumLatentTraitProperties() const override { return -FSharedData::GetLatentPropertyIndex(~(size_t)0); } \
 	virtual UE::AnimNext::FTraitLatentPropertyMemoryLayout GetLatentPropertyMemoryLayout(const FAnimNextTraitSharedData& SharedData, FName PropertyName, uint32 PropertyIndex) const override; \
-	virtual bool IsPropertyLatent(FName PropertyName) const override;
+	virtual bool IsPropertyLatent(const FAnimNextTraitSharedData& InSharedData, FName PropertyName) const override;
 #else
 #define ANIM_NEXT_IMPL_DECLARE_ANIM_TRAIT_LATENT_PROPERTY_SUPPORT() \
 	virtual uint32 GetNumLatentTraitProperties() const override { return -FSharedData::GetLatentPropertyIndex(~(size_t)0); } \
@@ -117,7 +117,7 @@ struct FRigVMPinInfoArray;
 
 #if WITH_EDITOR
 #define ANIM_NEXT_IMPL_DEFINE_ANIM_TRAIT_IS_PROPERTY_LATENT(TraitName) \
-	bool TraitName::IsPropertyLatent(FName PropertyName) const \
+	bool TraitName::IsPropertyLatent(const FAnimNextTraitSharedData& InSharedData, FName PropertyName) const \
 	{ \
 		const UScriptStruct* SharedDataStruct = FSharedData::StaticStruct(); \
 		const FProperty* Property = SharedDataStruct->FindPropertyByName(PropertyName); \
@@ -452,7 +452,7 @@ namespace UE::AnimNext
 			const TFunction<uint16(FName PropertyName)>& GetTraitLatentPropertyIndex) const;
 
 		// Returns whether or not the specified property has been marked as latent
-		virtual bool IsPropertyLatent(FName PropertyName) const { return false; }
+		virtual bool IsPropertyLatent(const FAnimNextTraitSharedData& InSharedData, FName PropertyName) const { return false; }
 
 		// Makes the Trait Editor only display this Trait in Advanced view
 		virtual bool IsHidden() const { return false; }
diff --git a/Engine/Plugins/Experimental/AnimNextControlRig/Source/AnimNextControlRig/Internal/ControlRigTrait.h b/Engine/Plugins/Experimental/AnimNextControlRig/Source/AnimNextControlRig/Internal/ControlRigTrait.h
--- a/Engine/Plugins/Experimental/AnimNextControlRig/Source/AnimNextControlRig/Internal/ControlRigTrait.h
+++ b/Engine/Plugins/Experimental/AnimNextControlRig/Source/AnimNextControlRig/Internal/ControlRigTrait.h
@@ -263,6 +263,8 @@ protected:
 		const FString& InDefaultValue,
 		struct FRigVMPinInfoArray& OutPinArray) const override;
 
+	bool IsPropertyLatent(const FAnimNextTraitSharedData& InSharedData, FName PropertyName) const override;
+
 	virtual uint32 GetLatentPropertyHandles(
 		const FAnimNextTraitSharedData* InSharedData,
 		TArray<FLatentPropertyMetadata>& OutLatentPropertyHandles,
diff --git a/Engine/Plugins/Experimental/AnimNextControlRig/Source/AnimNextControlRig/Private/ControlRigTrait.cpp b/Engine/Plugins/Experimental/AnimNextControlRig/Source/AnimNextControlRig/Private/ControlRigTrait.cpp
--- a/Engine/Plugins/Experimental/AnimNextControlRig/Source/AnimNextControlRig/Private/ControlRigTrait.cpp
+++ b/Engine/Plugins/Experimental/AnimNextControlRig/Source/AnimNextControlRig/Private/ControlRigTrait.cpp
@@ -640,8 +640,7 @@ void FControlRigTrait::GetProgrammaticPins(FAnimNextTraitSharedData* InSharedDat
 						const TRigVMTypeIndex TypeIndex = RigVMTypeUtils::TypeIndexFromPinType(Control.PinType);
 						if (ensure(TypeIndex != INDEX_NONE))
 						{
-							const int32 Index = OutPinArray.AddPin(InController, InParentPinIndex, ControlName, ERigVMPinDirection::Input, TypeIndex, ControlDefaultValue, ERigVMPinDefaultValueType::AutoDetect, nullptr, nullptr, true);
-							OutPinArray[Index].bIsLazy = true;
+							OutPinArray.AddPin(InController, InParentPinIndex, ControlName, ERigVMPinDirection::Input, TypeIndex, ControlDefaultValue, ERigVMPinDefaultValueType::AutoDetect, nullptr, nullptr, true);
 						}
 					}
 				}
@@ -650,6 +649,31 @@ void FControlRigTrait::GetProgrammaticPins(FAnimNextTraitSharedData* InSharedDat
 	}
 }
 
+bool FControlRigTrait::IsPropertyLatent(const FAnimNextTraitSharedData& InSharedData, FName PropertyName) const
+{
+	const FControlRigTraitSharedData& SharedData = (const FControlRigTraitSharedData&)InSharedData;
+
+	// --- Exposed Public Variables ---
+	if (SharedData.ExposedPropertyVariableNames.Num() > 0)
+	{
+		if (SharedData.ExposedPropertyVariableNames.Contains(PropertyName))
+		{
+			return true;
+		}
+	}
+
+	// --- Exposed Controls ---
+	if (SharedData.ExposedPropertyControlNames.Num() > 0)
+	{
+		if (SharedData.ExposedPropertyControlNames.Contains(PropertyName))
+		{
+			return true;
+		}
+	}
+
+	return false;
+}
+
 uint32 FControlRigTrait::GetLatentPropertyHandles(
 	const FAnimNextTraitSharedData* InSharedData,
 	TArray<FLatentPropertyMetadata>& OutLatentPropertyHandles,
