AnimNext: Fix ensure/crash when live recompiling anim graphs

Fixes crash by ensuring that the cached child ptr to the root of a hosted graph instance is not used, we instead read the child ptr directly
Owners of graph instances (e.g. blend stack, blend spaces etc.) were not marking their graphs as 'updated', so relevancy of newly recompiled graphs was not being taken into account.

diff --git a/Engine/Plugins/Experimental/AnimNextAnimGraph/Source/AnimNextAnimGraph/Internal/Traits/BlendStackTrait.h b/Engine/Plugins/Experimental/AnimNextAnimGraph/Source/AnimNextAnimGraph/Internal/Traits/BlendStackTrait.h
--- a/Engine/Plugins/Experimental/AnimNextAnimGraph/Source/AnimNextAnimGraph/Internal/Traits/BlendStackTrait.h
+++ b/Engine/Plugins/Experimental/AnimNextAnimGraph/Source/AnimNextAnimGraph/Internal/Traits/BlendStackTrait.h
@@ -119,6 +119,19 @@ namespace UE::AnimNext
 
 			void Initialize(IBlendStack::FGraphRequest&& GraphRequest);
 			void Terminate();
+
+			FWeakTraitPtr GetChildPtr() const
+			{
+				if (Request.Type == EGraphRequestType::Child)
+				{
+					return ChildPtr;
+				}
+				else if (Instance)
+				{
+					return Instance->GraphInstancePtr;
+				}
+				return FWeakTraitPtr();
+			}
 		};
 
 		struct FInstanceData : FTrait::FInstanceData
diff --git a/Engine/Plugins/Experimental/AnimNextAnimGraph/Source/AnimNextAnimGraph/Private/Graph/AnimNextGraphInstance.cpp b/Engine/Plugins/Experimental/AnimNextAnimGraph/Source/AnimNextAnimGraph/Private/Graph/AnimNextGraphInstance.cpp
--- a/Engine/Plugins/Experimental/AnimNextAnimGraph/Source/AnimNextAnimGraph/Private/Graph/AnimNextGraphInstance.cpp
+++ b/Engine/Plugins/Experimental/AnimNextAnimGraph/Source/AnimNextAnimGraph/Private/Graph/AnimNextGraphInstance.cpp
@@ -118,15 +118,10 @@ bool FAnimNextGraphInstance::UsesEntryPoint(FName InEntryPoint) const
 }
 
 bool FAnimNextGraphInstance::IsRoot() const
-{
+{	
 	return this == RootGraphInstance;
 }
 
-bool FAnimNextGraphInstance::HasUpdated() const
-{
-	return bHasUpdatedOnce;
-}
-
 void FAnimNextGraphInstance::AddStructReferencedObjects(FReferenceCollector& Collector)
 {
 	if (!IsRoot())
@@ -160,7 +155,12 @@ GraphInstanceComponentMapType::TConstIterator FAnimNextGraphInstance::GetCompone
 	return RootGraphInstance->Components.CreateConstIterator();
 }
 
-void FAnimNextGraphInstance::Update()
+bool FAnimNextGraphInstance::HasUpdated() const
+{
+	return bHasUpdatedOnce;
+}
+
+void FAnimNextGraphInstance::MarkAsUpdated()
 {
 	bHasUpdatedOnce = true;
 }
diff --git a/Engine/Plugins/Experimental/AnimNextAnimGraph/Source/AnimNextAnimGraph/Private/PlayAnim/PlayAnimSlotTrait.cpp b/Engine/Plugins/Experimental/AnimNextAnimGraph/Source/AnimNextAnimGraph/Private/PlayAnim/PlayAnimSlotTrait.cpp
--- a/Engine/Plugins/Experimental/AnimNextAnimGraph/Source/AnimNextAnimGraph/Private/PlayAnim/PlayAnimSlotTrait.cpp
+++ b/Engine/Plugins/Experimental/AnimNextAnimGraph/Source/AnimNextAnimGraph/Private/PlayAnim/PlayAnimSlotTrait.cpp
@@ -135,7 +135,7 @@ namespace UE::AnimNext
 		for (const FPlayAnimSlotRequest& SlotRequest : InstanceData->SlotRequests)
 		{
 			// Even if the request is inactive, we queue an empty handle
-			Children.Add(SlotRequest.ChildPtr);
+			Children.Add(SlotRequest.GetChildPtr());
 		}
 	}
 
@@ -243,7 +243,7 @@ namespace UE::AnimNext
 			}
 
 			FTraitStackBinding ChildStack;
-			ensure(Context.GetStack(SlotRequest.ChildPtr, ChildStack));
+			ensure(Context.GetStack(SlotRequest.GetChildPtr(), ChildStack));
 
 			TTraitBinding<ITimeline> ChildTimelineTrait;
 			ensure(ChildStack.GetInterface(ChildTimelineTrait));
@@ -332,13 +332,19 @@ namespace UE::AnimNext
 		{
 			const FPlayAnimSlotRequest& SlotRequest = InstanceData->SlotRequests[RequestIndex];
 			const float BlendWeight = DiscreteBlendTrait.GetBlendWeight(Context, RequestIndex);
+			const bool bGraphHasNeverUpdated = SlotRequest.GraphInstance.IsValid() && !SlotRequest.GraphInstance->HasUpdated();
 
 			FTraitUpdateState RequestSlotTraitState = TraitState
 				.WithWeight(BlendWeight)
 				.AsBlendingOut(RequestIndex != InstanceData->CurrentlyActiveRequestIndex)
-				.AsNewlyRelevant(!SlotRequest.bWasRelevant);
+				.AsNewlyRelevant(!SlotRequest.bWasRelevant || bGraphHasNeverUpdated);
 
-			TraversalQueue.Push(InstanceData->SlotRequests[RequestIndex].ChildPtr, RequestSlotTraitState);
+			if (SlotRequest.GraphInstance.IsValid())
+			{
+				SlotRequest.GraphInstance->MarkAsUpdated();
+			}
+
+			TraversalQueue.Push(InstanceData->SlotRequests[RequestIndex].GetChildPtr(), RequestSlotTraitState);
 		}
 	}
 
@@ -393,7 +399,6 @@ namespace UE::AnimNext
 			{
 				const FName EntryPoint = SlotRequest.AnimationGraph->DefaultEntryPoint;
 				SlotRequest.GraphInstance = SlotRequest.AnimationGraph->AllocateInstance(nullptr, &Binding.GetTraitPtr().GetNodeInstance()->GetOwner(), EntryPoint);
-				SlotRequest.ChildPtr = SlotRequest.GraphInstance.IsValid() ? SlotRequest.GraphInstance->GetGraphRootPtr() : FTraitPtr();
 				SlotRequest.GraphInstance->BindPublicVariables(SlotRequest.Payload.Get());
 
 				// TODO: Validate that our child implements the ITimeline interface
diff --git a/Engine/Plugins/Experimental/AnimNextAnimGraph/Source/AnimNextAnimGraph/Private/PlayAnim/PlayAnimSlotTrait.h b/Engine/Plugins/Experimental/AnimNextAnimGraph/Source/AnimNextAnimGraph/Private/PlayAnim/PlayAnimSlotTrait.h
--- a/Engine/Plugins/Experimental/AnimNextAnimGraph/Source/AnimNextAnimGraph/Private/PlayAnim/PlayAnimSlotTrait.h
+++ b/Engine/Plugins/Experimental/AnimNextAnimGraph/Source/AnimNextAnimGraph/Private/PlayAnim/PlayAnimSlotTrait.h
@@ -112,6 +112,19 @@ namespace UE::AnimNext
 
 			// Initializes a request to begin playing
 			void Initialize(FInjectionRequestPtr InRequest, const FInjectionBlendSettings& InBlendSettings, const UAnimNextAnimationGraph* InAnimationGraph, FAnimNextDataInterfacePayload&& InPayload);
+
+			FWeakTraitPtr GetChildPtr() const
+			{
+				if (State == EPlayAnimRequestState::ActiveSource)
+				{
+					return ChildPtr;
+				}
+				else if (GraphInstance)
+				{
+					return GraphInstance->GraphInstancePtr;
+				}
+				return FWeakTraitPtr();
+			}
 		};
 
 		using FSharedData = FAnimNextPlayAnimSlotTraitSharedData;
diff --git a/Engine/Plugins/Experimental/AnimNextAnimGraph/Source/AnimNextAnimGraph/Private/TraitInterfaces/IUpdate.cpp b/Engine/Plugins/Experimental/AnimNextAnimGraph/Source/AnimNextAnimGraph/Private/TraitInterfaces/IUpdate.cpp
--- a/Engine/Plugins/Experimental/AnimNextAnimGraph/Source/AnimNextAnimGraph/Private/TraitInterfaces/IUpdate.cpp
+++ b/Engine/Plugins/Experimental/AnimNextAnimGraph/Source/AnimNextAnimGraph/Private/TraitInterfaces/IUpdate.cpp
@@ -507,8 +507,8 @@ namespace UE::AnimNext
 			FTraitUpdateState(DeltaTime)
 			.AsNewlyRelevant(!GraphInstance.HasUpdated());
 
-		// Update the graph instance itself
-		GraphInstance.Update();
+		// Mark the graph instance itself as updated
+		GraphInstance.MarkAsUpdated();
 
 		// Before we start the traversal, we give the graph instance components the chance to do some work
 		TraversalContext.BindTo(GraphInstance);
diff --git a/Engine/Plugins/Experimental/AnimNextAnimGraph/Source/AnimNextAnimGraph/Private/Traits/BlendSpacePlayer.cpp b/Engine/Plugins/Experimental/AnimNextAnimGraph/Source/AnimNextAnimGraph/Private/Traits/BlendSpacePlayer.cpp
--- a/Engine/Plugins/Experimental/AnimNextAnimGraph/Source/AnimNextAnimGraph/Private/Traits/BlendSpacePlayer.cpp
+++ b/Engine/Plugins/Experimental/AnimNextAnimGraph/Source/AnimNextAnimGraph/Private/Traits/BlendSpacePlayer.cpp
@@ -189,8 +189,14 @@ namespace UE::AnimNext
 			{
 				if (SampleState.bSampledThisFrame)
 				{
-					bool bNewlyRelevant = SampleState.bSampledThisFrame && !SampleState.bSampledLastFrame;
+					const bool bGraphHasNeverUpdated = SampleState.Instance.IsValid() && !SampleState.Instance->HasUpdated();
+					bool bNewlyRelevant = (SampleState.bSampledThisFrame && !SampleState.bSampledLastFrame) || bGraphHasNeverUpdated;
 					TraversalQueue.Push(SampleState.ChildPtr, TraitState.WithWeight(SampleState.Weight).AsNewlyRelevant(bNewlyRelevant));
+
+					if (SampleState.Instance.IsValid())
+					{
+						SampleState.Instance->MarkAsUpdated();
+					}
 				}
 			}
 		}
diff --git a/Engine/Plugins/Experimental/AnimNextAnimGraph/Source/AnimNextAnimGraph/Private/Traits/BlendStackTrait.cpp b/Engine/Plugins/Experimental/AnimNextAnimGraph/Source/AnimNextAnimGraph/Private/Traits/BlendStackTrait.cpp
--- a/Engine/Plugins/Experimental/AnimNextAnimGraph/Source/AnimNextAnimGraph/Private/Traits/BlendStackTrait.cpp
+++ b/Engine/Plugins/Experimental/AnimNextAnimGraph/Source/AnimNextAnimGraph/Private/Traits/BlendStackTrait.cpp
@@ -134,7 +134,7 @@ namespace UE::AnimNext
 		for (const FBlendStackTrait::FGraphState& ChildGraph : InstanceData->ChildGraphs)
 		{
 			// Even if the request is inactive, we queue an empty handle
-			Children.Add(ChildGraph.ChildPtr);
+			Children.Add(ChildGraph.GetChildPtr());
 		}
 	}
 
@@ -150,14 +150,21 @@ namespace UE::AnimNext
 		{
 			FGraphState& Graph = InstanceData->ChildGraphs[ChildIndex];
 			const float BlendWeight = DiscreteBlendTrait.GetBlendWeight(Context, ChildIndex);
+			const bool bGraphHasNeverUpdated = Graph.Instance.IsValid() && !Graph.Instance->HasUpdated();
 
+			// Flag the child instance as updated
 			FTraitUpdateState ChildGraphTraitState = TraitState
 				.WithWeight(BlendWeight)
 				.AsBlendingOut(ChildIndex != InstanceData->CurrentlyActiveGraphIndex)
-				.AsNewlyRelevant(Graph.bNewlyCreated);
+				.AsNewlyRelevant(Graph.bNewlyCreated || bGraphHasNeverUpdated);
 			Graph.bNewlyCreated = false;
+			
+			if (Graph.Instance.IsValid())
+			{
+				Graph.Instance->MarkAsUpdated();
+			}
 
-			TraversalQueue.Push(InstanceData->ChildGraphs[ChildIndex].ChildPtr, ChildGraphTraitState);
+			TraversalQueue.Push(InstanceData->ChildGraphs[ChildIndex].GetChildPtr(), ChildGraphTraitState);
 		}
 	}
 
@@ -213,7 +220,6 @@ namespace UE::AnimNext
 					const FName EntryPoint = NAME_None;
 					FAnimNextGraphInstance& Owner = Binding.GetTraitPtr().GetNodeInstance()->GetOwner();
 					Graph.Instance = Graph.Request.AnimationGraph->AllocateInstance(Owner.GetModuleInstance(), &Owner, EntryPoint);
-					Graph.ChildPtr = Graph.Instance->GetGraphRootPtr();
 					FModuleInjectionDataInterfaceAdapter ModuleAdapter(Context.GetRootGraphInstance().GetModuleInstance(), Graph.Request.BindingModuleHandle);
 					IDataInterfaceHost* DataInterfaceHost = &ModuleAdapter;
 					Graph.Instance->BindPublicVariables(Graph.Request.GraphPayload.Get(), TConstArrayView<IDataInterfaceHost*>(&DataInterfaceHost, 1));
@@ -359,7 +365,7 @@ namespace UE::AnimNext
 			const FBlendStackCoreTrait::FGraphState& ActiveGraph = InstanceData->ChildGraphs[CurrentlyActiveGraphIndex];
 
 			FTraitStackBinding ChildTraitStack;
-			ensure(Context.GetStack(ActiveGraph.ChildPtr, ChildTraitStack));
+			ensure(Context.GetStack(ActiveGraph.GetChildPtr(), ChildTraitStack));
 
 			TTraitBinding<ITimeline> Timeline;
 			FTraitStackBinding StackBinding;
diff --git a/Engine/Plugins/Experimental/AnimNextAnimGraph/Source/AnimNextAnimGraph/Private/Traits/SubGraphHost.cpp b/Engine/Plugins/Experimental/AnimNextAnimGraph/Source/AnimNextAnimGraph/Private/Traits/SubGraphHost.cpp
--- a/Engine/Plugins/Experimental/AnimNextAnimGraph/Source/AnimNextAnimGraph/Private/Traits/SubGraphHost.cpp
+++ b/Engine/Plugins/Experimental/AnimNextAnimGraph/Source/AnimNextAnimGraph/Private/Traits/SubGraphHost.cpp
@@ -156,11 +156,17 @@ namespace UE::AnimNext
 			}
 
 			const float BlendWeight = DiscreteBlendTrait.GetBlendWeight(Context, SubGraphIndex);
+			const bool bGraphHasNeverUpdated = SubGraphSlot.GraphInstance.IsValid() && !SubGraphSlot.GraphInstance->HasUpdated();
 
 			FTraitUpdateState SubGraphTraitState = TraitState
 				.WithWeight(BlendWeight)
 				.AsBlendingOut(SubGraphIndex != InstanceData->CurrentlyActiveSubGraphIndex)
-				.AsNewlyRelevant(!SubGraphSlot.bWasRelevant);
+				.AsNewlyRelevant(!SubGraphSlot.bWasRelevant || bGraphHasNeverUpdated);
+
+			if (SubGraphSlot.GraphInstance.IsValid())
+			{
+				SubGraphSlot.GraphInstance->MarkAsUpdated();
+			}
 
 			TraversalQueue.Push(SubGraphSlot.GraphInstance->GetGraphRootPtr(), SubGraphTraitState);
 		}
diff --git a/Engine/Plugins/Experimental/AnimNextAnimGraph/Source/AnimNextAnimGraph/Public/Graph/AnimNextGraphInstance.h b/Engine/Plugins/Experimental/AnimNextAnimGraph/Source/AnimNextAnimGraph/Public/Graph/AnimNextGraphInstance.h
--- a/Engine/Plugins/Experimental/AnimNextAnimGraph/Source/AnimNextAnimGraph/Public/Graph/AnimNextGraphInstance.h
+++ b/Engine/Plugins/Experimental/AnimNextAnimGraph/Source/AnimNextAnimGraph/Public/Graph/AnimNextGraphInstance.h
@@ -96,9 +96,6 @@ struct ANIMNEXTANIMGRAPH_API FAnimNextGraphInstance : public FAnimNextDataInterf
 	// Returns whether or not this graph instance is the root graph instance or false otherwise
 	bool IsRoot() const;
 
-	// Returns whether or not this graph instance has updated at least once
-	bool HasUpdated() const;
-
 	// Adds strong/hard object references during GC
 	void AddStructReferencedObjects(class FReferenceCollector& Collector);
 
@@ -117,8 +114,11 @@ struct ANIMNEXTANIMGRAPH_API FAnimNextGraphInstance : public FAnimNextDataInterf
 	// Returns const iterators to the graph instance component container
 	GraphInstanceComponentMapType::TConstIterator GetComponentIterator() const;
 
-	// Called each time the graph updates
-	void Update();
+	// Returns whether or not this graph instance has updated at least once
+	bool HasUpdated() const;
+
+	// Called each time the graph updates to mark the instance as updated
+	void MarkAsUpdated();
 
 	// Get a reference to the internal public variables
 	const FInstancedPropertyBag& GetVariables() const { return Variables; }
