Revert per-bone attribute code in blend smoother per bone task to use default attribute blending to unblock internal project's root-motion locomotion

diff --git a/Engine/Plugins/Experimental/AnimNextAnimGraph/Source/AnimNextAnimGraph/Private/EvaluationVM/Tasks/BlendKeyframesPerBone.cpp b/Engine/Plugins/Experimental/AnimNextAnimGraph/Source/AnimNextAnimGraph/Private/EvaluationVM/Tasks/BlendKeyframesPerBone.cpp
--- a/Engine/Plugins/Experimental/AnimNextAnimGraph/Source/AnimNextAnimGraph/Private/EvaluationVM/Tasks/BlendKeyframesPerBone.cpp
+++ b/Engine/Plugins/Experimental/AnimNextAnimGraph/Source/AnimNextAnimGraph/Private/EvaluationVM/Tasks/BlendKeyframesPerBone.cpp
@@ -90,7 +90,8 @@ void FAnimNextBlendOverwriteKeyframePerBoneWithScaleTask::Execute(UE::AnimNext::
 
 	if (EnumHasAnyFlags(VM.GetFlags(), EEvaluationFlags::Attributes))
 	{
-		UE::Anim::Attributes::BlendAttributesPerBone({ Keyframe->Attributes }, { LODBoneIndexToWeightIndexMap }, { *BlendData }, { 0 }, { Keyframe->Attributes });
+		UE::Anim::Attributes::OverrideAttributes(Keyframe->Attributes, Keyframe->Attributes, BlendWeight);
+		//UE::Anim::Attributes::BlendAttributesPerBone({ Keyframe->Attributes }, { LODBoneIndexToWeightIndexMap }, { *BlendData }, { 0 }, { Keyframe->Attributes });
 	}
 
 	VM.PushValue(KEYFRAME_STACK_NAME, MoveTemp(Keyframe));
@@ -182,6 +183,8 @@ void FAnimNextBlendAddKeyframePerBoneWithScaleTask::Execute(UE::AnimNext::FEvalu
 
 	if (EnumHasAnyFlags(VM.GetFlags(), EEvaluationFlags::Attributes))
 	{
+		UE::Anim::Attributes::AccumulateAttributes(KeyframeA->Attributes, KeyframeB->Attributes, BlendWeightA, AAT_None);
+		/*
 		// TODO: Might need to be revisited, might not work as intended
 		UE::Anim::Attributes::BlendAttributesPerBone(
 			{ KeyframeA->Attributes, KeyframeB->Attributes },
@@ -189,6 +192,7 @@ void FAnimNextBlendAddKeyframePerBoneWithScaleTask::Execute(UE::AnimNext::FEvalu
 			{ *BlendDataA, *BlendDataB },
 			{ 0, 1 },
 			{ KeyframeA->Attributes });
+		*/
 	}
 
 	VM.PushValue(KEYFRAME_STACK_NAME, MoveTemp(KeyframeB));
