AnimNextAnimGraph - adding bResetChildOnActivation and bAlwaysUpdateChildren to FBlendTwoWayTrait (to mimic FAnimNode_TwoWayBlend behavior)

diff --git a/Engine/Plugins/Experimental/AnimNextAnimGraph/Source/AnimNextAnimGraph/Private/Traits/BlendTwoWay.cpp b/Engine/Plugins/Experimental/AnimNextAnimGraph/Source/AnimNextAnimGraph/Private/Traits/BlendTwoWay.cpp
--- a/Engine/Plugins/Experimental/AnimNextAnimGraph/Source/AnimNextAnimGraph/Private/Traits/BlendTwoWay.cpp
+++ b/Engine/Plugins/Experimental/AnimNextAnimGraph/Source/AnimNextAnimGraph/Private/Traits/BlendTwoWay.cpp
@@ -49,46 +49,60 @@ namespace UE::AnimNext
 		TTraitBinding<IContinuousBlend> ContinuousBlendTrait;
 		Binding.GetStackInterface(ContinuousBlendTrait);
 
+		const bool bAlwaysUpdateChildren = SharedData->GetbAlwaysUpdateChildren(Binding);
+		const bool bResetChildOnActivation = SharedData->GetbResetChildOnActivation(Binding);
+
 		const float BlendWeightB = ContinuousBlendTrait.GetBlendWeight(Context, 1);
-		if (!FAnimWeight::IsFullWeight(BlendWeightB))
+
+		const bool bIsChildARelevant = !FAnimWeight::IsFullWeight(BlendWeightB);
+		const bool bIsChildBRelevant = FAnimWeight::IsRelevant(BlendWeightB);
+
+		const bool bIsChildANeeded = bAlwaysUpdateChildren || bIsChildARelevant;
+		const bool bIsChildBNeeded = bAlwaysUpdateChildren || bIsChildBRelevant;
+
+		InstanceData->bIsChildANewlyRelevant = false;
+		InstanceData->bIsChildBNewlyRelevant = false;
+
+		if (bIsChildANeeded)
 		{
 			if (!InstanceData->ChildA.IsValid())
 			{
 				// We need to blend a child that isn't instanced yet, allocate it
 				InstanceData->ChildA = Context.AllocateNodeInstance(Binding, SharedData->ChildA);
+				InstanceData->bIsChildANewlyRelevant = true;
 			}
-			else
+			else if (bResetChildOnActivation && bIsChildARelevant && !InstanceData->bWasChildARelevant)
 			{
-				InstanceData->bWasChildARelevant = true;
-			}
-
-			if (!FAnimWeight::IsRelevant(BlendWeightB))
-			{
-				// We no longer need this child, release it
-				InstanceData->ChildB.Reset();
-				InstanceData->bWasChildBRelevant = false;
+				InstanceData->bIsChildANewlyRelevant = true;
 			}
 		}
+		else
+		{
+			// We no longer need this child, release it
+			InstanceData->ChildA.Reset();
+		}
 
-		if (FAnimWeight::IsRelevant(BlendWeightB))
+		if (bIsChildBNeeded)
 		{
 			if (!InstanceData->ChildB.IsValid())
 			{
 				// We need to blend a child that isn't instanced yet, allocate it
 				InstanceData->ChildB = Context.AllocateNodeInstance(Binding, SharedData->ChildB);
+				InstanceData->bIsChildBNewlyRelevant = true;
 			}
-			else
-			{
-				InstanceData->bWasChildBRelevant = true;
-			}
-
-			if (FAnimWeight::IsFullWeight(BlendWeightB))
+			else if (bResetChildOnActivation && bIsChildBRelevant && !InstanceData->bWasChildBRelevant)
 			{
-				// We no longer need this child, release it
-				InstanceData->ChildA.Reset();
-				InstanceData->bWasChildARelevant = false;
+				InstanceData->bIsChildBNewlyRelevant = true;
 			}
 		}
+		else
+		{
+			// We no longer need this child, release it
+			InstanceData->ChildB.Reset();
+		}
+
+		InstanceData->bWasChildARelevant = bIsChildARelevant;
+		InstanceData->bWasChildBRelevant = bIsChildBRelevant;
 	}
 
 	void FBlendTwoWayTrait::QueueChildrenForTraversal(FUpdateTraversalContext& Context, const TTraitBinding<IUpdateTraversal>& Binding, const FTraitUpdateState& TraitState, FUpdateTraversalQueue& TraversalQueue) const
@@ -102,12 +116,12 @@ namespace UE::AnimNext
 		if (InstanceData->ChildA.IsValid())
 		{
 			const float BlendWeightA = 1.0f - BlendWeightB;
-			TraversalQueue.Push(InstanceData->ChildA, TraitState.WithWeight(BlendWeightA).AsNewlyRelevant(!InstanceData->bWasChildARelevant));
+			TraversalQueue.Push(InstanceData->ChildA, TraitState.WithWeight(BlendWeightA).AsNewlyRelevant(InstanceData->bIsChildANewlyRelevant));
 		}
 
 		if (InstanceData->ChildB.IsValid())
 		{
-			TraversalQueue.Push(InstanceData->ChildB, TraitState.WithWeight(BlendWeightB).AsNewlyRelevant(!InstanceData->bWasChildBRelevant));
+			TraversalQueue.Push(InstanceData->ChildB, TraitState.WithWeight(BlendWeightB).AsNewlyRelevant(InstanceData->bIsChildBNewlyRelevant));
 		}
 	}
 
diff --git a/Engine/Plugins/Experimental/AnimNextAnimGraph/Source/AnimNextAnimGraph/Private/Traits/BlendTwoWay.h b/Engine/Plugins/Experimental/AnimNextAnimGraph/Source/AnimNextAnimGraph/Private/Traits/BlendTwoWay.h
--- a/Engine/Plugins/Experimental/AnimNextAnimGraph/Source/AnimNextAnimGraph/Private/Traits/BlendTwoWay.h
+++ b/Engine/Plugins/Experimental/AnimNextAnimGraph/Source/AnimNextAnimGraph/Private/Traits/BlendTwoWay.h
@@ -29,9 +29,19 @@ struct FAnimNextBlendTwoWayTraitSharedData : public FAnimNextTraitSharedData
 	UPROPERTY(EditAnywhere, Category = "Default")
 	float BlendWeight = 0.0f;
 
+	/** This reinitializes child when re-activated */
+	UPROPERTY(EditAnywhere, Category = "Default")
+	bool bResetChildOnActivation = true;
+
+	/** Always update children, regardless of whether or not that child has weight. */
+	UPROPERTY(EditAnywhere, Category = "Default")
+	bool bAlwaysUpdateChildren = false;
+
 	// Latent pin support boilerplate
 	#define TRAIT_LATENT_PROPERTIES_ENUMERATOR(GeneratorMacro) \
 		GeneratorMacro(BlendWeight) \
+		GeneratorMacro(bResetChildOnActivation) \
+		GeneratorMacro(bAlwaysUpdateChildren) \
 
 	GENERATE_TRAIT_LATENT_PROPERTIES(FAnimNextBlendTwoWayTraitSharedData, TRAIT_LATENT_PROPERTIES_ENUMERATOR)
 	#undef TRAIT_LATENT_PROPERTIES_ENUMERATOR
@@ -55,8 +65,11 @@ namespace UE::AnimNext
 			FTraitPtr ChildA;
 			FTraitPtr ChildB;
 
-			bool bWasChildARelevant = false;
-			bool bWasChildBRelevant = false;
+			bool bIsChildANewlyRelevant : 1 = false;
+			bool bIsChildBNewlyRelevant : 1 = false;
+			
+			bool bWasChildARelevant : 1 = false;
+			bool bWasChildBRelevant : 1 = false;
 		};
 
 		// IEvaluate impl
