AnimNextAnimGraph - adding bAlwaysUpdateTrueChild and bAlwaysUpdateFalseChild to BlendByBool

#rnx

diff --git a/Engine/Plugins/Experimental/AnimNextAnimGraph/Source/AnimNextAnimGraph/Private/Traits/BlendByBool.cpp b/Engine/Plugins/Experimental/AnimNextAnimGraph/Source/AnimNextAnimGraph/Private/Traits/BlendByBool.cpp
--- a/Engine/Plugins/Experimental/AnimNextAnimGraph/Source/AnimNextAnimGraph/Private/Traits/BlendByBool.cpp
+++ b/Engine/Plugins/Experimental/AnimNextAnimGraph/Source/AnimNextAnimGraph/Private/Traits/BlendByBool.cpp
@@ -55,6 +55,7 @@ namespace UE::AnimNext
 
 	void FBlendByBoolTrait::QueueChildrenForTraversal(FUpdateTraversalContext& Context, const TTraitBinding<IUpdateTraversal>& Binding, const FTraitUpdateState& TraitState, FUpdateTraversalQueue& TraversalQueue) const
 	{
+		const FSharedData* SharedData = Binding.GetSharedData<FSharedData>();
 		const FInstanceData* InstanceData = Binding.GetInstanceData<FInstanceData>();
 
 		// The destination child index has been updated in PreUpdate, we can use the cached version
@@ -64,7 +65,7 @@ namespace UE::AnimNext
 		Binding.GetStackInterface(DiscreteBlendTrait);
 
 		const float BlendWeightTrue = DiscreteBlendTrait.GetBlendWeight(Context, TRUE_CHILD_INDEX);
-		if (InstanceData->TrueChild.IsValid() && FAnimWeight::IsRelevant(BlendWeightTrue))
+		if (InstanceData->TrueChild.IsValid() && (SharedData->bAlwaysUpdateTrueChild || FAnimWeight::IsRelevant(BlendWeightTrue)))
 		{
 			FTraitUpdateState TraitStateTrue = TraitState
 				.WithWeight(BlendWeightTrue)
@@ -148,31 +149,30 @@ namespace UE::AnimNext
 		FInstanceData* InstanceData = Binding.GetInstanceData<FInstanceData>();
 
 		// Allocate our new child instance
-		if (ChildIndex == TRUE_CHILD_INDEX)
+		if (!InstanceData->TrueChild.IsValid() && (ChildIndex == TRUE_CHILD_INDEX || SharedData->bAlwaysUpdateTrueChild))
 		{
-			if (!InstanceData->TrueChild.IsValid())
-			{
-				InstanceData->TrueChild = Context.AllocateNodeInstance(Binding, SharedData->TrueChild);
-			}
+			InstanceData->TrueChild = Context.AllocateNodeInstance(Binding, SharedData->TrueChild);
 		}
-		else if (ChildIndex == FALSE_CHILD_INDEX)
+		
+		if (!InstanceData->FalseChild.IsValid() && (ChildIndex == FALSE_CHILD_INDEX))
 		{
-			if (!InstanceData->FalseChild.IsValid())
-			{
-				InstanceData->FalseChild = Context.AllocateNodeInstance(Binding, SharedData->FalseChild);
-			}
+			InstanceData->FalseChild = Context.AllocateNodeInstance(Binding, SharedData->FalseChild);
 		}
 	}
 
 	void FBlendByBoolTrait::OnBlendTerminated(FExecutionContext& Context, const TTraitBinding<IDiscreteBlend>& Binding, int32 ChildIndex) const
 	{
+		const FSharedData* SharedData = Binding.GetSharedData<FSharedData>();
 		FInstanceData* InstanceData = Binding.GetInstanceData<FInstanceData>();
 
 		// Deallocate our child instance
 		if (ChildIndex == TRUE_CHILD_INDEX)
 		{
-			InstanceData->TrueChild.Reset();
-			InstanceData->bWasTrueChildRelevant = false;
+			if (!SharedData->bAlwaysUpdateTrueChild)
+			{
+				InstanceData->TrueChild.Reset();
+				InstanceData->bWasTrueChildRelevant = false;
+			}
 		}
 		else if (ChildIndex == FALSE_CHILD_INDEX)
 		{
diff --git a/Engine/Plugins/Experimental/AnimNextAnimGraph/Source/AnimNextAnimGraph/Private/Traits/BlendByBool.h b/Engine/Plugins/Experimental/AnimNextAnimGraph/Source/AnimNextAnimGraph/Private/Traits/BlendByBool.h
--- a/Engine/Plugins/Experimental/AnimNextAnimGraph/Source/AnimNextAnimGraph/Private/Traits/BlendByBool.h
+++ b/Engine/Plugins/Experimental/AnimNextAnimGraph/Source/AnimNextAnimGraph/Private/Traits/BlendByBool.h
@@ -28,6 +28,10 @@ struct FAnimNextBlendByBoolTraitSharedData : public FAnimNextTraitSharedData
 	UPROPERTY(EditAnywhere, Category = "Default")
 	bool bCondition = false;
 
+	/** Always update TrueChild, regardless of whether or not that child has weight. */
+	UPROPERTY(EditAnywhere, Category = "Default", meta = (Inline))
+	bool bAlwaysUpdateTrueChild = false;
+
 	// Latent pin support boilerplate
 	#define TRAIT_LATENT_PROPERTIES_ENUMERATOR(GeneratorMacro) \
 		GeneratorMacro(bCondition) \
