Use temporary attribute container to perform blending (similar to curve blending) in AnimNextRuntime

diff --git a/Engine/Plugins/Experimental/AnimNextAnimGraph/Source/AnimNextAnimGraph/Private/EvaluationVM/Tasks/BlendKeyframes.cpp b/Engine/Plugins/Experimental/AnimNextAnimGraph/Source/AnimNextAnimGraph/Private/EvaluationVM/Tasks/BlendKeyframes.cpp
--- a/Engine/Plugins/Experimental/AnimNextAnimGraph/Source/AnimNextAnimGraph/Private/EvaluationVM/Tasks/BlendKeyframes.cpp
+++ b/Engine/Plugins/Experimental/AnimNextAnimGraph/Source/AnimNextAnimGraph/Private/EvaluationVM/Tasks/BlendKeyframes.cpp
@@ -82,8 +82,10 @@ void FAnimNextBlendTwoKeyframesTask::Execute(UE::AnimNext::FEvaluationVM& VM) co
 		}
 
 		if (EnumHasAnyFlags(VM.GetFlags(), EEvaluationFlags::Attributes))
-		{
-			UE::Anim::Attributes::BlendAttributes({ KeyframeA->Attributes, KeyframeB->Attributes }, { WeightOfPoseA, WeightOfPoseB }, { 0, 1 }, KeyframeB->Attributes);
+		{			
+			UE::Anim::FStackAttributeContainer OutputAttributes;
+			UE::Anim::Attributes::BlendAttributes({ KeyframeA->Attributes, KeyframeB->Attributes }, { WeightOfPoseA, WeightOfPoseB }, { 0, 1 }, OutputAttributes);			
+			KeyframeB->Attributes.MoveFrom(OutputAttributes);
 		}
 
 		VM.PushValue(KEYFRAME_STACK_NAME, MoveTemp(KeyframeB));
diff --git a/Engine/Plugins/Experimental/AnimNextAnimGraph/Source/AnimNextAnimGraph/Private/EvaluationVM/Tasks/BlendKeyframesPerBone.cpp b/Engine/Plugins/Experimental/AnimNextAnimGraph/Source/AnimNextAnimGraph/Private/EvaluationVM/Tasks/BlendKeyframesPerBone.cpp
--- a/Engine/Plugins/Experimental/AnimNextAnimGraph/Source/AnimNextAnimGraph/Private/EvaluationVM/Tasks/BlendKeyframesPerBone.cpp
+++ b/Engine/Plugins/Experimental/AnimNextAnimGraph/Source/AnimNextAnimGraph/Private/EvaluationVM/Tasks/BlendKeyframesPerBone.cpp
@@ -90,8 +90,9 @@ void FAnimNextBlendOverwriteKeyframePerBoneWithScaleTask::Execute(UE::AnimNext::
 
 	if (EnumHasAnyFlags(VM.GetFlags(), EEvaluationFlags::Attributes))
 	{
-		UE::Anim::Attributes::OverrideAttributes(Keyframe->Attributes, Keyframe->Attributes, BlendWeight);
-		//UE::Anim::Attributes::BlendAttributesPerBone({ Keyframe->Attributes }, { LODBoneIndexToWeightIndexMap }, { *BlendData }, { 0 }, { Keyframe->Attributes });
+		UE::Anim::FStackAttributeContainer OutputAttributes;
+		UE::Anim::Attributes::BlendAttributesPerBone({ Keyframe->Attributes }, { LODBoneIndexToWeightIndexMap }, { *BlendData }, { 0 }, { OutputAttributes });
+		Keyframe->Attributes.MoveFrom(OutputAttributes);
 	}
 
 	VM.PushValue(KEYFRAME_STACK_NAME, MoveTemp(Keyframe));
@@ -183,16 +184,16 @@ void FAnimNextBlendAddKeyframePerBoneWithScaleTask::Execute(UE::AnimNext::FEvalu
 
 	if (EnumHasAnyFlags(VM.GetFlags(), EEvaluationFlags::Attributes))
 	{
-		UE::Anim::Attributes::AccumulateAttributes(KeyframeA->Attributes, KeyframeB->Attributes, BlendWeightA, AAT_None);
-		/*
 		// TODO: Might need to be revisited, might not work as intended
+		UE::Anim::FStackAttributeContainer OutputAttributes;
 		UE::Anim::Attributes::BlendAttributesPerBone(
 			{ KeyframeA->Attributes, KeyframeB->Attributes },
 			{ LODBoneIndexToWeightIndexMap },
 			{ *BlendDataA, *BlendDataB },
 			{ 0, 1 },
-			{ KeyframeA->Attributes });
-		*/
+			{ OutputAttributes });
+
+		KeyframeA->Attributes.MoveFrom(OutputAttributes);
 	}
 
 	VM.PushValue(KEYFRAME_STACK_NAME, MoveTemp(KeyframeB));
@@ -378,7 +379,9 @@ void FAnimNextBlendKeyframePerBoneWithScaleTask::Execute(UE::AnimNext::FEvaluati
 		const TArray<const float> SourceWeights{ 1.0f - ScaleFactor, ScaleFactor };
 		static const TArray<const int32> WeightView{ 0, 1 };
 
-		UE::Anim::Attributes::BlendAttributes(SourceAttributes, SourceWeights, WeightView, BaseAttributes);
+		FStackAttributeContainer OutputAttributes;
+		UE::Anim::Attributes::BlendAttributes(SourceAttributes, SourceWeights, WeightView, OutputAttributes);
+		BaseAttributes.MoveFrom(OutputAttributes);
 	}
 
 	VM.PushValue(KEYFRAME_STACK_NAME, MoveTemp(KeyframeA));
