Anim Next: Allow injection requests to specify whether to fallback to the module's Default Injection Site or not

Sequencer.AnimNext.DefaultInjectionSite now allows projects to specify a default injection site for anim next targets, falling back to the module default if not found.

diff --git a/Engine/Plugins/Experimental/AnimNextAnimGraph/Source/AnimNextAnimGraph/Private/Injection/AnimNextModuleInjectionComponent.cpp b/Engine/Plugins/Experimental/AnimNextAnimGraph/Source/AnimNextAnimGraph/Private/Injection/AnimNextModuleInjectionComponent.cpp
--- a/Engine/Plugins/Experimental/AnimNextAnimGraph/Source/AnimNextAnimGraph/Private/Injection/AnimNextModuleInjectionComponent.cpp
+++ b/Engine/Plugins/Experimental/AnimNextAnimGraph/Source/AnimNextAnimGraph/Private/Injection/AnimNextModuleInjectionComponent.cpp
@@ -49,25 +49,25 @@ void FAnimNextModuleInjectionComponent::OnInjectionEvent(UE::AnimNext::FInjectio
 	const FInjectionRequestArgs& RequestArgs = InEvent.Request->GetArgs();
 
 	FName FoundName = NAME_None;
-	TStructView<FAnimNextAnimGraph> InjectableGraph = InjectionInfo.FindInjectableGraphInstance(RequestArgs.Name, FoundName);
+	TStructView<FAnimNextAnimGraph> InjectableGraph = InjectionInfo.FindInjectableGraphInstance(RequestArgs.Site, FoundName);
 	if(!InjectableGraph.IsValid())
 	{
-		UE_LOGFMT(LogAnimation, Warning, "Could not find injection site {SiteName} for injection request", RequestArgs.Name);
+		UE_LOGFMT(LogAnimation, Warning, "Could not find injection site {SiteName} for injection request", RequestArgs.Site.DesiredSiteName);
 		InEvent.MarkConsumed();
 		return;
 	}
 
 	// Correct the name found above as we may have targeted NAME_None (any)
-	if(RequestArgs.Name != FoundName)
+	if(RequestArgs.Site.DesiredSiteName != FoundName)
 	{
-		InEvent.Request->GetMutableArgs().Name = FoundName;
+		InEvent.Request->GetMutableArgs().Site = FInjectionSite(FoundName);
 	}
 
 	// Mark as handled so any additional trait events dont get processed at the module level
 	InEvent.MarkHandled();
 
 	// Store request as it will need to be re-applied each frame to ensure that bindings do not override it
-	FInjectionRecord& InjectionRecord = CurrentRequests.FindOrAdd(InEvent.Request->GetArgs().Name);
+	FInjectionRecord& InjectionRecord = CurrentRequests.FindOrAdd(InEvent.Request->GetArgs().Site.DesiredSiteName);
 
 	FAnimNextAnimGraph& Graph = InjectableGraph.Get<FAnimNextAnimGraph>();
 	switch(RequestArgs.Type)
@@ -110,10 +110,10 @@ void FAnimNextModuleInjectionComponent::OnUninjectionEvent(UE::AnimNext::FInject
 
 	// Injection modifiers can have private access to take over low-level tasks
 	FName FoundName = NAME_None;
-	TStructView<FAnimNextAnimGraph> InjectableGraph = InjectionInfo.FindInjectableGraphInstance(RequestArgs.Name, FoundName);
+	TStructView<FAnimNextAnimGraph> InjectableGraph = InjectionInfo.FindInjectableGraphInstance(RequestArgs.Site, FoundName);
 	if(!InjectableGraph.IsValid())
 	{
-		UE_LOGFMT(LogAnimation, Warning, "Could not find injection site {SiteName} for un-injection request", RequestArgs.Name);
+		UE_LOGFMT(LogAnimation, Warning, "Could not find injection site {SiteName} for un-injection request", RequestArgs.Site.DesiredSiteName);
 		InEvent.MarkConsumed();
 		return;
 	}
@@ -152,7 +152,7 @@ void FAnimNextModuleInjectionComponent::OnUninjectionEvent(UE::AnimNext::FInject
 
 	if(!InjectionRecord.IsValid())
 	{
-		CurrentRequests.Remove(InEvent.Request->GetArgs().Name);
+		CurrentRequests.Remove(InEvent.Request->GetArgs().Site.DesiredSiteName);
 	}
 }
 
@@ -165,7 +165,7 @@ void FAnimNextModuleInjectionComponent::OnReapplyInjection(const UE::AnimNext::F
 	for(const TPair<FName, FInjectionRecord>& RequestPair : Component.CurrentRequests)
 	{
 		FName FoundName = NAME_None;
-		TStructView<FAnimNextAnimGraph> InjectableGraph = Component.InjectionInfo.FindInjectableGraphInstance(RequestPair.Key, FoundName);
+		TStructView<FAnimNextAnimGraph> InjectableGraph = Component.InjectionInfo.FindInjectableGraphInstance(UE::AnimNext::FInjectionSite(RequestPair.Key), FoundName);
 		if(!InjectableGraph.IsValid())
 		{
 			continue;
diff --git a/Engine/Plugins/Experimental/AnimNextAnimGraph/Source/AnimNextAnimGraph/Private/Injection/InjectionCallbackProxy.cpp b/Engine/Plugins/Experimental/AnimNextAnimGraph/Source/AnimNextAnimGraph/Private/Injection/InjectionCallbackProxy.cpp
--- a/Engine/Plugins/Experimental/AnimNextAnimGraph/Source/AnimNextAnimGraph/Private/Injection/InjectionCallbackProxy.cpp
+++ b/Engine/Plugins/Experimental/AnimNextAnimGraph/Source/AnimNextAnimGraph/Private/Injection/InjectionCallbackProxy.cpp
@@ -38,7 +38,7 @@ bool UInjectionCallbackProxy::Inject(
 	}
 
 	UE::AnimNext::FInjectionRequestArgs RequestArgs;
-	RequestArgs.Name = SiteName;
+	RequestArgs.Site = FAnimNextInjectionSite(SiteName);
 	RequestArgs.Object = Object;
 	RequestArgs.BlendInSettings = BlendInSettings;
 	RequestArgs.BlendOutSettings = BlendOutSettings;
diff --git a/Engine/Plugins/Experimental/AnimNextAnimGraph/Source/AnimNextAnimGraph/Private/Injection/InjectionInfo.cpp b/Engine/Plugins/Experimental/AnimNextAnimGraph/Source/AnimNextAnimGraph/Private/Injection/InjectionInfo.cpp
--- a/Engine/Plugins/Experimental/AnimNextAnimGraph/Source/AnimNextAnimGraph/Private/Injection/InjectionInfo.cpp
+++ b/Engine/Plugins/Experimental/AnimNextAnimGraph/Source/AnimNextAnimGraph/Private/Injection/InjectionInfo.cpp
@@ -1,6 +1,7 @@
 ﻿// Copyright Epic Games, Inc. All Rights Reserved.
 
 #include "InjectionInfo.h"
+#include "Injection/InjectionSite.h"
 
 #include "DataInterface/AnimNextDataInterfaceInstance.h"
 #include "Graph/AnimNextAnimGraph.h"
@@ -77,19 +78,32 @@ bool FInjectionInfo::GetDefaultInjectableGraph(FName& OutName, TStructView<FAnim
 	return OutName != NAME_None && OutGraph.IsValid();
 }
 
-TStructView<FAnimNextAnimGraph> FInjectionInfo::FindInjectableGraphInstance(FName InName, FName& OutRealName) const
+TStructView<FAnimNextAnimGraph> FInjectionInfo::FindInjectableGraphInstance(const FInjectionSite& InSite, FName& OutRealName) const
 {
-	const int32 InjectableGraphNum = InjectableGraphs.Num();
-	for(int32 InjectableGraphIndex = 0; InjectableGraphIndex < InjectableGraphNum; ++InjectableGraphIndex)
+	FName SiteName = InSite.DesiredSiteName;
+	if(SiteName != NAME_None)
 	{
-		const FInjectableGraphInfo& InjectableGraphInfo = InjectableGraphs[InjectableGraphIndex];
-		if((InName == NAME_None && InjectableGraphIndex == DefaultInjectableGraphIndex) || InName == InjectableGraphInfo.Name)
+		// Linear search all the injectable graphs for the name
+		const int32 InjectableGraphNum = InjectableGraphs.Num();
+		for(int32 InjectableGraphIndex = 0; InjectableGraphIndex < InjectableGraphNum; ++InjectableGraphIndex)
 		{
-			OutRealName = InjectableGraphInfo.Name;
-			return InjectableGraphInfo.StructView;
+			const FInjectableGraphInfo& InjectableGraphInfo = InjectableGraphs[InjectableGraphIndex];
+			if(SiteName == InjectableGraphInfo.Name)
+			{
+				OutRealName = SiteName;
+				return InjectableGraphInfo.StructView;
+			}
 		}
 	}
 
+	const bool bUseModuleDefault = SiteName == NAME_None || InSite.bUseModuleFallback;
+	if(bUseModuleDefault && InjectableGraphs.IsValidIndex(DefaultInjectableGraphIndex))
+	{
+		const FInjectableGraphInfo& InjectableGraphInfo = InjectableGraphs[DefaultInjectableGraphIndex];
+		OutRealName = InjectableGraphInfo.Name;
+		return InjectableGraphInfo.StructView;
+	}
+
 	return TStructView<FAnimNextAnimGraph>();
 }
 
diff --git a/Engine/Plugins/Experimental/AnimNextAnimGraph/Source/AnimNextAnimGraph/Private/Injection/InjectionInfo.h b/Engine/Plugins/Experimental/AnimNextAnimGraph/Source/AnimNextAnimGraph/Private/Injection/InjectionInfo.h
--- a/Engine/Plugins/Experimental/AnimNextAnimGraph/Source/AnimNextAnimGraph/Private/Injection/InjectionInfo.h
+++ b/Engine/Plugins/Experimental/AnimNextAnimGraph/Source/AnimNextAnimGraph/Private/Injection/InjectionInfo.h
@@ -8,9 +8,11 @@
 struct FAnimNextDataInterfaceInstance;
 struct FAnimNextAnimGraph;
 struct FAnimNextModuleInjectionComponent;
+struct FAnimNextInjectionSite;
 
 namespace UE::AnimNext
 {
+using FInjectionSite = FAnimNextInjectionSite;
 
 // Info used to track injection sites for a data interface instance (graph, module etc)
 struct FInjectionInfo
@@ -23,9 +25,9 @@ struct FInjectionInfo
 	bool GetDefaultInjectableGraph(FName& OutName, TStructView<FAnimNextAnimGraph>& OutGraph) const;
 
 	// Find an injectable graph instance by name.
-	// @param   InName            The name of the injection site. If NAME_None is passed then the default site will be returned
+	// @param   InSite            The injection site. If bUseModuleFallback is true, the default site will be returned if the site does not exist
 	// @param   OutRealName       The name of the actual injection site we found (in the case NAME_None was passed)
-	TStructView<FAnimNextAnimGraph> FindInjectableGraphInstance(FName InName, FName& OutRealName) const;
+	TStructView<FAnimNextAnimGraph> FindInjectableGraphInstance(const FInjectionSite& InSite, FName& OutRealName) const;
 
 	// Iterate each injectable graph instance
 	template<typename PredicateType>
diff --git a/Engine/Plugins/Experimental/AnimNextAnimGraph/Source/AnimNextAnimGraph/Private/Injection/InjectionRequest.cpp b/Engine/Plugins/Experimental/AnimNextAnimGraph/Source/AnimNextAnimGraph/Private/Injection/InjectionRequest.cpp
--- a/Engine/Plugins/Experimental/AnimNextAnimGraph/Source/AnimNextAnimGraph/Private/Injection/InjectionRequest.cpp
+++ b/Engine/Plugins/Experimental/AnimNextAnimGraph/Source/AnimNextAnimGraph/Private/Injection/InjectionRequest.cpp
@@ -81,7 +81,7 @@ namespace UE::AnimNext
 		switch(InRequestArgs.Type)
 		{
 		case EAnimNextInjectionType::InjectObject:
-			if(InRequestArgs.Name == NAME_None)
+			if(InRequestArgs.Site.DesiredSiteName == NAME_None)
 			{
 				UE_LOGFMT(LogAnimation, Warning, "FInjectionRequest::ValidateArgs: Missing injection site name");
 				return false;
diff --git a/Engine/Plugins/Experimental/AnimNextAnimGraph/Source/AnimNextAnimGraph/Private/Injection/InjectionSiteTrait.cpp b/Engine/Plugins/Experimental/AnimNextAnimGraph/Source/AnimNextAnimGraph/Private/Injection/InjectionSiteTrait.cpp
--- a/Engine/Plugins/Experimental/AnimNextAnimGraph/Source/AnimNextAnimGraph/Private/Injection/InjectionSiteTrait.cpp
+++ b/Engine/Plugins/Experimental/AnimNextAnimGraph/Source/AnimNextAnimGraph/Private/Injection/InjectionSiteTrait.cpp
@@ -94,7 +94,7 @@ namespace UE::AnimNext
 		if (Graph.InjectionData.InjectionSerialNumber == Event.SerialNumber)
 		{
 			const FAnimNextInjectionRequestArgs& RequestArgs = Event.Request->GetArgs();
-			ensureMsgf(!InstanceData->PendingRequest.IsValid(), TEXT("Injection site %s already contained a pending request, it will be overwritten"), *RequestArgs.Name.ToString());
+			ensureMsgf(!InstanceData->PendingRequest.IsValid(), TEXT("Injection site %s already contained a pending request, it will be overwritten"), *RequestArgs.Site.DesiredSiteName.ToString());
 
 			// Overwrite any request we might have, we'll pick it up on the next update
 			InstanceData->PendingRequest.Reset();
diff --git a/Engine/Plugins/Experimental/AnimNextAnimGraph/Source/AnimNextAnimGraph/Private/Injection/InjectionUtils.cpp b/Engine/Plugins/Experimental/AnimNextAnimGraph/Source/AnimNextAnimGraph/Private/Injection/InjectionUtils.cpp
--- a/Engine/Plugins/Experimental/AnimNextAnimGraph/Source/AnimNextAnimGraph/Private/Injection/InjectionUtils.cpp
+++ b/Engine/Plugins/Experimental/AnimNextAnimGraph/Source/AnimNextAnimGraph/Private/Injection/InjectionUtils.cpp
@@ -39,7 +39,7 @@ namespace UE::AnimNext
 
 	FInjectionRequestPtr FInjectionUtils::PlayAnim(
 		UAnimNextComponent* InComponent,
-		FName InSiteName,
+		const FInjectionSite& InSite,
 		UAnimSequence* InAnimSequence,
 		float InPlayRate,
 		float InStartPosition,
@@ -50,7 +50,7 @@ namespace UE::AnimNext
 		return PlayAnim(
 			InComponent,
 			InComponent->GetModuleHandle(),
-			InSiteName,
+			InSite,
 			InAnimSequence,
 			InPlayRate,
 			InStartPosition,
@@ -62,7 +62,7 @@ namespace UE::AnimNext
 	FInjectionRequestPtr FInjectionUtils::PlayAnim(
 		UObject* InHost,
 		FModuleHandle InModuleHandle,
-		FName InSiteName,
+		const FInjectionSite& InSite,
 		UAnimSequence* InAnimSequence,
 		float InPlayRate, float InStartPosition,
 		const FInjectionBlendSettings& InBlendInSettings,
@@ -80,7 +80,7 @@ namespace UE::AnimNext
 		PlayNativeInterface.Loop = InAnimSequence->bLoop;
 
 		FInjectionRequestArgs RequestArgs;
-		RequestArgs.Name = InSiteName;
+		RequestArgs.Site = InSite;
 		RequestArgs.Object = InAnimSequence;
 		RequestArgs.Type = EAnimNextInjectionType::InjectObject;
 		RequestArgs.BlendInSettings = InBlendInSettings;
@@ -92,7 +92,7 @@ namespace UE::AnimNext
 
 	FInjectionRequestPtr FInjectionUtils::InjectAsset(
 		UAnimNextComponent* InComponent,
-		FName InSiteName,
+		const FInjectionSite& InSite,
 		UObject* InAsset,
 		FAnimNextDataInterfacePayload&& InPayload,
 		UAnimNextComponent* InBindingComponent,
@@ -103,7 +103,7 @@ namespace UE::AnimNext
 		return InjectAsset(
 			InComponent,
 			InComponent->GetModuleHandle(),
-			InSiteName,
+			InSite,
 			InAsset,
 			MoveTemp(InPayload),
 			InBindingComponent ? InBindingComponent->GetModuleHandle() : FModuleHandle(),
@@ -115,7 +115,7 @@ namespace UE::AnimNext
 	FInjectionRequestPtr FInjectionUtils::InjectAsset(
 		UObject* InHost,
 		FModuleHandle InModuleHandle,
-		FName InSiteName,
+		const FInjectionSite& InSite,
 		UObject* InAsset,
 		FAnimNextDataInterfacePayload&& InPayload,
 		FModuleHandle InBindingModuleHandle,
@@ -126,7 +126,7 @@ namespace UE::AnimNext
 		check(IsInGameThread());
 
 		FInjectionRequestArgs RequestArgs;
-		RequestArgs.Name = InSiteName;
+		RequestArgs.Site = InSite;
 		RequestArgs.Object = InAsset;
 		RequestArgs.BindingModuleHandle = InBindingModuleHandle;
 		RequestArgs.Type = EAnimNextInjectionType::InjectObject;
@@ -140,25 +140,25 @@ namespace UE::AnimNext
 	FInjectionRequestPtr FInjectionUtils::InjectEvaluationModifier(
 		UAnimNextComponent* InComponent,
 		const TSharedRef<IEvaluationModifier>& InEvaluationModifier,
-		FName InSiteName)
+		const FInjectionSite& InSite)
 	{
 		return InjectEvaluationModifier(
 			InComponent,
 			InComponent->GetModuleHandle(),
 			InEvaluationModifier,
-			InSiteName);
+			InSite);
 	}
 
 	FInjectionRequestPtr FInjectionUtils::InjectEvaluationModifier(
 		UObject* InHost,
 		FModuleHandle InModuleHandle,
 		const TSharedRef<IEvaluationModifier>& InEvaluationModifier,
-		FName InSiteName)
+		const FInjectionSite& InSite)
 	{
 		check(IsInGameThread());
 
 		FInjectionRequestArgs RequestArgs;
-		RequestArgs.Name = InSiteName;
+		RequestArgs.Site = InSite;
 		RequestArgs.Type = EAnimNextInjectionType::EvaluationModifier;
 		RequestArgs.EvaluationModifier = InEvaluationModifier;
 
diff --git a/Engine/Plugins/Experimental/AnimNextAnimGraph/Source/AnimNextAnimGraph/Private/PlayAnim/PlayAnimCallbackProxy.cpp b/Engine/Plugins/Experimental/AnimNextAnimGraph/Source/AnimNextAnimGraph/Private/PlayAnim/PlayAnimCallbackProxy.cpp
--- a/Engine/Plugins/Experimental/AnimNextAnimGraph/Source/AnimNextAnimGraph/Private/PlayAnim/PlayAnimCallbackProxy.cpp
+++ b/Engine/Plugins/Experimental/AnimNextAnimGraph/Source/AnimNextAnimGraph/Private/PlayAnim/PlayAnimCallbackProxy.cpp
@@ -50,7 +50,7 @@ bool UPlayAnimCallbackProxy::Play(
 
 		PlayingRequest = UE::AnimNext::FInjectionUtils::PlayAnim(
 			AnimNextComponent,
-			SiteName,
+			UE::AnimNext::FInjectionSite(SiteName),
 			AnimSequence,
 			PlayRate,
 			StartPosition,
diff --git a/Engine/Plugins/Experimental/AnimNextAnimGraph/Source/AnimNextAnimGraph/Private/PlayAnim/PlayAnimSlotTrait.cpp b/Engine/Plugins/Experimental/AnimNextAnimGraph/Source/AnimNextAnimGraph/Private/PlayAnim/PlayAnimSlotTrait.cpp
--- a/Engine/Plugins/Experimental/AnimNextAnimGraph/Source/AnimNextAnimGraph/Private/PlayAnim/PlayAnimSlotTrait.cpp
+++ b/Engine/Plugins/Experimental/AnimNextAnimGraph/Source/AnimNextAnimGraph/Private/PlayAnim/PlayAnimSlotTrait.cpp
@@ -86,7 +86,7 @@ namespace UE::AnimNext
 		const FName SlotName = SharedData->GetSlotName(Binding);
 
 		const FAnimNextInjectionRequestArgs& RequestArgs = Event.Request->GetArgs();
-		if (SlotName == RequestArgs.Name)
+		if (SlotName == RequestArgs.Site.DesiredSiteName)
 		{
 			FInstanceData* InstanceData = Binding.GetInstanceData<FInstanceData>();
 
@@ -108,7 +108,7 @@ namespace UE::AnimNext
 		const FName SlotName = SharedData->GetSlotName(Binding);
 
 		const FAnimNextInjectionRequestArgs& RequestArgs = Event.Request->GetArgs();
-		if (SlotName == RequestArgs.Name)
+		if (SlotName == RequestArgs.Site.DesiredSiteName)
 		{
 			FInstanceData* InstanceData = Binding.GetInstanceData<FInstanceData>();
 
diff --git a/Engine/Plugins/Experimental/AnimNextAnimGraph/Source/AnimNextAnimGraph/Public/Injection/InjectionRequest.h b/Engine/Plugins/Experimental/AnimNextAnimGraph/Source/AnimNextAnimGraph/Public/Injection/InjectionRequest.h
--- a/Engine/Plugins/Experimental/AnimNextAnimGraph/Source/AnimNextAnimGraph/Public/Injection/InjectionRequest.h
+++ b/Engine/Plugins/Experimental/AnimNextAnimGraph/Source/AnimNextAnimGraph/Public/Injection/InjectionRequest.h
@@ -9,6 +9,7 @@
 #include "TraitCore/TraitEvent.h"
 #include "TraitInterfaces/ITimeline.h"
 #include "Injection/InjectionStatus.h"
+#include "Injection/InjectionSite.h"
 #include "Module/ModuleHandle.h"
 
 #include "InjectionRequest.generated.h"
@@ -94,9 +95,9 @@ private:
 	friend FAnimNextModuleInjectionComponent;
 	friend UInjectionCallbackProxy;
 
-	// The injection site name to target with this request.
+	// The injection site to target with this request.
 	UPROPERTY()
-	FName Name;
+	FAnimNextInjectionSite Site;
 
 	// The blend settings to use when blending in
 	UPROPERTY()
diff --git a/Engine/Plugins/Experimental/AnimNextAnimGraph/Source/AnimNextAnimGraph/Public/Injection/InjectionSite.h b/Engine/Plugins/Experimental/AnimNextAnimGraph/Source/AnimNextAnimGraph/Public/Injection/InjectionSite.h
new file mode 100644
--- /dev/null
+++ b/Engine/Plugins/Experimental/AnimNextAnimGraph/Source/AnimNextAnimGraph/Public/Injection/InjectionSite.h
@@ -0,0 +1,44 @@
+﻿// Copyright Epic Games, Inc. All Rights Reserved.
+
+#pragma once
+
+#include "CoreTypes.h"
+#include "UObject/NameTypes.h"
+
+#include "InjectionSite.generated.h"
+
+// Specifies the desired injection site for an injection request.
+// Optionally allows for fallback to the Default Injection Site specified by the module
+//     if DesiredSiteName is None, or does not exist in the graph
+USTRUCT()
+struct FAnimNextInjectionSite
+{
+	GENERATED_BODY()
+
+	FAnimNextInjectionSite()
+		: DesiredSiteName(NAME_None)
+		, bUseModuleFallback(true)
+	{
+	}
+
+	FAnimNextInjectionSite(const FName& InSiteName)
+		: DesiredSiteName(InSiteName)
+		, bUseModuleFallback(DesiredSiteName == NAME_None)
+	{
+	}
+
+	// The name of the site to inject into
+	UPROPERTY()
+	FName DesiredSiteName;
+
+	// Flag specifying whether the request can (or should) fallback to the Default Injection Site
+	// specified in the module if DesiredSiteName was not found
+	UPROPERTY()
+	uint8 bUseModuleFallback : 1;
+};
+
+
+namespace UE::AnimNext
+{
+	using FInjectionSite = FAnimNextInjectionSite;
+}
\ No newline at end of file
diff --git a/Engine/Plugins/Experimental/AnimNextAnimGraph/Source/AnimNextAnimGraph/Public/Injection/InjectionUtils.h b/Engine/Plugins/Experimental/AnimNextAnimGraph/Source/AnimNextAnimGraph/Public/Injection/InjectionUtils.h
--- a/Engine/Plugins/Experimental/AnimNextAnimGraph/Source/AnimNextAnimGraph/Public/Injection/InjectionUtils.h
+++ b/Engine/Plugins/Experimental/AnimNextAnimGraph/Source/AnimNextAnimGraph/Public/Injection/InjectionUtils.h
@@ -4,6 +4,7 @@
 
 #include "CoreMinimal.h"
 #include "InjectionRequest.h"
+#include "Injection/InjectionSite.h"
 
 struct FAnimNextDataInterfaceInstance;
 class UAnimNextComponent;
@@ -22,7 +23,7 @@ namespace UE::AnimNext
 		// @param InComponent          The UAnimNextComponent that hosts the module being played on
 		// @param InHost               The object that hosts the module (e.g. UAnimNextComponent) being played on
 		// @param InModuleHandle       The module handle being injected into
-		// @param InSiteName           The injection site to 'play' the animation sequence on
+		// @param InSite               The injection site to 'play' the animation sequence on
 		// @param InAnimSequence       The animation sequence to play
 		// @param InPlayRate           The play rate multiplier to use
 		// @param InStartPosition      The start position to use
@@ -32,7 +33,7 @@ namespace UE::AnimNext
 		// @return the new injection request for tracking, or an invalid request if injection failed. This can be discarded for fire-and-forget style behavior.
 		ANIMNEXTANIMGRAPH_API static FInjectionRequestPtr PlayAnim(
 			UAnimNextComponent* InComponent,
-			FName InSiteName,
+			const FInjectionSite& InSite,
 			UAnimSequence* InAnimSequence,
 			float InPlayRate = 1.0f,
 			float InStartPosition = 0.0f,
@@ -44,7 +45,7 @@ namespace UE::AnimNext
 		ANIMNEXTANIMGRAPH_API static FInjectionRequestPtr PlayAnim(
 			UObject* InHost,
 			FModuleHandle InModuleHandle,
-			FName InSiteName,
+			const FInjectionSite& InSite,
 			UAnimSequence* InAnimSequence,
 			float InPlayRate = 1.0f,
 			float InStartPosition = 0.0f,
@@ -56,7 +57,7 @@ namespace UE::AnimNext
 		// @param InComponent          The UAnimNextComponent that hosts the module being played on
 		// @param InHost               The object that hosts the module (e.g. UAnimNextComponent) being played on
 		// @param InHandle             The module handle being injected into
-		// @param InSiteName           The injection site to 'play' the asset on
+		// @param InSite               The injection site to 'play' the asset on
 		// @param InAsset              The asset to 'play'
 		// @param InPayload            The payload for the asset to use
 		// @param InBindingComponent   A component to use to supply additional data interface bindings at the injection site
@@ -67,7 +68,7 @@ namespace UE::AnimNext
 		// @return the new injection request for tracking, or an invalid request if injection failed. This can be discarded for fire-and-forget style behavior.
 		ANIMNEXTANIMGRAPH_API static FInjectionRequestPtr InjectAsset(
 			UAnimNextComponent* InComponent,
-			FName InSiteName,
+			const FInjectionSite& InSite,
 			UObject* InAsset,
 			FAnimNextDataInterfacePayload&& InPayload = FAnimNextDataInterfacePayload(),
 			UAnimNextComponent* InBindingComponent = nullptr,
@@ -79,7 +80,7 @@ namespace UE::AnimNext
 		ANIMNEXTANIMGRAPH_API static FInjectionRequestPtr InjectAsset(
 			UObject* InHost,
 			FModuleHandle InModuleHandle,
-			FName InSiteName,
+			const FInjectionSite& InSite,
 			UObject* InAsset,
 			FAnimNextDataInterfacePayload&& InPayload = FAnimNextDataInterfacePayload(),
 			FModuleHandle InBindingModuleHandle = FModuleHandle(),
@@ -93,19 +94,19 @@ namespace UE::AnimNext
 		// @param InHost               The object that hosts the module (e.g. UAnimNextComponent) being injected into
 		// @param InModuleHandle       The module handle being injected into
 		// @param InEvaluationModifier The evaluation modifier to inject
-		// @param InSiteName           The injection site to use. If this is NAME_None, then the first site will be used.
+		// @param InSite               The injection site to use. If this is NAME_None, then the first site will be used.
 		// @return the new injection request for tracking, or an invalid request if injection failed. This can be discarded for fire-and-forget style behavior.
 		ANIMNEXTANIMGRAPH_API static FInjectionRequestPtr InjectEvaluationModifier(
 			UAnimNextComponent* InComponent,
 			const TSharedRef<IEvaluationModifier>& InEvaluationModifier,
-			FName InSiteName = NAME_None);
+			const FInjectionSite& InSite = FInjectionSite());
 
 		// @see InjectExternalGraph
 		ANIMNEXTANIMGRAPH_API static FInjectionRequestPtr InjectEvaluationModifier(
 			UObject* InHost,
 			FModuleHandle InModuleHandle,
 			const TSharedRef<IEvaluationModifier>& InEvaluationModifier,
-			FName InSiteName = NAME_None);
+			const FInjectionSite& InSite = FInjectionSite());
 
 		// Inject into a running AnimNext module (raw custom args)
 		// @param InComponent          The UAnimNextComponent that hosts the module being injected into
diff --git a/Engine/Plugins/Experimental/MovieSceneAnimMixer/Source/MovieSceneAnimMixer/Private/Systems/MovieSceneAnimNextTargetSystem.cpp b/Engine/Plugins/Experimental/MovieSceneAnimMixer/Source/MovieSceneAnimMixer/Private/Systems/MovieSceneAnimNextTargetSystem.cpp
--- a/Engine/Plugins/Experimental/MovieSceneAnimMixer/Source/MovieSceneAnimMixer/Private/Systems/MovieSceneAnimNextTargetSystem.cpp
+++ b/Engine/Plugins/Experimental/MovieSceneAnimMixer/Source/MovieSceneAnimMixer/Private/Systems/MovieSceneAnimNextTargetSystem.cpp
@@ -40,6 +40,14 @@
 
 namespace UE::MovieScene
 {
+	FName GSequencerDefaultAnimNextInjectionSite = NAME_None;
+	FAutoConsoleVariableRef CVarDefaultAnimNextInjectionSite(
+		TEXT("Sequencer.AnimNext.DefaultInjectionSite"),
+		GSequencerDefaultAnimNextInjectionSite,
+		TEXT("(Default: None) Specifies the default injection site name for Sequencer Anim Next Targets that is used when none is specified on the target itself."),
+		ECVF_Default
+	);
+
 	/** Pre-animated state for skeletal animations */
 	struct FPreAnimatedAnimNextState
 	{
@@ -158,17 +166,27 @@ namespace UE::MovieScene
 			if (AnimNextComponent)
 			{
 				// If not currently injected, inject
-				FMovieSceneAnimNextTargetData* TargetData = System->CurrentTargets.FindByPredicate([AnimNextComponent, InjectionTarget](const FMovieSceneAnimNextTargetData& Data) { return Data.AnimNextComponent == AnimNextComponent && Data.InjectionSiteName == InjectionTarget->InjectionSiteName;});
+				FMovieSceneAnimNextTargetData* TargetData = System->CurrentTargets.FindByPredicate([AnimNextComponent, InjectionTarget](const FMovieSceneAnimNextTargetData& Data) { return Data.AnimNextComponent == AnimNextComponent && Data.InjectionSite.DesiredSiteName == InjectionTarget->InjectionSiteName;});
 
 				if (!TargetData)
 				{
 					TargetData = &System->CurrentTargets.AddDefaulted_GetRef();
 					TargetData->AnimNextComponent = AnimNextComponent;
-					TargetData->InjectionSiteName = InjectionTarget->InjectionSiteName;
+
+					if (InjectionTarget->InjectionSiteName == NAME_None)
+					{
+						TargetData->InjectionSite.DesiredSiteName = GSequencerDefaultAnimNextInjectionSite;
+						TargetData->InjectionSite.bUseModuleFallback = true;
+					}
+					else
+					{
+						TargetData->InjectionSite.DesiredSiteName = InjectionTarget->InjectionSiteName;
+						TargetData->InjectionSite.bUseModuleFallback = false;
+					}
 					TargetData->Modifier = MakeShared<FMovieSceneAnimMixerEvaluationModifier>(MixerTask);
 
 					// Inject
-					TargetData->InjectionRequestHandle = UE::AnimNext::FInjectionUtils::InjectEvaluationModifier(AnimNextComponent, TargetData->Modifier.ToSharedRef(), TargetData->InjectionSiteName);
+					TargetData->InjectionRequestHandle = UE::AnimNext::FInjectionUtils::InjectEvaluationModifier(AnimNextComponent, TargetData->Modifier.ToSharedRef(), TargetData->InjectionSite);
 
 					// For now, we always restore state
 					PreAnimatedStorage->BeginTrackingEntity(EntityID, true, FRootInstanceHandle(), AnimNextComponent);
diff --git a/Engine/Plugins/Experimental/MovieSceneAnimMixer/Source/MovieSceneAnimMixer/Public/Systems/MovieSceneAnimNextTargetSystem.h b/Engine/Plugins/Experimental/MovieSceneAnimMixer/Source/MovieSceneAnimMixer/Public/Systems/MovieSceneAnimNextTargetSystem.h
--- a/Engine/Plugins/Experimental/MovieSceneAnimMixer/Source/MovieSceneAnimMixer/Public/Systems/MovieSceneAnimNextTargetSystem.h
+++ b/Engine/Plugins/Experimental/MovieSceneAnimMixer/Source/MovieSceneAnimMixer/Public/Systems/MovieSceneAnimNextTargetSystem.h
@@ -51,7 +51,7 @@ struct MOVIESCENEANIMMIXER_API FMovieSceneAnimNextInjectionTarget : public FMovi
 struct FMovieSceneAnimNextTargetData
 {
 	TWeakObjectPtr<UAnimNextComponent> AnimNextComponent = nullptr;
-	FName InjectionSiteName = FName();
+	UE::AnimNext::FInjectionSite InjectionSite;
 	UE::AnimNext::FInjectionRequestPtr InjectionRequestHandle = nullptr;
 	TSharedPtr<FMovieSceneAnimMixerEvaluationModifier> Modifier = nullptr;
 };
