Anim Next: FInjectionRequests now reference their host by weak object ptr

diff --git a/Engine/Plugins/Experimental/AnimNextAnimGraph/Source/AnimNextAnimGraph/Private/Injection/InjectionRequest.cpp b/Engine/Plugins/Experimental/AnimNextAnimGraph/Source/AnimNextAnimGraph/Private/Injection/InjectionRequest.cpp
--- a/Engine/Plugins/Experimental/AnimNextAnimGraph/Source/AnimNextAnimGraph/Private/Injection/InjectionRequest.cpp
+++ b/Engine/Plugins/Experimental/AnimNextAnimGraph/Source/AnimNextAnimGraph/Private/Injection/InjectionRequest.cpp
@@ -29,18 +29,25 @@ namespace UE::AnimNext
 			return false;	// Invalid args
 		}
 
-		SetupBindingModuleDependencies(InHost, InHandle, InRequestArgs);
+		UWorld* World = InHost->GetWorld();
+		if(World == nullptr)
+		{
+			return false;	// No world to play in
+		}
+
+		SetupBindingModuleDependencies(World, InHandle, InRequestArgs);
 
 		RequestArgs = MoveTemp(InRequestArgs);
 		LifetimeEvents = MoveTemp(InLifetimeEvents);
-		Host = InHost;
+		WeakHost = InHost;
+		WeakWorld = World;
 		Handle = InHandle;
 		Status = EInjectionStatus::Pending;
 
 		auto InjectEvent = MakeTraitEvent<FInjection_InjectEvent>();
 		InjectEvent->Request = AsShared();
 
-		UAnimNextWorldSubsystem* Subsystem = InHost->GetWorld()->GetSubsystem<UAnimNextWorldSubsystem>();
+		UAnimNextWorldSubsystem* Subsystem = World->GetSubsystem<UAnimNextWorldSubsystem>();
 		check(Subsystem);
 		Subsystem->QueueInputTraitEventHandle(InHandle, InjectEvent);
 
@@ -71,9 +78,12 @@ namespace UE::AnimNext
 		auto UninjectEvent = MakeTraitEvent<FInjection_UninjectEvent>();
 		UninjectEvent->Request = AsShared();
 
-		UAnimNextWorldSubsystem* Subsystem = Host->GetWorld()->GetSubsystem<UAnimNextWorldSubsystem>();
-		check(Subsystem);
-		Subsystem->QueueInputTraitEventHandle(Handle, UninjectEvent);
+		if (UWorld* World = WeakWorld.Get())
+		{
+			UAnimNextWorldSubsystem* Subsystem = World->GetSubsystem<UAnimNextWorldSubsystem>();
+			check(Subsystem);
+			Subsystem->QueueInputTraitEventHandle(Handle, UninjectEvent);
+		}
 	}
 
 	bool FInjectionRequest::ValidateArgs(const FInjectionRequestArgs& InRequestArgs)
@@ -104,31 +114,27 @@ namespace UE::AnimNext
 		return true;
 	}
 
-	void FInjectionRequest::SetupBindingModuleDependencies(UObject* InHost, FModuleHandle InHandle, const FInjectionRequestArgs& InRequestArgs)
+	void FInjectionRequest::SetupBindingModuleDependencies(UWorld* InWorld, FModuleHandle InHandle, const FInjectionRequestArgs& InRequestArgs)
 	{
 		check(IsInGameThread());
+		check(InWorld);
 
 		if(InRequestArgs.BindingModuleHandle.IsValid())
 		{
-			check(InHost);
-			UWorld* World = InHost->GetWorld();
-			check(World);
-			UAnimNextWorldSubsystem* Subsystem = World->GetSubsystem<UAnimNextWorldSubsystem>();
+			UAnimNextWorldSubsystem* Subsystem = InWorld->GetSubsystem<UAnimNextWorldSubsystem>();
 			check(Subsystem);
 			Subsystem->AddModuleDependencyHandle(InRequestArgs.BindingModuleHandle, InHandle);
 		}
 	}
 
-	void FInjectionRequest::RemoveBindingModuleDependencies(UObject* InHost, FModuleHandle InHandle, const FInjectionRequestArgs& InRequestArgs)
+	void FInjectionRequest::RemoveBindingModuleDependencies(UWorld* InWorld, FModuleHandle InHandle, const FInjectionRequestArgs& InRequestArgs)
 	{
 		check(IsInGameThread());
+		check(InWorld);
 
 		if(InRequestArgs.BindingModuleHandle.IsValid())
 		{
-			check(InHost);
-			UWorld* World = InHost->GetWorld();
-			check(World);
-			UAnimNextWorldSubsystem* Subsystem = World->GetSubsystem<UAnimNextWorldSubsystem>();
+			UAnimNextWorldSubsystem* Subsystem = InWorld->GetSubsystem<UAnimNextWorldSubsystem>();
 			check(Subsystem);
 			Subsystem->RemoveModuleDependencyHandle(InRequestArgs.BindingModuleHandle, InHandle);
 		}
@@ -222,7 +228,10 @@ namespace UE::AnimNext
 			Status = EInjectionStatus::Completed | (Status & EInjectionStatus::Interrupted);
 			LifetimeEvents.OnCompleted.ExecuteIfBound(*this);
 			// Remove any dependencies we have set up
-			RemoveBindingModuleDependencies(Host, Handle, RequestArgs);
+			if (UWorld* World = WeakWorld.Get())
+			{
+				RemoveBindingModuleDependencies(World, Handle, RequestArgs);
+			}
 			break;
 		case EInjectionStatus::Expired:
 			ensureMsgf(Status == EInjectionStatus::Pending, TEXT("Expected Injection status to be pending, found: %u"), (uint32)Status);
@@ -249,6 +258,5 @@ namespace UE::AnimNext
 	void FInjectionRequest::AddReferencedObjects(FReferenceCollector& Collector)
 	{
 		Collector.AddPropertyReferencesWithStructARO(FInjectionRequestArgs::StaticStruct(), &RequestArgs);
-		Collector.AddReferencedObject(Host);
 	}
 }
diff --git a/Engine/Plugins/Experimental/AnimNextAnimGraph/Source/AnimNextAnimGraph/Public/Injection/InjectionRequest.h b/Engine/Plugins/Experimental/AnimNextAnimGraph/Source/AnimNextAnimGraph/Public/Injection/InjectionRequest.h
--- a/Engine/Plugins/Experimental/AnimNextAnimGraph/Source/AnimNextAnimGraph/Public/Injection/InjectionRequest.h
+++ b/Engine/Plugins/Experimental/AnimNextAnimGraph/Source/AnimNextAnimGraph/Public/Injection/InjectionRequest.h
@@ -18,6 +18,7 @@ class FReferenceCollector;
 struct FAnimNextModuleInjectionComponent;
 struct FAnimNextGraphInstance;
 class UInjectionCallbackProxy;
+class UWorld;
 
 namespace UE::AnimNext
 {
@@ -220,10 +221,10 @@ namespace UE::AnimNext
 		static bool ValidateArgs(const FInjectionRequestArgs& InRequestArgs);
 
 		// Setup dependencies between the module being injected into and any module used for data interface bindings
-		static void SetupBindingModuleDependencies(UObject* InHost, FModuleHandle InHandle, const FInjectionRequestArgs& InRequestArgs);
+		static void SetupBindingModuleDependencies(UWorld* InWorld, FModuleHandle InHandle, const FInjectionRequestArgs& InRequestArgs);
 
 		// Remove dependencies between the module being injected into and any module used for data interface bindings
-		static void RemoveBindingModuleDependencies(UObject* InHost, FModuleHandle InHandle, const FInjectionRequestArgs& InRequestArgs);
+		static void RemoveBindingModuleDependencies(UWorld* InWorld, FModuleHandle InHandle, const FInjectionRequestArgs& InRequestArgs);
 
 		// The request arguments
 		FInjectionRequestArgs RequestArgs;
@@ -232,7 +233,10 @@ namespace UE::AnimNext
 		FInjectionLifetimeEvents LifetimeEvents;
 
 		// The object we are playing on
-		TObjectPtr<UObject> Host;
+		TWeakObjectPtr<UObject> WeakHost;
+
+		// The world within which we are playing
+		TWeakObjectPtr<UWorld> WeakWorld;
 
 		// Handle to the module instance
 		FModuleHandle Handle;
diff --git a/Engine/Plugins/Experimental/MovieSceneAnimMixer/Source/MovieSceneAnimMixer/Private/Systems/MovieSceneAnimNextTargetSystem.cpp b/Engine/Plugins/Experimental/MovieSceneAnimMixer/Source/MovieSceneAnimMixer/Private/Systems/MovieSceneAnimNextTargetSystem.cpp
--- a/Engine/Plugins/Experimental/MovieSceneAnimMixer/Source/MovieSceneAnimMixer/Private/Systems/MovieSceneAnimNextTargetSystem.cpp
+++ b/Engine/Plugins/Experimental/MovieSceneAnimMixer/Source/MovieSceneAnimMixer/Private/Systems/MovieSceneAnimNextTargetSystem.cpp
@@ -263,8 +263,7 @@ void UMovieSceneAnimNextTargetSystem::OnUnlink()
 	for (const FMovieSceneAnimNextTargetData& TargetData : CurrentTargets)
 	{
 		// Remove injection
-		if (TargetData.AnimNextComponent.Get() != nullptr && // This shouldn't be necessary due to UMovieSceneAnimNextTargetSystem::AddReferencedObjects, but somehow there is still a crash so we're going with the belt-and-braces approach
-			TargetData.InjectionRequestHandle.IsValid())
+		if (TargetData.InjectionRequestHandle.IsValid())
 		{
 			UE::AnimNext::FInjectionUtils::Uninject(TargetData.InjectionRequestHandle);
 		}
