Fixed AnimNext programmatic pins getting default value from the trait pin if no pin is connected ti the programatic pin
Fixed ControlRig generating correct default values for exposed variables. At creation it returns the mapped variable default value, after that, returns the pin default value if not connected or connected pin otherwise

diff --git a/Engine/Plugins/Experimental/AnimNext/Source/AnimNext/Internal/Graph/RigVMTrait_AnimNextPublicVariables.h b/Engine/Plugins/Experimental/AnimNext/Source/AnimNext/Internal/Graph/RigVMTrait_AnimNextPublicVariables.h
--- a/Engine/Plugins/Experimental/AnimNext/Source/AnimNext/Internal/Graph/RigVMTrait_AnimNextPublicVariables.h
+++ b/Engine/Plugins/Experimental/AnimNext/Source/AnimNext/Internal/Graph/RigVMTrait_AnimNextPublicVariables.h
@@ -1,4 +1,4 @@
-﻿// Copyright Epic Games, Inc. All Rights Reserved.
+// Copyright Epic Games, Inc. All Rights Reserved.
 
 #pragma once
 
@@ -27,7 +27,7 @@ struct ANIMNEXT_API FRigVMTrait_AnimNextPublicVariables : public FRigVMTrait
 	// FRigVMTrait interface
 #if WITH_EDITOR
 	virtual FString GetDisplayName() const override;
-	virtual void GetProgrammaticPins(URigVMController* InController, int32 InParentPinIndex, const FString& InDefaultValue, FRigVMPinInfoArray& OutPinArray) const override;
+	virtual void GetProgrammaticPins(URigVMController* InController, int32 InParentPinIndex, const URigVMPin* InTraitPin, const FString& InDefaultValue, FRigVMPinInfoArray& OutPinArray) const override;
 	virtual bool ShouldCreatePinForProperty(const FProperty* InProperty) const override;
 
 	// Editor-implemented function ptrs
diff --git a/Engine/Plugins/Experimental/AnimNext/Source/AnimNext/Private/Graph/RigVMTrait_AnimNextPublicVariables.cpp b/Engine/Plugins/Experimental/AnimNext/Source/AnimNext/Private/Graph/RigVMTrait_AnimNextPublicVariables.cpp
--- a/Engine/Plugins/Experimental/AnimNext/Source/AnimNext/Private/Graph/RigVMTrait_AnimNextPublicVariables.cpp
+++ b/Engine/Plugins/Experimental/AnimNext/Source/AnimNext/Private/Graph/RigVMTrait_AnimNextPublicVariables.cpp
@@ -1,4 +1,4 @@
-﻿// Copyright Epic Games, Inc. All Rights Reserved.
+// Copyright Epic Games, Inc. All Rights Reserved.
 
 #include "Graph/RigVMTrait_AnimNextPublicVariables.h"
 #include "AnimNextRigVMAsset.h"
@@ -14,7 +14,7 @@ FString FRigVMTrait_AnimNextPublicVariables::GetDisplayName() const
 	return GetDisplayNameFunc(*this);
 }
 
-void FRigVMTrait_AnimNextPublicVariables::GetProgrammaticPins(URigVMController* InController, int32 InParentPinIndex, const FString& InDefaultValue, FRigVMPinInfoArray& OutPinArray) const
+void FRigVMTrait_AnimNextPublicVariables::GetProgrammaticPins(URigVMController* InController, int32 InParentPinIndex, const URigVMPin* InTraitPin, const FString& InDefaultValue, FRigVMPinInfoArray& OutPinArray) const
 {
 	GetProgrammaticPinsFunc(*this, InController, InParentPinIndex, InDefaultValue, OutPinArray);
 }
diff --git a/Engine/Plugins/Experimental/AnimNextAnimGraph/Source/AnimNextAnimGraph/Internal/Graph/RigDecorator_AnimNextCppTrait.h b/Engine/Plugins/Experimental/AnimNextAnimGraph/Source/AnimNextAnimGraph/Internal/Graph/RigDecorator_AnimNextCppTrait.h
--- a/Engine/Plugins/Experimental/AnimNextAnimGraph/Source/AnimNextAnimGraph/Internal/Graph/RigDecorator_AnimNextCppTrait.h
+++ b/Engine/Plugins/Experimental/AnimNextAnimGraph/Source/AnimNextAnimGraph/Internal/Graph/RigDecorator_AnimNextCppTrait.h
@@ -22,7 +22,7 @@ struct ANIMNEXTANIMGRAPH_API FRigDecorator_AnimNextCppDecorator : public FRigVMT
 	TObjectPtr<UScriptStruct> DecoratorSharedDataStruct = nullptr;
 
 #if WITH_EDITOR
-	virtual void GetProgrammaticPins(URigVMController* InController, int32 InParentPinIndex, const FString& InDefaultValue, struct FRigVMPinInfoArray& OutPinArray) const override;
+	virtual void GetProgrammaticPins(URigVMController* InController, int32 InParentPinIndex, const URigVMPin* InTraitPin, const FString& InDefaultValue, struct FRigVMPinInfoArray& OutPinArray) const override;
 	
 	virtual UScriptStruct* GetTraitSharedDataStruct() const override
 	{
diff --git a/Engine/Plugins/Experimental/AnimNextAnimGraph/Source/AnimNextAnimGraph/Private/Graph/RigDecorator_AnimNextCppTrait.cpp b/Engine/Plugins/Experimental/AnimNextAnimGraph/Source/AnimNextAnimGraph/Private/Graph/RigDecorator_AnimNextCppTrait.cpp
--- a/Engine/Plugins/Experimental/AnimNextAnimGraph/Source/AnimNextAnimGraph/Private/Graph/RigDecorator_AnimNextCppTrait.cpp
+++ b/Engine/Plugins/Experimental/AnimNextAnimGraph/Source/AnimNextAnimGraph/Private/Graph/RigDecorator_AnimNextCppTrait.cpp
@@ -10,7 +10,7 @@
 #include UE_INLINE_GENERATED_CPP_BY_NAME(RigDecorator_AnimNextCppTrait)
 
 #if WITH_EDITOR
-void FRigDecorator_AnimNextCppDecorator::GetProgrammaticPins(URigVMController* InController, int32 InParentPinIndex, const FString& InDefaultValue, FRigVMPinInfoArray& OutPinArray) const
+void FRigDecorator_AnimNextCppDecorator::GetProgrammaticPins(URigVMController* InController, int32 InParentPinIndex, const URigVMPin* InTraitPin, const FString& InDefaultValue, FRigVMPinInfoArray& OutPinArray) const
 {
 	if (DecoratorSharedDataStruct == nullptr)
 	{
@@ -39,19 +39,29 @@ void FRigDecorator_AnimNextCppDecorator::GetProgrammaticPins(URigVMController* I
 		return ERigVMPinDefaultValueType::AutoDetect;
 	};
 
-	const int32 StartPinIndex = OutPinArray.Num();
 
 	// Give the Trait an opportunity to generate its own pins
 	const UE::AnimNext::FTrait* Trait = GetTrait();
 	check(Trait != nullptr);
 	FAnimNextTraitSharedData* TraitSharedData = (FAnimNextTraitSharedData*)DefaultValueMemoryScope.GetStructMemory();
-	Trait->GetProgrammaticPins(TraitSharedData, InController, InParentPinIndex, InDefaultValue, OutPinArray);
-	
+
+	// Obtain programmatic pins
+	const int32 ProgrammaticPinIndexStart = OutPinArray.Num();
+	Trait->GetProgrammaticPins(TraitSharedData, InController, InParentPinIndex, InTraitPin, InDefaultValue, OutPinArray);
+
+	const int32 ProgrammaticPinIndexEnd = OutPinArray.Num();
+	for (int32 PinIndex = ProgrammaticPinIndexStart; PinIndex < ProgrammaticPinIndexEnd; ++PinIndex)
+	{
+		FRigVMPinInfo& PinInfo = OutPinArray[PinIndex];
+		PinInfo.bIsLazy = true;			// we flag all programmatic pins as lazy, as currently we have to remap them inside the trait
+		PinInfo.Property = nullptr;		// we setup the pin explicitly, avoid RigVM do additional setup
+	}
+
 	// Generate the shared data struct pins
 	OutPinArray.AddPins(DecoratorSharedDataStruct, InController, ERigVMPinDirection::Invalid, InParentPinIndex, DefaultValueTypeGetter, DefaultValueMemoryScope.GetStructMemory(), true);
 
 	const int32 NumPins = OutPinArray.Num();
-	for (int32 PinIndex = StartPinIndex; PinIndex < NumPins; ++PinIndex)
+	for (int32 PinIndex = ProgrammaticPinIndexEnd; PinIndex < NumPins; ++PinIndex)  // iterate pins created after the programmatic ones
 	{
 		FRigVMPinInfo& PinInfo = OutPinArray[PinIndex];
 
diff --git a/Engine/Plugins/Experimental/AnimNextAnimGraph/Source/AnimNextAnimGraph/Public/TraitCore/Trait.h b/Engine/Plugins/Experimental/AnimNextAnimGraph/Source/AnimNextAnimGraph/Public/TraitCore/Trait.h
--- a/Engine/Plugins/Experimental/AnimNextAnimGraph/Source/AnimNextAnimGraph/Public/TraitCore/Trait.h
+++ b/Engine/Plugins/Experimental/AnimNextAnimGraph/Source/AnimNextAnimGraph/Public/TraitCore/Trait.h
@@ -18,6 +18,7 @@
 
 class FArchive;
 class URigVMController;
+class URigVMPin;
 struct FRigVMPinInfoArray;
 
 #define ANIM_NEXT_IMPL_DECLARE_ANIM_TRAIT_BASIC(TraitName, SuperTraitName) \
@@ -461,7 +462,7 @@ namespace UE::AnimNext
 		virtual bool MultipleInstanceSupport() const { return false; }
 
 		// Enable Traits to generate TraitStack pins
-		virtual void GetProgrammaticPins(FAnimNextTraitSharedData* InSharedData, URigVMController* InController, int32 InParentPinIndex, const FString& InDefaultValue, struct FRigVMPinInfoArray& OutPinArray) const {}
+		virtual void GetProgrammaticPins(FAnimNextTraitSharedData* InSharedData, URigVMController* InController, int32 InParentPinIndex, const URigVMPin* InTraitPin, const FString& InDefaultValue, struct FRigVMPinInfoArray& OutPinArray) const {}
 #endif
 
 	protected:
diff --git a/Engine/Plugins/Experimental/AnimNextAnimGraph/Source/AnimNextAnimGraphUncookedOnly/Private/Graph/AnimNextAnimationGraph_EditorData.cpp b/Engine/Plugins/Experimental/AnimNextAnimGraph/Source/AnimNextAnimGraphUncookedOnly/Private/Graph/AnimNextAnimationGraph_EditorData.cpp
--- a/Engine/Plugins/Experimental/AnimNextAnimGraph/Source/AnimNextAnimGraphUncookedOnly/Private/Graph/AnimNextAnimationGraph_EditorData.cpp
+++ b/Engine/Plugins/Experimental/AnimNextAnimGraph/Source/AnimNextAnimGraphUncookedOnly/Private/Graph/AnimNextAnimationGraph_EditorData.cpp
@@ -411,9 +411,21 @@ namespace UE::AnimNext::UncookedOnly::Private
 		{
 			for (FTraitEntryMapping& TraitEntry : TraitStack.TraitEntries)
 			{
+				TSharedPtr<FStructOnScope> DecoratorScope = TraitEntry.DecoratorEntryPin->GetTraitInstance();
+				const FRigDecorator_AnimNextCppDecorator* Decorator = (const FRigDecorator_AnimNextCppDecorator*)DecoratorScope->GetStructMemory();
+				const UScriptStruct* SharedDataStruct = Decorator->GetTraitSharedDataStruct();
+
 				for (URigVMPin* Pin : TraitEntry.DecoratorEntryPin->GetSubPins())
 				{
-					if (Pin->IsLazy() && !Pin->GetLinks().IsEmpty())
+					if (!Pin->IsLazy())
+					{
+						continue;
+					}
+					
+					// note that Pin->IsProgrammaticPin(); does not work, it does not check the shared struct
+					const bool bIsProgrammaticPin = SharedDataStruct->FindPropertyByName(Pin->GetFName()) == nullptr;
+					const bool bHasLinks = !Pin->GetLinks().IsEmpty();
+					if (bHasLinks || bIsProgrammaticPin)
 					{
 						// This pin has something linked to it, it is a latent pin
 						check(OutLatentPins.Num() < ((1 << 16) - 1));	// We reserve MAX_uint16 as an invalid value and we must fit on 15 bits when packed
@@ -428,13 +440,23 @@ namespace UE::AnimNext::UncookedOnly::Private
 						// All our programmatic pins are lazy inputs
 						PinInfo.Direction = ERigVMPinDirection::Input;
 						PinInfo.bIsLazy = true;
+						PinInfo.DefaultValue = Pin->GetDefaultValue();
+						PinInfo.DefaultValueType = ERigVMPinDefaultValueType::AutoDetect;
 
 						OutLatentPins.Pins.Emplace(PinInfo);
 
-						const TArray<URigVMLink*>& PinLinks = Pin->GetLinks();
-						check(PinLinks.Num() == 1);
+						if (bHasLinks)
+						{
+							const TArray<URigVMLink*>& PinLinks = Pin->GetLinks();
+							check(PinLinks.Num() == 1);
 
-						OutLatentPinMapping.Add(LatentPinName, PinLinks[0]->GetSourcePin());
+							OutLatentPinMapping.Add(LatentPinName, PinLinks[0]->GetSourcePin());
+						}
+						else if(bIsProgrammaticPin)
+						{
+							// this is a programmatic pin, we make it latent with itself, so we can remap it at trait level
+							OutLatentPinMapping.Add(LatentPinName, Pin);
+						}
 					}
 				}
 			}
diff --git a/Engine/Plugins/Experimental/AnimNextControlRig/Source/AnimNextControlRig/Internal/ControlRigTrait.h b/Engine/Plugins/Experimental/AnimNextControlRig/Source/AnimNextControlRig/Internal/ControlRigTrait.h
--- a/Engine/Plugins/Experimental/AnimNextControlRig/Source/AnimNextControlRig/Internal/ControlRigTrait.h
+++ b/Engine/Plugins/Experimental/AnimNextControlRig/Source/AnimNextControlRig/Internal/ControlRigTrait.h
@@ -260,6 +260,7 @@ protected:
 	virtual void GetProgrammaticPins(FAnimNextTraitSharedData* InSharedData,
 		URigVMController* InController,
 		int32 InParentPinIndex,
+		const URigVMPin* TraitPin,
 		const FString& InDefaultValue,
 		struct FRigVMPinInfoArray& OutPinArray) const override;
 
diff --git a/Engine/Plugins/Experimental/AnimNextControlRig/Source/AnimNextControlRig/Private/ControlRigTrait.cpp b/Engine/Plugins/Experimental/AnimNextControlRig/Source/AnimNextControlRig/Private/ControlRigTrait.cpp
--- a/Engine/Plugins/Experimental/AnimNextControlRig/Source/AnimNextControlRig/Private/ControlRigTrait.cpp
+++ b/Engine/Plugins/Experimental/AnimNextControlRig/Source/AnimNextControlRig/Private/ControlRigTrait.cpp
@@ -18,6 +18,7 @@
 #include "Engine/SkeletalMesh.h"
 #include "Animation/AnimNodeBase.h"
 #include "AnimationDataSource.h"
+#include "RigVMCore/RigVMMemoryStorageStruct.h"
 #if WITH_EDITOR
 #include "RigVMModel/RigVMController.h"
 #include "ControlRigBlueprint.h" // to get the preview skeleton
@@ -595,7 +596,7 @@ bool FControlRigTrait::SetupControlRig(UObject* InAnimContext, UClass* InControl
 }
 
 #if WITH_EDITOR
-void FControlRigTrait::GetProgrammaticPins(FAnimNextTraitSharedData* InSharedData, URigVMController* InController, int32 InParentPinIndex, const FString& InDefaultValue, struct FRigVMPinInfoArray& OutPinArray) const
+void FControlRigTrait::GetProgrammaticPins(FAnimNextTraitSharedData* InSharedData, URigVMController* InController, int32 InParentPinIndex, const URigVMPin* InTraitPin, const FString& InDefaultValue, struct FRigVMPinInfoArray& OutPinArray) const
 {
 	if (FControlRigTraitSharedData* SharedData = (FControlRigTraitSharedData*)InSharedData)
 	{
@@ -615,7 +616,19 @@ void FControlRigTrait::GetProgrammaticPins(FAnimNextTraitSharedData* InSharedDat
 							continue;
 						}
 
-						OutPinArray.AddPin(const_cast<FProperty*>(Variable.Property), InController, ERigVMPinDirection::Input, InParentPinIndex, ERigVMPinDefaultValueType::AutoDetect, Variable.Memory, true);
+						if (const URigVMPin* SubPin = InTraitPin->FindSubPin(Variable.Name.ToString()))
+						{
+							const FString PinDefaultValue = SubPin->GetDefaultValue();
+							FRigVMMemoryStorageStruct StorageDefaultValue(ERigVMMemoryType::External, {FRigVMPropertyDescription(Variable.Property, PinDefaultValue, Variable.Name)});
+
+							const uint8* DefaultMemory = StorageDefaultValue.GetDataByName<const uint8>(Variable.Name);
+							
+							OutPinArray.AddPin(const_cast<FProperty*>(Variable.Property), InController, ERigVMPinDirection::Input, InParentPinIndex, ERigVMPinDefaultValueType::AutoDetect, DefaultMemory, true);
+						}
+						else
+						{
+							OutPinArray.AddPin(const_cast<FProperty*>(Variable.Property), InController, ERigVMPinDirection::Input, InParentPinIndex, ERigVMPinDefaultValueType::AutoDetect, Variable.Memory, true);
+						}
 					}
 				}
 			}
diff --git a/Engine/Plugins/Runtime/RigVM/Source/RigVM/Public/RigVMCore/RigVMTrait.h b/Engine/Plugins/Runtime/RigVM/Source/RigVM/Public/RigVMCore/RigVMTrait.h
--- a/Engine/Plugins/Runtime/RigVM/Source/RigVM/Public/RigVMCore/RigVMTrait.h
+++ b/Engine/Plugins/Runtime/RigVM/Source/RigVM/Public/RigVMCore/RigVMTrait.h
@@ -47,7 +47,7 @@ public:
 	virtual void OnTraitAdded(URigVMController* InController, URigVMNode* InNode) {}
 
 	// allows the trait to return dynamic pins (parent pin index must be INDEX_NONE or point to a valid index of the parent pin in the OutPinArray)
-	virtual void GetProgrammaticPins(URigVMController* InController, int32 InParentPinIndex, const FString& InDefaultValue, struct FRigVMPinInfoArray& OutPinArray) const {}
+	virtual void GetProgrammaticPins(URigVMController* InController, int32 InParentPinIndex, const URigVMPin* InTraitPin, const FString& InDefaultValue, struct FRigVMPinInfoArray& OutPinArray) const {}
 
 	virtual UScriptStruct* GetTraitSharedDataStruct() const { return nullptr; }
 
diff --git a/Engine/Plugins/Runtime/RigVM/Source/RigVMDeveloper/Private/RigVMModel/RigVMController.cpp b/Engine/Plugins/Runtime/RigVM/Source/RigVMDeveloper/Private/RigVMModel/RigVMController.cpp
--- a/Engine/Plugins/Runtime/RigVM/Source/RigVMDeveloper/Private/RigVMModel/RigVMController.cpp
+++ b/Engine/Plugins/Runtime/RigVM/Source/RigVMDeveloper/Private/RigVMModel/RigVMController.cpp
@@ -16591,7 +16591,7 @@ FName URigVMController::AddTrait(URigVMNode* InNode, UScriptStruct* InTraitScrip
 	AddPinsForStruct(InTraitScriptStruct, InNode, TraitPin, TraitPin->GetDirection(), DefaultValue, true);
 
 	FRigVMPinInfoArray ProgrammaticPins;
-	Trait->GetProgrammaticPins(this, INDEX_NONE, DefaultValue, ProgrammaticPins);
+	Trait->GetProgrammaticPins(this, INDEX_NONE, TraitPin, DefaultValue, ProgrammaticPins);
 
 	const FRigVMRegistry& Registry = FRigVMRegistry::Get();
 	const FRigVMPinInfoArray PreviousPins;
@@ -18191,7 +18191,7 @@ bool URigVMController::GenerateNewPinInfos(const FRigVMRegistry& Registry, URigV
 				TSharedPtr<FStructOnScope> TraitScope = Pin->GetTraitInstance();
 				FRigVMTrait* VMTrait = (FRigVMTrait*)TraitScope->GetStructMemory();
 
-				VMTrait->GetProgrammaticPins(this, NewPinIndex, Pin->GetDefaultValue(), NewPinInfos);
+				VMTrait->GetProgrammaticPins(this, NewPinIndex, Pin, Pin->GetDefaultValue(), NewPinInfos);
 			}
 		}
 	}
