EvaluationProgram rewind debugger recording

diff --git a/Engine/Plugins/Experimental/AnimNextAnimGraph/Source/AnimNextAnimGraph/AnimNextAnimGraph.Build.cs b/Engine/Plugins/Experimental/AnimNextAnimGraph/Source/AnimNextAnimGraph/AnimNextAnimGraph.Build.cs
--- a/Engine/Plugins/Experimental/AnimNextAnimGraph/Source/AnimNextAnimGraph/AnimNextAnimGraph.Build.cs
+++ b/Engine/Plugins/Experimental/AnimNextAnimGraph/Source/AnimNextAnimGraph/AnimNextAnimGraph.Build.cs
@@ -22,7 +22,8 @@ namespace UnrealBuildTool.Rules
 					"Engine",
 					"AnimNext",
 					"HierarchyTableRuntime",
-					"HierarchyTableAnimationRuntime"
+					"HierarchyTableAnimationRuntime",
+					"TraceLog",
 				}
 			);
 
diff --git a/Engine/Plugins/Experimental/AnimNextAnimGraph/Source/AnimNextAnimGraph/Private/EvaluationVM/SerializableEvaluationProgram.cpp b/Engine/Plugins/Experimental/AnimNextAnimGraph/Source/AnimNextAnimGraph/Private/EvaluationVM/SerializableEvaluationProgram.cpp
new file mode 100644
--- /dev/null
+++ b/Engine/Plugins/Experimental/AnimNextAnimGraph/Source/AnimNextAnimGraph/Private/EvaluationVM/SerializableEvaluationProgram.cpp
@@ -0,0 +1,13 @@
+// Copyright Epic Games, Inc. All Rights Reserved.
+
+#include "EvaluationVM/SerializableEvaluationProgram.h"
+#include "EvaluationVM/EvaluationProgram.h"
+
+FSerializableEvaluationProgram::FSerializableEvaluationProgram(const UE::AnimNext::FEvaluationProgram& Other)
+{
+	for (const TSharedPtr<FAnimNextEvaluationTask>& Task : Other.Tasks)
+	{
+		Tasks.Add(FInstancedStruct(Task->GetStruct()));
+		Task->GetStruct()->CopyScriptStruct(Tasks.Last().GetMutableMemory(), &*Task);
+	}
+}
\ No newline at end of file
diff --git a/Engine/Plugins/Experimental/AnimNextAnimGraph/Source/AnimNextAnimGraph/Private/Graph/TraceAnimNextGraphInstances.cpp b/Engine/Plugins/Experimental/AnimNextAnimGraph/Source/AnimNextAnimGraph/Private/Graph/TraceAnimNextGraphInstances.cpp
--- a/Engine/Plugins/Experimental/AnimNextAnimGraph/Source/AnimNextAnimGraph/Private/Graph/TraceAnimNextGraphInstances.cpp
+++ b/Engine/Plugins/Experimental/AnimNextAnimGraph/Source/AnimNextAnimGraph/Private/Graph/TraceAnimNextGraphInstances.cpp
@@ -6,8 +6,21 @@
 #include "TraitCore/ExecutionContext.h"
 #include "TraitCore/NodeInstance.h"
 #include "TraitInterfaces/IHierarchy.h"
+#include "EvaluationVM/EvaluationProgram.h"
+#include "EvaluationVM/SerializableEvaluationProgram.h"
+#include "Serialization/MemoryWriter.h"
+#include "Serialization/ObjectAndNameAsStringProxyArchive.h"
 
 #if ANIMNEXT_TRACE_ENABLED
+
+UE_TRACE_EVENT_BEGIN(AnimNextAnimGraph, EvaluationProgram)
+	UE_TRACE_EVENT_FIELD(uint64, Cycle)
+	UE_TRACE_EVENT_FIELD(double, RecordingTime)
+	UE_TRACE_EVENT_FIELD(uint64, OuterObjectId)
+	UE_TRACE_EVENT_FIELD(uint64, InstanceId)
+	UE_TRACE_EVENT_FIELD(uint8[], ProgramData)
+UE_TRACE_EVENT_END()
+
 namespace UE::AnimNext
 {
 	void TraceGraphInstances(const FAnimNextGraphInstance& RootGraph)
@@ -54,5 +67,53 @@ namespace UE::AnimNext
 			}
 		}
 	}
+
+	void TraceEvaluationProgram(const FEvaluationProgram& Program, const FAnimNextGraphInstance& RootGraph)
+	{
+		bool bChannelEnabled = UE_TRACE_CHANNELEXPR_IS_ENABLED(AnimNextChannel);
+		if (!bChannelEnabled)
+		{
+			return;
+		}
+
+		uint64 OuterObjectId = 0;
+		if (FAnimNextModuleInstance* ModuleInstance = RootGraph.GetModuleInstance())
+		{
+			if (const UObject* OuterObject = ModuleInstance->GetObject())
+			{
+				OuterObjectId = FObjectTrace::GetObjectId(OuterObject);
+			}
+		}
+		
+		uint64 InstanceId = RootGraph.GetUniqueId();
+		
+		double RecordingTime = 0;
+		if (RootGraph.GetModuleInstance())
+		{
+			if (UObject* OuterObject = RootGraph.GetModuleInstance()->GetObject())
+			{
+				if (CANNOT_TRACE_OBJECT(OuterObject))
+				{
+					return;
+				}
+				RecordingTime = FObjectTrace::GetWorldElapsedTime(OuterObject->GetWorld());
+			}
+		}
+
+		TArray<uint8> ArchiveData;
+		FMemoryWriter WriterArchive(ArchiveData);
+		FObjectAndNameAsStringProxyArchive Archive(WriterArchive, /*bInLoadIfFindFails*/true);
+		
+		static const FSerializableEvaluationProgram Defaults;
+		FSerializableEvaluationProgram SerializableProgram(Program);
+		FSerializableEvaluationProgram::StaticStruct()->SerializeItem(Archive, &SerializableProgram, &Defaults);
+		
+		UE_TRACE_LOG(AnimNextAnimGraph, EvaluationProgram, AnimNextChannel)
+				<< EvaluationProgram.Cycle(FPlatformTime::Cycles64())
+				<< EvaluationProgram.OuterObjectId(OuterObjectId)
+				<< EvaluationProgram.RecordingTime(RecordingTime)
+				<< EvaluationProgram.InstanceId(InstanceId)
+				<< EvaluationProgram.ProgramData(ArchiveData.GetData(), ArchiveData.Num());
+	}
 }
 #endif
diff --git a/Engine/Plugins/Experimental/AnimNextAnimGraph/Source/AnimNextAnimGraph/Private/Graph/TraceAnimNextGraphInstances.h b/Engine/Plugins/Experimental/AnimNextAnimGraph/Source/AnimNextAnimGraph/Private/Graph/TraceAnimNextGraphInstances.h
--- a/Engine/Plugins/Experimental/AnimNextAnimGraph/Source/AnimNextAnimGraph/Private/Graph/TraceAnimNextGraphInstances.h
+++ b/Engine/Plugins/Experimental/AnimNextAnimGraph/Source/AnimNextAnimGraph/Private/Graph/TraceAnimNextGraphInstances.h
@@ -1,6 +1,7 @@
 // Copyright Epic Games, Inc. All Rights Reserved.
 #pragma once
 
+#include "EvaluationVM/EvaluationProgram.h"
 #include "RewindDebugger/AnimNextTrace.h"
 
 #if ANIMNEXT_TRACE_ENABLED
@@ -9,13 +10,18 @@ struct FAnimNextGraphInstance;
 
 namespace UE::AnimNext
 {
+	struct FEvaluationProgram;
+	
 	void TraceGraphInstances(const FAnimNextGraphInstance& RootGraph);
+	void TraceEvaluationProgram(const FEvaluationProgram& Program, const FAnimNextGraphInstance& RootGraph);
 }
 
 #define TRACE_ANIMNEXT_GRAPHINSTANCES(RootGraph) UE::AnimNext::TraceGraphInstances(RootGraph);
+#define TRACE_ANIMNEXT_EVALUATIONPROGRAM(Program, RootGraph) UE::AnimNext::TraceEvaluationProgram(Program, RootGraph);
 
 #else
 
 #define TRACE_ANIMNEXT_GRAPHINSTANCES(RootGraph)
+#define TRACE_ANIMNEXT_EVALUATIONPROGRAM(Program, RootGraph)
 
 #endif
diff --git a/Engine/Plugins/Experimental/AnimNextAnimGraph/Source/AnimNextAnimGraph/Private/TraitInterfaces/IEvaluate.cpp b/Engine/Plugins/Experimental/AnimNextAnimGraph/Source/AnimNextAnimGraph/Private/TraitInterfaces/IEvaluate.cpp
--- a/Engine/Plugins/Experimental/AnimNextAnimGraph/Source/AnimNextAnimGraph/Private/TraitInterfaces/IEvaluate.cpp
+++ b/Engine/Plugins/Experimental/AnimNextAnimGraph/Source/AnimNextAnimGraph/Private/TraitInterfaces/IEvaluate.cpp
@@ -9,6 +9,7 @@
 #include "EvaluationVM/EvaluationVM.h"
 #include "Graph/AnimNextGraphInstance.h"
 #include "Graph/AnimNext_LODPose.h"
+#include "Graph/TraceAnimNextGraphInstances.h"
 
 namespace UE::AnimNext
 {
@@ -224,6 +225,8 @@ namespace UE::AnimNext
 	{
 		const FEvaluationProgram EvaluationProgram = EvaluateGraph(GraphInstance);
 
+		TRACE_ANIMNEXT_EVALUATIONPROGRAM(EvaluationProgram, GraphInstance);
+
 		FEvaluationVM EvaluationVM(EEvaluationFlags::All, RefPose, GraphLODLevel);
 		bool bHasValidOutput = false;
 
diff --git a/Engine/Plugins/Experimental/AnimNextAnimGraph/Source/AnimNextAnimGraph/Public/EvaluationVM/EvaluationProgram.h b/Engine/Plugins/Experimental/AnimNextAnimGraph/Source/AnimNextAnimGraph/Public/EvaluationVM/EvaluationProgram.h
--- a/Engine/Plugins/Experimental/AnimNextAnimGraph/Source/AnimNextAnimGraph/Public/EvaluationVM/EvaluationProgram.h
+++ b/Engine/Plugins/Experimental/AnimNextAnimGraph/Source/AnimNextAnimGraph/Public/EvaluationVM/EvaluationProgram.h
@@ -7,6 +7,9 @@
 
 #include <type_traits>
 
+
+struct FSerializableEvaluationProgram;
+
 namespace UE::AnimNext
 {
 	struct FEvaluationVM;
@@ -46,6 +49,7 @@ namespace UE::AnimNext
 		// Returns the program as a string suitable for debug purposes
 		FString ToString() const;
 
+		friend struct ::FSerializableEvaluationProgram;
 	private:
 		// Disallow copy
 		FEvaluationProgram(const FEvaluationProgram&) = delete;
diff --git a/Engine/Plugins/Experimental/AnimNextAnimGraph/Source/AnimNextAnimGraph/Public/EvaluationVM/SerializableEvaluationProgram.h b/Engine/Plugins/Experimental/AnimNextAnimGraph/Source/AnimNextAnimGraph/Public/EvaluationVM/SerializableEvaluationProgram.h
new file mode 100644
--- /dev/null
+++ b/Engine/Plugins/Experimental/AnimNextAnimGraph/Source/AnimNextAnimGraph/Public/EvaluationVM/SerializableEvaluationProgram.h
@@ -0,0 +1,28 @@
+// Copyright Epic Games, Inc. All Rights Reserved.
+
+#pragma once
+
+#include "CoreMinimal.h"
+#include "StructUtils/InstancedStruct.h"
+
+#include "SerializableEvaluationProgram.generated.h"
+
+namespace UE::AnimNext
+{
+	struct FEvaluationProgram;
+}
+
+USTRUCT()
+struct FSerializableEvaluationProgram
+{
+	GENERATED_BODY()
+	
+	FSerializableEvaluationProgram()
+	{
+	}
+	
+	FSerializableEvaluationProgram(const UE::AnimNext::FEvaluationProgram& Other);
+	
+	UPROPERTY(VisibleAnywhere, Category = "Tasks", meta = (ExpandByDefault))
+	TArray<FInstancedStruct> Tasks;	
+};
diff --git a/Engine/Plugins/Experimental/AnimNextAnimGraph/Source/AnimNextAnimGraph/Public/EvaluationVM/Tasks/BlendKeyframes.h b/Engine/Plugins/Experimental/AnimNextAnimGraph/Source/AnimNextAnimGraph/Public/EvaluationVM/Tasks/BlendKeyframes.h
--- a/Engine/Plugins/Experimental/AnimNextAnimGraph/Source/AnimNextAnimGraph/Public/EvaluationVM/Tasks/BlendKeyframes.h
+++ b/Engine/Plugins/Experimental/AnimNextAnimGraph/Source/AnimNextAnimGraph/Public/EvaluationVM/Tasks/BlendKeyframes.h
@@ -52,7 +52,7 @@ struct ANIMNEXTANIMGRAPH_API FAnimNextBlendTwoKeyframesTask : public FAnimNextEv
 	float GetInterpolationAlpha(const UE::AnimNext::FKeyframeState* KeyframeA, const UE::AnimNext::FKeyframeState* KeyframeB) const;
 
 	// The interpolation alpha between the two input keyframes
-	UPROPERTY()
+	UPROPERTY(VisibleAnywhere, Category = Properties)
 	float InterpolationAlpha = 0.0f;
 
 	// The curve to evaluate and extract the interpolation alpha between the two input keyframes
@@ -85,7 +85,7 @@ struct ANIMNEXTANIMGRAPH_API FAnimNextBlendOverwriteKeyframeWithScaleTask : publ
 	virtual void Execute(UE::AnimNext::FEvaluationVM& VM) const override;
 
 	// The scale factor to apply to the keyframe
-	UPROPERTY()
+	UPROPERTY(VisibleAnywhere, Category = Properties)
 	float ScaleFactor = 1.0f;
 };
 
@@ -112,6 +112,6 @@ struct ANIMNEXTANIMGRAPH_API FAnimNextBlendAddKeyframeWithScaleTask : public FAn
 	virtual void Execute(UE::AnimNext::FEvaluationVM& VM) const override;
 
 	// The scale factor to apply to the keyframe
-	UPROPERTY()
+	UPROPERTY(VisibleAnywhere, Category = Properties)
 	float ScaleFactor = 0.0f;
 };
diff --git a/Engine/Plugins/Experimental/AnimNextAnimGraph/Source/AnimNextAnimGraph/Public/EvaluationVM/Tasks/DeadBlending.h b/Engine/Plugins/Experimental/AnimNextAnimGraph/Source/AnimNextAnimGraph/Public/EvaluationVM/Tasks/DeadBlending.h
--- a/Engine/Plugins/Experimental/AnimNextAnimGraph/Source/AnimNextAnimGraph/Public/EvaluationVM/Tasks/DeadBlending.h
+++ b/Engine/Plugins/Experimental/AnimNextAnimGraph/Source/AnimNextAnimGraph/Public/EvaluationVM/Tasks/DeadBlending.h
@@ -114,7 +114,7 @@ struct ANIMNEXTANIMGRAPH_API FAnimNextDeadBlendingApplyTask : public FAnimNextEv
 	float TimeSinceTransition = 0.0f;
 	EAlphaBlendOption BlendMode = EAlphaBlendOption::Linear;
 	
-	UPROPERTY()
+	UPROPERTY(VisibleAnywhere, Category = Properties)
 	TWeakObjectPtr<UCurveFloat> CustomBlendCurve = nullptr;
 };
 
diff --git a/Engine/Plugins/Experimental/AnimNextAnimGraph/Source/AnimNextAnimGraph/Public/EvaluationVM/Tasks/PushAnimSequenceKeyframe.h b/Engine/Plugins/Experimental/AnimNextAnimGraph/Source/AnimNextAnimGraph/Public/EvaluationVM/Tasks/PushAnimSequenceKeyframe.h
--- a/Engine/Plugins/Experimental/AnimNextAnimGraph/Source/AnimNextAnimGraph/Public/EvaluationVM/Tasks/PushAnimSequenceKeyframe.h
+++ b/Engine/Plugins/Experimental/AnimNextAnimGraph/Source/AnimNextAnimGraph/Public/EvaluationVM/Tasks/PushAnimSequenceKeyframe.h
@@ -29,7 +29,7 @@ struct ANIMNEXTANIMGRAPH_API FAnimNextAnimSequenceKeyframeTask : public FAnimNex
 	virtual void Execute(UE::AnimNext::FEvaluationVM& VM) const override;
 
 	// Anim Sequence to grab the keyframe from
-	UPROPERTY()
+	UPROPERTY(VisibleAnywhere, Category = Properties)
 	TWeakObjectPtr<UAnimSequence> AnimSequence;
 
 	/** Delta time range required for root motion extraction **/
@@ -37,23 +37,23 @@ struct ANIMNEXTANIMGRAPH_API FAnimNextAnimSequenceKeyframeTask : public FAnimNex
 
 	// The point in time within the animation sequence at which we sample the keyframe.
 	// If negative, the sample time hasn't been provided and we should use the keyframe index.
-	UPROPERTY()
+	UPROPERTY(VisibleAnywhere, Category = Properties)
 	double SampleTime = -1.0;
 
 	// The specific keyframe within the animation sequence to retrieve.
 	// If ~0, the keyframe index hasn't been provided and we should use the sample time.
-	UPROPERTY()
+	UPROPERTY(VisibleAnywhere, Category = Properties)
 	uint32 KeyframeIndex = ~0;
 
 	// Whether to interpolate or step the animation sequence.
 	// Only used when the sample time is used.
-	UPROPERTY()
+	UPROPERTY(VisibleAnywhere, Category = Properties)
 	bool bInterpolate = false;
 
 	// Whether to extract trajectory or not
-	UPROPERTY()
+	UPROPERTY(VisibleAnywhere, Category = Properties)
 	bool bExtractTrajectory = false;
 
-	UPROPERTY()
+	UPROPERTY(VisibleAnywhere, Category = Properties)
 	bool bLooping = false;
 };
diff --git a/Engine/Plugins/Experimental/AnimNextAnimGraph/Source/AnimNextAnimGraph/Public/EvaluationVM/Tasks/PushReferenceKeyframe.h b/Engine/Plugins/Experimental/AnimNextAnimGraph/Source/AnimNextAnimGraph/Public/EvaluationVM/Tasks/PushReferenceKeyframe.h
--- a/Engine/Plugins/Experimental/AnimNextAnimGraph/Source/AnimNextAnimGraph/Public/EvaluationVM/Tasks/PushReferenceKeyframe.h
+++ b/Engine/Plugins/Experimental/AnimNextAnimGraph/Source/AnimNextAnimGraph/Public/EvaluationVM/Tasks/PushReferenceKeyframe.h
@@ -28,6 +28,6 @@ struct ANIMNEXTANIMGRAPH_API FAnimNextPushReferenceKeyframeTask : public FAnimNe
 	virtual void Execute(UE::AnimNext::FEvaluationVM& VM) const override;
 
 	// Whether or not the reference pose comes from the skeleton or is the additive identity
-	UPROPERTY()
+	UPROPERTY(VisibleAnywhere, Category = Properties)
 	bool bIsAdditive = false;
 };
diff --git a/Engine/Plugins/Experimental/AnimNextAnimGraph/Source/AnimNextAnimGraphEditor/AnimNextAnimGraphEditor.Build.cs b/Engine/Plugins/Experimental/AnimNextAnimGraph/Source/AnimNextAnimGraphEditor/AnimNextAnimGraphEditor.Build.cs
--- a/Engine/Plugins/Experimental/AnimNextAnimGraph/Source/AnimNextAnimGraphEditor/AnimNextAnimGraphEditor.Build.cs
+++ b/Engine/Plugins/Experimental/AnimNextAnimGraph/Source/AnimNextAnimGraphEditor/AnimNextAnimGraphEditor.Build.cs
@@ -38,7 +38,13 @@ namespace UnrealBuildTool.Rules
 					"RigVMEditor",
 					"GraphEditor",
 					"Persona",
-					"PropertyEditor"
+					"PropertyEditor",
+					"TraceAnalysis",
+					"TraceLog",
+					"TraceServices",
+					"TraceInsights",
+					"RewindDebuggerInterface",
+					"GameplayInsights",
 				}
 			);
 		}
diff --git a/Engine/Plugins/Experimental/AnimNextAnimGraph/Source/AnimNextAnimGraphEditor/Private/AnimNextAnimGraphEditorModule.cpp b/Engine/Plugins/Experimental/AnimNextAnimGraph/Source/AnimNextAnimGraphEditor/Private/AnimNextAnimGraphEditorModule.cpp
--- a/Engine/Plugins/Experimental/AnimNextAnimGraph/Source/AnimNextAnimGraphEditor/Private/AnimNextAnimGraphEditorModule.cpp
+++ b/Engine/Plugins/Experimental/AnimNextAnimGraph/Source/AnimNextAnimGraphEditor/Private/AnimNextAnimGraphEditorModule.cpp
@@ -37,10 +37,19 @@
 #include "Graph/AnimNextGraphDetails.h"
 #include "Graph/PostProcessAnimationCustomization.h"
 #include "Module/AnimNextEventGraphSchema.h"
+#include "RewindDebugger/AnimNextAnimGraphTraceModule.h"
+#include "RewindDebugger/EvaluationProgramTrack.h"
 #include "PersonaModule.h"
+#include "Developer/TraceServices/Public/TraceServices/ModuleService.h"
 
 #define LOCTEXT_NAMESPACE "FAnimNextAnimGraphEditorModule"
 
+namespace
+{
+	FAnimNextAnimGraphTraceModule GAnimNextAnimGraphTraceModule;
+	UE::AnimNextEditor::FEvaluationProgramTrackCreator GAnimNextModulesTrackCreator;
+}
+
 namespace UE::AnimNext::Editor
 {
 
@@ -64,6 +73,9 @@ class FAnimNextAnimGraphEditorModule : public IModuleInterface
 		PropertyModule.RegisterCustomPropertyTypeLayout(
 					FAnimNextAnimGraph::StaticStruct()->GetFName(),
 					FOnGetPropertyTypeCustomizationInstance::CreateLambda([] { return MakeShared<FAnimNextGraphDetails>(); }));
+					
+		IModularFeatures::Get().RegisterModularFeature(TraceServices::ModuleFeatureName, &GAnimNextAnimGraphTraceModule);
+		IModularFeatures::Get().RegisterModularFeature(RewindDebugger::IRewindDebuggerTrackCreator::ModularFeatureName, &GAnimNextModulesTrackCreator);
 
 		Workspace::IWorkspaceEditorModule& WorkspaceEditorModule = FModuleManager::Get().LoadModuleChecked<Workspace::IWorkspaceEditorModule>("WorkspaceEditor");
 
@@ -127,6 +139,9 @@ class FAnimNextAnimGraphEditorModule : public IModuleInterface
 		{
 			SettingsModule->UnregisterSettings("Project", "Plugins", "AnimNextAnimGraph");
 		}
+		
+		IModularFeatures::Get().UnregisterModularFeature(TraceServices::ModuleFeatureName, &GAnimNextAnimGraphTraceModule);
+		IModularFeatures::Get().UnregisterModularFeature(RewindDebugger::IRewindDebuggerTrackCreator::ModularFeatureName, &GAnimNextModulesTrackCreator);
 
 		if(FModuleManager::Get().IsModuleLoaded("WorkspaceEditor"))
 		{
diff --git a/Engine/Plugins/Experimental/AnimNextAnimGraph/Source/AnimNextAnimGraphEditor/Private/RewindDebugger/AnimNextAnimGraphAnalyzer.cpp b/Engine/Plugins/Experimental/AnimNextAnimGraph/Source/AnimNextAnimGraphEditor/Private/RewindDebugger/AnimNextAnimGraphAnalyzer.cpp
new file mode 100644
--- /dev/null
+++ b/Engine/Plugins/Experimental/AnimNextAnimGraph/Source/AnimNextAnimGraphEditor/Private/RewindDebugger/AnimNextAnimGraphAnalyzer.cpp
@@ -0,0 +1,50 @@
+// Copyright Epic Games, Inc. All Rights Reserved.
+
+#include "AnimNextAnimGraphAnalyzer.h"
+
+#include "HAL/LowLevelMemTracker.h"
+#include "Serialization/MemoryReader.h"
+#include "TraceServices/Model/AnalysisSession.h"
+#include "AnimNextAnimGraphProvider.h"
+#include "TraceServices/Utils.h"
+#include "StructUtils/PropertyBag.h"
+#include "Serialization/ObjectReader.h"
+
+FAnimNextAnimGraphAnalyzer::FAnimNextAnimGraphAnalyzer(TraceServices::IAnalysisSession& InSession, FAnimNextAnimGraphProvider& InProvider)
+	: Session(InSession)
+	, Provider(InProvider)
+{
+}
+
+void FAnimNextAnimGraphAnalyzer::OnAnalysisBegin(const FOnAnalysisContext& Context)
+{
+	auto& Builder = Context.InterfaceBuilder;
+
+	Builder.RouteEvent(RouteId_EvaluationProgram, "AnimNextAnimGraph", "EvaluationProgram");
+}
+
+bool FAnimNextAnimGraphAnalyzer::OnEvent(uint16 RouteId, EStyle Style, const FOnEventContext& Context)
+{
+	LLM_SCOPE_BYNAME(TEXT("Insights/FAnimNextAnimGraphAnalyzer"));
+
+	TraceServices::FAnalysisSessionEditScope _(Session);
+
+	const auto& EventData = Context.EventData;
+	switch (RouteId)
+	{
+		case RouteId_EvaluationProgram:
+		{
+			uint64 OuterObjectId = EventData.GetValue<uint64>("OuterObjectId");
+			uint64 GraphInstanceId = EventData.GetValue<uint64>("InstanceId");
+			uint64 Cycle = EventData.GetValue<uint64>("Cycle");
+			double RecordingTime = EventData.GetValue<double>("RecordingTime");
+
+			TArrayView<const uint8> ProgramData = EventData.GetArrayView<uint8>("ProgramData");
+
+			Provider.AppendEvaluationProgram(Context.EventTime.AsSeconds(Cycle), RecordingTime, OuterObjectId, GraphInstanceId, ProgramData);
+			break;
+		}
+	}
+
+	return true;
+}
diff --git a/Engine/Plugins/Experimental/AnimNextAnimGraph/Source/AnimNextAnimGraphEditor/Private/RewindDebugger/AnimNextAnimGraphAnalyzer.h b/Engine/Plugins/Experimental/AnimNextAnimGraph/Source/AnimNextAnimGraphEditor/Private/RewindDebugger/AnimNextAnimGraphAnalyzer.h
new file mode 100644
--- /dev/null
+++ b/Engine/Plugins/Experimental/AnimNextAnimGraph/Source/AnimNextAnimGraphEditor/Private/RewindDebugger/AnimNextAnimGraphAnalyzer.h
@@ -0,0 +1,27 @@
+// Copyright Epic Games, Inc. All Rights Reserved.
+
+#pragma once
+
+#include "Trace/Analyzer.h"
+
+class FAnimNextAnimGraphProvider;
+namespace TraceServices { class IAnalysisSession; }
+
+class FAnimNextAnimGraphAnalyzer : public UE::Trace::IAnalyzer
+{
+public:
+	FAnimNextAnimGraphAnalyzer(TraceServices::IAnalysisSession& InSession, FAnimNextAnimGraphProvider& InProvider);
+
+	virtual void OnAnalysisBegin(const FOnAnalysisContext& Context) override;
+	virtual void OnAnalysisEnd() override {}
+	virtual bool OnEvent(uint16 RouteId, EStyle Style, const FOnEventContext& Context) override;
+
+private:
+	enum : uint16
+	{
+		RouteId_EvaluationProgram,
+	};
+
+	TraceServices::IAnalysisSession& Session;
+	FAnimNextAnimGraphProvider& Provider;
+};
diff --git a/Engine/Plugins/Experimental/AnimNextAnimGraph/Source/AnimNextAnimGraphEditor/Private/RewindDebugger/AnimNextAnimGraphProvider.cpp b/Engine/Plugins/Experimental/AnimNextAnimGraph/Source/AnimNextAnimGraphEditor/Private/RewindDebugger/AnimNextAnimGraphProvider.cpp
new file mode 100644
--- /dev/null
+++ b/Engine/Plugins/Experimental/AnimNextAnimGraph/Source/AnimNextAnimGraphEditor/Private/RewindDebugger/AnimNextAnimGraphProvider.cpp
@@ -0,0 +1,62 @@
+// Copyright Epic Games, Inc. All Rights Reserved.
+
+#include "AnimNextAnimGraphProvider.h"
+#include "ObjectTrace.h"
+
+FName FAnimNextAnimGraphProvider::ProviderName("AnimNextAnimGraphProvider");
+
+#define LOCTEXT_NAMESPACE "AnimNextAnimGraphProvider"
+
+FAnimNextAnimGraphProvider::FAnimNextAnimGraphProvider(TraceServices::IAnalysisSession& InSession)
+	: Session(InSession)
+{
+}
+
+void FAnimNextAnimGraphProvider::AppendEvaluationProgram(double ProfileTime, double RecordingTime, uint64 OuterObjectId, uint64 GraphInstanceId, const TArrayView<const uint8>& ProgramData)
+{
+	Session.WriteAccessCheck();
+
+	EvaluationGraphs.Add(OuterObjectId, GraphInstanceId);
+	
+	const TSharedRef<FEvaluationProgramData>* Data = EvaluationProgramData.Find(GraphInstanceId);
+	if (Data == nullptr)
+	{
+		TSharedRef<FEvaluationProgramData> NewData = MakeShared<FEvaluationProgramData>(Session);
+       	NewData->GraphInstanceId = GraphInstanceId;
+		EvaluationProgramData.Add(GraphInstanceId, NewData);
+		Data = EvaluationProgramData.Find(GraphInstanceId);
+	}
+	
+	if (!ProgramData.IsEmpty())
+	{
+		(*Data)->EvaluationProgramTimeline.AppendEvent(ProfileTime, TArray<uint8>(ProgramData));
+	}
+}
+
+const FEvaluationProgramData* FAnimNextAnimGraphProvider::GetEvaluationProgramData(uint64 GraphInstanceId) const
+{
+	Session.ReadAccessCheck();
+	
+	const TSharedRef<FEvaluationProgramData>* Data = EvaluationProgramData.Find(GraphInstanceId);
+	if (Data)
+	{
+		return &**Data;
+	}
+	return nullptr;
+}
+
+
+void FAnimNextAnimGraphProvider::EnumerateEvaluationGraphs(uint64 OuterObjectId, TFunctionRef<void(uint64 GraphInstanceId)> Callback) const
+{
+	Session.ReadAccessCheck();
+
+	TArray<uint64> GraphIds;
+	EvaluationGraphs.MultiFind(OuterObjectId, GraphIds);
+	
+	for (uint64 GraphId : GraphIds)
+	{
+		Callback(GraphId);
+	}
+}
+
+#undef LOCTEXT_NAMESPACE
diff --git a/Engine/Plugins/Experimental/AnimNextAnimGraph/Source/AnimNextAnimGraphEditor/Private/RewindDebugger/AnimNextAnimGraphProvider.h b/Engine/Plugins/Experimental/AnimNextAnimGraph/Source/AnimNextAnimGraphEditor/Private/RewindDebugger/AnimNextAnimGraphProvider.h
new file mode 100644
--- /dev/null
+++ b/Engine/Plugins/Experimental/AnimNextAnimGraph/Source/AnimNextAnimGraphEditor/Private/RewindDebugger/AnimNextAnimGraphProvider.h
@@ -0,0 +1,46 @@
+// Copyright Epic Games, Inc. All Rights Reserved.
+
+#pragma once
+
+#include "Engine/EngineBaseTypes.h"
+#include "Model/PointTimeline.h"
+#include "Containers/Map.h"
+#include "Templates/SharedPointer.h"
+#include "Model/IntervalTimeline.h"
+#include "TraceServices/Model/AnalysisSession.h"
+
+namespace TraceServices { class IAnalysisSession; }
+
+struct FEvaluationProgramData
+{
+	FEvaluationProgramData(TraceServices::IAnalysisSession& Session)
+		: GraphInstanceId(0),
+		EvaluationProgramTimeline(Session.GetLinearAllocator( ))
+			
+	{
+	}
+
+	uint64 GraphInstanceId;
+	TraceServices::TPointTimeline<TArray<uint8>> EvaluationProgramTimeline;
+};
+
+class FAnimNextAnimGraphProvider : public TraceServices::IProvider
+{
+public:
+	static FName ProviderName;
+
+	FAnimNextAnimGraphProvider(TraceServices::IAnalysisSession& InSession);
+
+	const FEvaluationProgramData* GetEvaluationProgramData(uint64 GraphInstanceId) const;
+	void AppendEvaluationProgram(double ProfileTime, double RecordingTime, uint64 ObjectInstanceId, uint64 GraphInstanceId, const TArrayView<const uint8>& ProgramData);
+	
+	void EnumerateEvaluationGraphs(uint64 OuterObjectId, TFunctionRef<void(uint64 GraphInstanceId)> Callback) const;
+	
+private:
+	TraceServices::IAnalysisSession& Session;
+
+	TMap<uint64, TSharedRef<FEvaluationProgramData>> EvaluationProgramData;
+
+	// map from object id to a list of graph instance ids that have evaluation data
+	TMultiMap<uint64, uint64> EvaluationGraphs;
+};
diff --git a/Engine/Plugins/Experimental/AnimNextAnimGraph/Source/AnimNextAnimGraphEditor/Private/RewindDebugger/AnimNextAnimGraphTraceModule.cpp b/Engine/Plugins/Experimental/AnimNextAnimGraph/Source/AnimNextAnimGraphEditor/Private/RewindDebugger/AnimNextAnimGraphTraceModule.cpp
new file mode 100644
--- /dev/null
+++ b/Engine/Plugins/Experimental/AnimNextAnimGraph/Source/AnimNextAnimGraphEditor/Private/RewindDebugger/AnimNextAnimGraphTraceModule.cpp
@@ -0,0 +1,32 @@
+// Copyright Epic Games, Inc. All Rights Reserved.
+
+#include "AnimNextAnimGraphTraceModule.h"
+#include "AnimNextAnimGraphProvider.h"
+#include "AnimNextAnimGraphAnalyzer.h"
+
+FName FAnimNextAnimGraphTraceModule::ModuleName("AnimNextAnimGraph");
+
+void FAnimNextAnimGraphTraceModule::GetModuleInfo(TraceServices::FModuleInfo& OutModuleInfo)
+{
+	OutModuleInfo.Name = ModuleName;
+	OutModuleInfo.DisplayName = TEXT("AnimNextAnimGraph");
+}
+
+void FAnimNextAnimGraphTraceModule::OnAnalysisBegin(TraceServices::IAnalysisSession& InSession)
+{
+	TSharedPtr<FAnimNextAnimGraphProvider> AnimNextAnimGraphProvider = MakeShared<FAnimNextAnimGraphProvider>(InSession);
+	InSession.AddProvider(FAnimNextAnimGraphProvider::ProviderName, AnimNextAnimGraphProvider);
+
+	InSession.AddAnalyzer(new FAnimNextAnimGraphAnalyzer(InSession, *AnimNextAnimGraphProvider));
+}
+
+void FAnimNextAnimGraphTraceModule::GetLoggers(TArray<const TCHAR *>& OutLoggers)
+{
+	OutLoggers.Add(TEXT("AnimNextAnimGraph"));
+}
+
+void FAnimNextAnimGraphTraceModule::GenerateReports(const TraceServices::IAnalysisSession& Session, const TCHAR* CmdLine, const TCHAR* OutputDirectory)
+{
+
+}
+
diff --git a/Engine/Plugins/Experimental/AnimNextAnimGraph/Source/AnimNextAnimGraphEditor/Private/RewindDebugger/AnimNextAnimGraphTraceModule.h b/Engine/Plugins/Experimental/AnimNextAnimGraph/Source/AnimNextAnimGraphEditor/Private/RewindDebugger/AnimNextAnimGraphTraceModule.h
new file mode 100644
--- /dev/null
+++ b/Engine/Plugins/Experimental/AnimNextAnimGraph/Source/AnimNextAnimGraphEditor/Private/RewindDebugger/AnimNextAnimGraphTraceModule.h
@@ -0,0 +1,23 @@
+// Copyright Epic Games, Inc. All Rights Reserved.
+
+#pragma once
+
+#include "CoreMinimal.h"
+#include "TraceServices/ModuleService.h"
+
+class IAnimationProvider;
+
+class FAnimNextAnimGraphTraceModule : public TraceServices::IModule
+{
+public:
+	// TraceServices::IModule interface
+	virtual void GetModuleInfo(TraceServices::FModuleInfo& OutModuleInfo) override;
+	virtual void OnAnalysisBegin(TraceServices::IAnalysisSession& Session) override;
+	virtual void GetLoggers(TArray<const TCHAR *>& OutLoggers) override;
+	virtual void GenerateReports(const TraceServices::IAnalysisSession& Session, const TCHAR* CmdLine, const TCHAR* OutputDirectory) override;
+	virtual const TCHAR* GetCommandLineArgument() override { return TEXT("chooser"); }
+
+private:
+	static FName ModuleName;
+};
+
diff --git a/Engine/Plugins/Experimental/AnimNextAnimGraph/Source/AnimNextAnimGraphEditor/Private/RewindDebugger/EvaluationProgramTrack.cpp b/Engine/Plugins/Experimental/AnimNextAnimGraph/Source/AnimNextAnimGraphEditor/Private/RewindDebugger/EvaluationProgramTrack.cpp
new file mode 100644
--- /dev/null
+++ b/Engine/Plugins/Experimental/AnimNextAnimGraph/Source/AnimNextAnimGraphEditor/Private/RewindDebugger/EvaluationProgramTrack.cpp
@@ -0,0 +1,202 @@
+// Copyright Epic Games, Inc. All Rights Reserved.
+
+#include "EvaluationProgramTrack.h"
+
+#include "AnimNextAnimGraphProvider.h"
+#include "Editor.h"
+#include "IGameplayProvider.h"
+#include "IRewindDebugger.h"
+#include "Subsystems/AssetEditorSubsystem.h"
+#include "Modules/ModuleManager.h"
+#include "PropertyEditorModule.h"
+#include "Serialization/ObjectAndNameAsStringProxyArchive.h"
+#include "Serialization/ObjectReader.h"
+
+
+#define LOCTEXT_NAMESPACE "EvaluationProgramTrack"
+
+namespace UE::AnimNextEditor
+{
+
+static const FName AnimNextModulesName("AnimNextModules");
+
+ FName FEvaluationProgramTrackCreator::GetTargetTypeNameInternal() const
+ {
+ 	static const FName ObjectName("AnimNextComponent");
+	return ObjectName;
+}
+
+	
+FText FEvaluationProgramTrack::GetDisplayNameInternal() const
+{
+ 	return NSLOCTEXT("RewindDebugger", "EvaluationProgramTrackName", "EvaluationProgram");
+}
+
+void FEvaluationProgramTrackCreator::GetTrackTypesInternal(TArray<RewindDebugger::FRewindDebuggerTrackType>& Types) const
+{
+	Types.Add({AnimNextModulesName, LOCTEXT("AnimNextModule", "AnimNextModules")});
+}
+
+TSharedPtr<RewindDebugger::FRewindDebuggerTrack> FEvaluationProgramTrackCreator::CreateTrackInternal(uint64 ObjectId) const
+{
+	return MakeShared<FEvaluationProgramTrack>(ObjectId);
+}
+
+		
+FEvaluationProgramTrack::FEvaluationProgramTrack(uint64 InObjectId) :
+	ObjectId(InObjectId)
+{
+	Initialize();
+}
+	
+FEvaluationProgramTrack::FEvaluationProgramTrack(uint64 InObjectId, uint64 InInstanceId) :
+   	ObjectId(InObjectId),
+	InstanceId(InInstanceId)
+{
+   	Initialize();
+}
+	
+	
+FEvaluationProgramTrack::~FEvaluationProgramTrack()
+{
+ 	if (UEvaluationProgramDetailsObject* DetailsObject = DetailsObjectWeakPtr.Get())
+   	{
+   		DetailsObject->ClearFlags(RF_Standalone);
+   	}
+}
+	
+void FEvaluationProgramTrack::Initialize()
+{
+	ExistenceRange = MakeShared<SEventTimelineView::FTimelineEventData>();
+	ExistenceRange->Windows.Add({0,0, GetDisplayNameInternal(), GetDisplayNameInternal(), FLinearColor(0.1f,0.15f,0.11f)});
+	
+   	FPropertyEditorModule& PropertyEditorModule = FModuleManager::GetModuleChecked<FPropertyEditorModule>("PropertyEditor");
+   	FDetailsViewArgs DetailsViewArgs;
+   	DetailsViewArgs.NameAreaSettings = FDetailsViewArgs::HideNameArea;
+   	DetailsView = PropertyEditorModule.CreateDetailView(DetailsViewArgs);
+
+	InitializeDetailsObject();
+}
+	
+UEvaluationProgramDetailsObject* FEvaluationProgramTrack::InitializeDetailsObject()
+{
+	UEvaluationProgramDetailsObject* DetailsObject = NewObject<UEvaluationProgramDetailsObject>();
+	DetailsObject->SetFlags(RF_Standalone);
+	DetailsObjectWeakPtr = MakeWeakObjectPtr(DetailsObject);
+	DetailsView->SetObject(DetailsObject);
+ 	return DetailsObject;
+}
+
+bool FEvaluationProgramTrack::UpdateInternal()
+{
+	TRACE_CPUPROFILER_EVENT_SCOPE(FEvaluationProgramTrack::UpdateInternal);
+ 	
+	IRewindDebugger* RewindDebugger = IRewindDebugger::Instance();
+ 	const TRange<double> ViewRange = RewindDebugger->GetCurrentViewRange();
+	
+	const TraceServices::IAnalysisSession* AnalysisSession = RewindDebugger->GetAnalysisSession();
+
+ 	bool bChanged = false;
+	
+	if (const FAnimNextAnimGraphProvider* AnimNextAnimGraphProvider = IRewindDebugger::Instance()->GetAnalysisSession()->ReadProvider<FAnimNextAnimGraphProvider>("AnimNextAnimGraphProvider"))
+	{
+		double CurrentScrubTime = IRewindDebugger::Instance()->CurrentTraceTime();
+
+		UEvaluationProgramDetailsObject* DetailsObject = DetailsObjectWeakPtr.Get();
+		if (DetailsObject == nullptr)
+		{
+			// this should not happen unless the object was garbage collected (which should not happen since it's marked as Standalone)
+			Initialize();
+		}
+
+		if (InstanceId == 0)
+		{
+		 	AnimNextAnimGraphProvider->EnumerateEvaluationGraphs(ObjectId, [this](uint64 GraphId)
+		 	{
+		 		InstanceId = GraphId;
+		 	});
+		}
+		
+		if(InstanceId != 0)
+		{
+			if (const FEvaluationProgramData* Data = AnimNextAnimGraphProvider->GetEvaluationProgramData(InstanceId))
+			{
+				if (PreviousScrubTime != CurrentScrubTime)
+				{
+					PreviousScrubTime = CurrentScrubTime;
+					
+					const TraceServices::IFrameProvider& FramesProvider = TraceServices::ReadFrameProvider(*AnalysisSession);
+					TraceServices::FAnalysisSessionReadScope SessionReadScope(*AnalysisSession);
+
+					TraceServices::FFrame MarkerFrame;
+					if(FramesProvider.GetFrameFromTime(ETraceFrameType::TraceFrameType_Game, CurrentScrubTime, MarkerFrame))
+					{
+						Data->EvaluationProgramTimeline.EnumerateEvents(MarkerFrame.StartTime, MarkerFrame.EndTime, [DetailsObject](double InStartTime, double InEndTime, uint32 InDepth, const TArray<uint8>& VariableData)
+						{				
+							FMemoryReader Reader(VariableData);
+							FObjectAndNameAsStringProxyArchive Archive(Reader, /*bInLoadIfFindFails*/true);
+
+							static const FSerializableEvaluationProgram Defaults;
+							FSerializableEvaluationProgram::StaticStruct()->SerializeItem(Archive, &DetailsObject->Program, &Defaults);
+								
+							return TraceServices::EEventEnumerate::Stop;
+						});
+					}
+				}
+			}
+
+			// TArray<uint64, TInlineAllocator<32>> CurrentChildren;
+		
+			// // update/create child tracks
+			// AnimNextProvider->EnumerateChildInstances(InstanceId, [this, &bChanged, &CurrentChildren, &ViewRange](const FDataInterfaceData& ChildData)
+			// {
+			// 	if (ChildData.StartTime < ViewRange.GetUpperBoundValue() && ChildData.EndTime > ViewRange.GetLowerBoundValue())
+			// 	{
+			// 		CurrentChildren.Add(ChildData.InstanceId);
+			// 		if (!Children.ContainsByPredicate([&ChildData](TSharedPtr<FEvaluationProgramTrack>& Child) { return Child->InstanceId == ChildData.InstanceId; } ))
+			// 		{
+			// 			Children.Add(MakeShared<FEvaluationProgramTrack>(ObjectId, ChildData.InstanceId));
+			// 			bChanged = true;
+			// 		}
+			// 	}
+			// });
+			//
+			// int32 NumRemoved = Children.RemoveAll([&CurrentChildren](TSharedPtr<FEvaluationProgramTrack>& Child)
+			// {
+			// 	return !CurrentChildren.Contains(Child->InstanceId);
+			// });
+			// bChanged |= (NumRemoved > 0);
+		}
+
+		// for (auto& Child : Children)
+		// {
+		// 	bChanged |= Child->Update();
+		// }
+	}
+	
+ 	return bChanged;
+}
+
+bool FEvaluationProgramTrackCreator::HasDebugInfoInternal(uint64 ObjectId) const
+{
+	// TRACE_CPUPROFILER_EVENT_SCOPE(FEvaluationProgramTrack::HasDebugInfoInternal);
+
+	const IRewindDebugger* RewindDebugger = IRewindDebugger::Instance();
+	if (const TraceServices::IAnalysisSession* AnalysisSession = RewindDebugger->GetAnalysisSession())
+	{
+		TraceServices::FAnalysisSessionReadScope SessionReadScope(*AnalysisSession);
+    
+		const FAnimNextAnimGraphProvider* AnimNextAnimGraphProvider = IRewindDebugger::Instance()->GetAnalysisSession()->ReadProvider<FAnimNextAnimGraphProvider>("AnimNextAnimGraphProvider");
+    
+		// uint64 ModuleId;
+		// return AnimNextProvider->GetModuleId(ObjectId, ModuleId);
+		return true;
+	}
+
+	return false;
+}
+
+
+}
+
+#undef LOCTEXT_NAMESPACE
\ No newline at end of file
diff --git a/Engine/Plugins/Experimental/AnimNextAnimGraph/Source/AnimNextAnimGraphEditor/Private/RewindDebugger/EvaluationProgramTrack.h b/Engine/Plugins/Experimental/AnimNextAnimGraph/Source/AnimNextAnimGraphEditor/Private/RewindDebugger/EvaluationProgramTrack.h
new file mode 100644
--- /dev/null
+++ b/Engine/Plugins/Experimental/AnimNextAnimGraph/Source/AnimNextAnimGraphEditor/Private/RewindDebugger/EvaluationProgramTrack.h
@@ -0,0 +1,69 @@
+// Copyright Epic Games, Inc. All Rights Reserved.
+#pragma once
+
+#include "IRewindDebuggerTrackCreator.h"
+#include "RewindDebuggerTrack.h"
+#include "EvaluationVM/SerializableEvaluationProgram.h"
+#include "Textures/SlateIcon.h"
+#include "SEventTimelineView.h"
+#include "IDetailsView.h"
+#include "SEventTimelineView.h"
+
+#include "EvaluationProgramTrack.generated.h"
+
+UCLASS()
+class UEvaluationProgramDetailsObject : public UObject
+{
+	GENERATED_BODY()
+
+public:
+	UPROPERTY(VisibleAnywhere, Category=Details, meta=(ShowOnlyInnerProperties))
+	FSerializableEvaluationProgram Program;
+};
+
+namespace UE::AnimNextEditor
+{
+	
+class FEvaluationProgramTrack : public RewindDebugger::FRewindDebuggerTrack
+{
+public:
+	FEvaluationProgramTrack(uint64 InObjectId);
+	FEvaluationProgramTrack(uint64 InObjectId, uint64 InstanceId);
+	virtual ~FEvaluationProgramTrack();
+	
+	TSharedPtr<SEventTimelineView::FTimelineEventData> GetExistenceRange() const { return ExistenceRange; }
+private:
+	virtual bool UpdateInternal() override;
+	virtual TSharedPtr<SWidget> GetDetailsViewInternal() override { return DetailsView; }
+
+	virtual FSlateIcon GetIconInternal() override { return Icon; }
+	virtual FName GetNameInternal() const override { return "AnimNextModule"; }
+	virtual FText GetDisplayNameInternal() const override;
+	virtual uint64 GetObjectIdInternal() const override { return ObjectId; }
+
+	void Initialize();
+	UEvaluationProgramDetailsObject* InitializeDetailsObject();
+
+	TSharedPtr<IDetailsView> DetailsView;
+	FSlateIcon Icon;
+	uint64 ObjectId;
+	uint64 InstanceId = 0; 
+	double PreviousScrubTime = -1.0;
+	TWeakObjectPtr<UEvaluationProgramDetailsObject> DetailsObjectWeakPtr;
+	TSharedPtr<SEventTimelineView::FTimelineEventData> ExistenceRange;
+
+	// TArray<TSharedPtr<FEvaluationProgramTrack>> Children;
+};
+
+
+class FEvaluationProgramTrackCreator : public RewindDebugger::IRewindDebuggerTrackCreator
+{
+private:
+	virtual FName GetTargetTypeNameInternal() const;
+	virtual FName GetNameInternal() const override { return "AnimNextModule"; }
+	virtual void GetTrackTypesInternal(TArray<RewindDebugger::FRewindDebuggerTrackType>& Types) const override;
+	virtual TSharedPtr<RewindDebugger::FRewindDebuggerTrack> CreateTrackInternal(uint64 ObjectId) const override;
+	virtual bool HasDebugInfoInternal(uint64 ObjectId) const override;
+};
+
+}
\ No newline at end of file
