AnimNext: Fix for serialization weak object ptrs

diff --git a/Engine/Plugins/Experimental/AnimNextAnimGraph/Source/AnimNextAnimGraph/Internal/TraitCore/TraitReader.h b/Engine/Plugins/Experimental/AnimNextAnimGraph/Source/AnimNextAnimGraph/Internal/TraitCore/TraitReader.h
--- a/Engine/Plugins/Experimental/AnimNextAnimGraph/Source/AnimNextAnimGraph/Internal/TraitCore/TraitReader.h
+++ b/Engine/Plugins/Experimental/AnimNextAnimGraph/Source/AnimNextAnimGraph/Internal/TraitCore/TraitReader.h
@@ -56,6 +56,7 @@ namespace UE::AnimNext
 		virtual FArchive& operator<<(FObjectPtr& Obj) override;
 		virtual FArchive& operator<<(FSoftObjectPath& Value) override;
 		virtual FArchive& operator<<(FSoftObjectPtr& Value) override;
+		virtual FArchive& operator<<(FWeakObjectPtr& Value) override;
 
 	private:
 		// Call first to read the graph shared data
diff --git a/Engine/Plugins/Experimental/AnimNextAnimGraph/Source/AnimNextAnimGraph/Internal/TraitCore/TraitWriter.h b/Engine/Plugins/Experimental/AnimNextAnimGraph/Source/AnimNextAnimGraph/Internal/TraitCore/TraitWriter.h
--- a/Engine/Plugins/Experimental/AnimNextAnimGraph/Source/AnimNextAnimGraph/Internal/TraitCore/TraitWriter.h
+++ b/Engine/Plugins/Experimental/AnimNextAnimGraph/Source/AnimNextAnimGraph/Internal/TraitCore/TraitWriter.h
@@ -72,7 +72,8 @@ namespace UE::AnimNext
 		virtual FArchive& operator<<(FObjectPtr& Obj) override;
 		virtual FArchive& operator<<(FSoftObjectPath& Value) override;
 		virtual FArchive& operator<<(FSoftObjectPtr& Value) override;
-		
+		virtual FArchive& operator<<(FWeakObjectPtr& Value) override;
+
 	private:
 		struct FNodeMapping
 		{
diff --git a/Engine/Plugins/Experimental/AnimNextAnimGraph/Source/AnimNextAnimGraph/Private/TraitCore/TraitReader.cpp b/Engine/Plugins/Experimental/AnimNextAnimGraph/Source/AnimNextAnimGraph/Private/TraitCore/TraitReader.cpp
--- a/Engine/Plugins/Experimental/AnimNextAnimGraph/Source/AnimNextAnimGraph/Private/TraitCore/TraitReader.cpp
+++ b/Engine/Plugins/Experimental/AnimNextAnimGraph/Source/AnimNextAnimGraph/Private/TraitCore/TraitReader.cpp
@@ -231,6 +231,11 @@ namespace UE::AnimNext
 		return FArchiveUObject::SerializeObjectPtr(*this, Obj);
 	}
 
+	FArchive& FTraitReader::operator<<(FWeakObjectPtr& Value)
+	{
+		return FArchiveUObject::SerializeWeakObjectPtr(*this, Value);
+	}
+
 	FArchive& FTraitReader::operator<<(FSoftObjectPath& Value)
 	{
 		// Load the object index
diff --git a/Engine/Plugins/Experimental/AnimNextAnimGraph/Source/AnimNextAnimGraph/Private/TraitCore/TraitWriter.cpp b/Engine/Plugins/Experimental/AnimNextAnimGraph/Source/AnimNextAnimGraph/Private/TraitCore/TraitWriter.cpp
--- a/Engine/Plugins/Experimental/AnimNextAnimGraph/Source/AnimNextAnimGraph/Private/TraitCore/TraitWriter.cpp
+++ b/Engine/Plugins/Experimental/AnimNextAnimGraph/Source/AnimNextAnimGraph/Private/TraitCore/TraitWriter.cpp
@@ -253,6 +253,11 @@ namespace UE::AnimNext
 		return FArchiveUObject::SerializeObjectPtr(*this, Obj);
 	}
 
+	FArchive& FTraitWriter::operator<<(FWeakObjectPtr& Value)
+	{
+		return FArchiveUObject::SerializeWeakObjectPtr(*this, Value);
+	}
+
 	FArchive& FTraitWriter::operator<<(FSoftObjectPath& Value)
 	{
 		// Add our object for tracking
