Fix new RunAnimation V2 node not exporting graph asset as a workspace outliner reference

diff --git a/Engine/Plugins/Experimental/AnimNextAnimGraph/Source/AnimNextAnimGraph/Internal/Graph/RigUnit_AnimNextRunAnimationGraph_v2.h b/Engine/Plugins/Experimental/AnimNextAnimGraph/Source/AnimNextAnimGraph/Internal/Graph/RigUnit_AnimNextRunAnimationGraph_v2.h
--- a/Engine/Plugins/Experimental/AnimNextAnimGraph/Source/AnimNextAnimGraph/Internal/Graph/RigUnit_AnimNextRunAnimationGraph_v2.h
+++ b/Engine/Plugins/Experimental/AnimNextAnimGraph/Source/AnimNextAnimGraph/Internal/Graph/RigUnit_AnimNextRunAnimationGraph_v2.h
@@ -34,7 +34,7 @@ struct ANIMNEXTANIMGRAPH_API FRigUnit_AnimNextRunAnimationGraph_v2 : public FRig
 	static const UAnimNextAnimationGraph* GetHostGraphToRun(FAnimNextExecuteContext& InExecuteContext);
 
 	// The graph to run
-	UPROPERTY(EditAnywhere, Category = "Run", meta = (Input, ShowOnlyInnerProperties))
+	UPROPERTY(EditAnywhere, Category = "Run", meta = (Input, ShowOnlyInnerProperties, ExportAsReference="true"))
 	FAnimNextAnimGraph Graph;
 
 	// LOD to run the graph at
