Improve sync group support

diff --git a/Engine/Plugins/Experimental/AnimNextAnimGraph/Source/AnimNextAnimGraph/Private/Graph/SyncGroup_GraphInstanceComponent.cpp b/Engine/Plugins/Experimental/AnimNextAnimGraph/Source/AnimNextAnimGraph/Private/Graph/SyncGroup_GraphInstanceComponent.cpp
--- a/Engine/Plugins/Experimental/AnimNextAnimGraph/Source/AnimNextAnimGraph/Private/Graph/SyncGroup_GraphInstanceComponent.cpp
+++ b/Engine/Plugins/Experimental/AnimNextAnimGraph/Source/AnimNextAnimGraph/Private/Graph/SyncGroup_GraphInstanceComponent.cpp
@@ -7,19 +7,39 @@
 #include "Animation/AnimationAsset.h"
 #include "AnimationUtils.h"
 
-#define UE_LOG_ENABLE_MARKER_DEBUG 0
+// Enabled by default in development builds
+#define UE_DEBUG_SYNC_GROUPS !(UE_BUILD_SHIPPING || UE_BUILD_TEST)
 
 namespace UE::AnimNext
 {
 	namespace Private
 	{
+#if UE_DEBUG_SYNC_GROUPS
+		// Enables/disable debug logging in the console
+		static TAutoConsoleVariable<bool> CDebugVarEnableSyncLog(TEXT("a.AnimNext.EnableSyncLog"), false, TEXT("Toggles debug sync group logging"));
+
+		// Enables sync mode debugging
+		// 0: Auto Sync (uses markers when present, time otherwise, default behavior)
+		// 1: Time Sync (ignores markers)
+		// 2: No Sync (un-grouped sync behavior)
+		static TAutoConsoleVariable<int32> CDebugVarSyncMode(TEXT("a.AnimNext.SyncGroupMode"), 0, TEXT("Debug sync group follower mode. 0: Auto Sync, 1: Time Sync, 2: No Sync"));
+#endif
+
 		// Sync group member state as collected during graph traversal
 		struct FSyncGroupMember
 		{
 			FTraitUpdateState					TraitState;
-			FWeakTraitPtr						TraitPtr;
+
+			// Strong pointer to trait to keep it alive
+			FTraitPtr							TraitPtr;
 
 			FSyncGroupParameters				GroupParameters;
+
+			// Whether or not this member was part of the group on a given update
+			bool bIsActive = false;
+
+			// Whether or not we joined the group in the current update
+			bool bJustJoined = false;
 		};
 
 		// Sync group state as collected during graph traversal
@@ -34,11 +54,15 @@ namespace UE::AnimNext
 			// If this sync group was active during the previous update, this is the timeline progress of its leader
 			FTimelineState PreviousLeaderTimelineState;
 
-			// If this sync group was active during the previous update and using sync markers, this is the sync position of the leader
-			FMarkerSyncAnimPosition PreviousLeaderSyncPosition;
+#if UE_DEBUG_SYNC_GROUPS
+			int32 PreviousLeaderIndex = INDEX_NONE;
+#endif
 
 			// Whether or not this sync group was active during the previous update
-			bool bWasActive = false;
+			bool bIsActive = false;
+
+			// Whether or not this sync group just formed during the current update
+			bool bJustFormed = false;
 		};
 
 		// Sync group member context when performing group synchronization
@@ -47,8 +71,6 @@ namespace UE::AnimNext
 			const FSyncGroupMember* State = nullptr;
 
 			FTraitStackBinding TraitStack;
-			
-			FTimelineState TimelineState;
 			TTraitBinding<ITimeline> TimelineTrait;
 
 			FTimelineSyncMarkerArray SyncMarkers;
@@ -58,7 +80,6 @@ namespace UE::AnimNext
 			{
 				State = &InState;
 				TraitStack.Reset();
-				TimelineState.Reset();
 				TimelineTrait.Reset();
 				SyncMarkers.Reset();
 				bUseMarkerSyncing = false;
@@ -68,27 +89,24 @@ namespace UE::AnimNext
 		// Sync group context when performing group synchronization
 		struct FSyncGroupContext
 		{
+			const FSyncGroupState* State = nullptr;
+
 			TArrayView<FSyncGroupMemberContext> Members;
 
 			// If this sync group was active during the previous update, this is the timeline progress of its leader
 			FTimelineState PreviousLeaderTimelineState;
 
-			// If this sync group was active during the previous update and using sync markers, this is the sync position of the leader
-			FMarkerSyncAnimPosition PreviousLeaderSyncPosition;
-
 			// Current leader progress, set after leader has advanced
 			FTimelineState LeaderTimelineState;
 			int32 LeaderIndex = INDEX_NONE;
-			float LeaderDeltaTime = 0.0f;
-
-			// Whether or not this sync group was active during the previous update
-			bool bWasActive = false;
+			bool bIsLeaderPlayingForward = false;
+			float LeaderStartRatio = 0.0f;
+			float LeaderEndRatio = 0.0f;
+			TArray<FName, TMemStackAllocator<>> MarkersPassed;
 
 			bool bCanGroupUseMarkerSyncing = false;
 			FMarkerSyncAnimPosition SyncPositionStart;
 			FMarkerSyncAnimPosition SyncPositionEnd;
-			FMarkerPair LeaderStartMarker;
-			FMarkerPair LeaderEndMarker;
 
 			TSet<FName> ValidMarkers;
 			TSet<FName> CandidateMarkers;
@@ -117,13 +135,15 @@ namespace UE::AnimNext
 				Contextes[MemberIndex].Init(GroupState.Members[MemberIndex]);
 			}
 
+			GroupContext.State = &GroupState;
 			GroupContext.Members = Contextes;
 			GroupContext.PreviousLeaderTimelineState = GroupState.PreviousLeaderTimelineState;
-			GroupContext.PreviousLeaderSyncPosition = GroupState.PreviousLeaderSyncPosition;
 			GroupContext.LeaderTimelineState.Reset();
 			GroupContext.LeaderIndex = INDEX_NONE;
-			GroupContext.LeaderDeltaTime = 0.0f;
-			GroupContext.bWasActive = GroupState.bWasActive;
+			GroupContext.bIsLeaderPlayingForward = false;
+			GroupContext.LeaderStartRatio = 0.0f;
+			GroupContext.LeaderEndRatio = 0.0f;
+			GroupContext.MarkersPassed.Reset();
 			GroupContext.bCanGroupUseMarkerSyncing = false;
 			GroupContext.ValidMarkers.Reset();
 			GroupContext.CandidateMarkers.Reset();
@@ -177,6 +197,16 @@ namespace UE::AnimNext
 				LeaderIndex = 0;
 			}
 
+#if UE_DEBUG_SYNC_GROUPS
+			if (CDebugVarEnableSyncLog.GetValueOnAnyThread() && GroupContext.State->PreviousLeaderIndex != LeaderIndex)
+			{
+				UE_LOG(LogAnimMarkerSync, Warning,
+					TEXT("[%s] [%p] Is New Leader"),
+					*GroupContext.State->GroupName.ToString(),
+					GroupContext.Members[LeaderIndex].State->TraitPtr.GetNodeInstance());
+			}
+#endif
+
 			GroupContext.LeaderIndex = LeaderIndex;
 		}
 
@@ -352,8 +382,9 @@ namespace UE::AnimNext
 			return Algo::UpperBoundBy(SyncMarkers, Position, [](const FTimelineSyncMarker& Marker) { return Marker.GetPosition(); });
 		}
 
-		// Returns the boundary sync markers around the specified timeline position using the supplied predicate to determine marker validity
-		static void GetBoundarySyncMarkersAt(const TFunctionRef<bool(FName MarkerName, bool bIsNextMarker)>& MarkerValidityPredicate, const FTimelineSyncMarkerArray& SyncMarkers, const FTimelineState& TimelineState, FMarkerPair& OutPrevMarker, FMarkerPair& OutNextMarker)
+		// Returns the valid pair of sync markers around the specified timeline position
+		// This can return the animation boundary if the timeline isn't looping
+		static void GetNearestSyncMarkersAt(const TSet<FName>& ValidMarkers, const FTimelineSyncMarkerArray& SyncMarkers, const FTimelineState& TimelineState, FMarkerPair& OutPrevMarker, FMarkerPair& OutNextMarker)
 		{
 			const float CurrentTime = TimelineState.GetPosition();
 			const float Duration = TimelineState.GetDuration();
@@ -372,8 +403,7 @@ namespace UE::AnimNext
 				// Refine our search using the valid markers
 				while (PrevMarkerIndex >= 0)
 				{
-					//if (ValidMarkers.Contains(SyncMarkers[PrevMarkerIndex].GetName()))
-					if (MarkerValidityPredicate(SyncMarkers[PrevMarkerIndex].GetName(), false))
+					if (ValidMarkers.Contains(SyncMarkers[PrevMarkerIndex].GetName()))
 					{
 						// This marker is valid, use it
 						break;
@@ -391,8 +421,7 @@ namespace UE::AnimNext
 
 						while (PrevMarkerIndex >= 0)
 						{
-							//if (ValidMarkers.Contains(SyncMarkers[PrevMarkerIndex].GetName()))
-							if (MarkerValidityPredicate(SyncMarkers[PrevMarkerIndex].GetName(), false))
+							if (ValidMarkers.Contains(SyncMarkers[PrevMarkerIndex].GetName()))
 							{
 								// This marker is valid, use it
 								break;
@@ -424,8 +453,7 @@ namespace UE::AnimNext
 				// Refine our search using the valid markers
 				while (NextMarkerIndex < SyncMarkers.Num())
 				{
-					//if (ValidMarkers.Contains(SyncMarkers[NextMarkerIndex].GetName()))
-					if (MarkerValidityPredicate(SyncMarkers[NextMarkerIndex].GetName(), true))
+					if (ValidMarkers.Contains(SyncMarkers[NextMarkerIndex].GetName()))
 					{
 						// This marker is valid, use it
 						break;
@@ -443,8 +471,7 @@ namespace UE::AnimNext
 
 						while (NextMarkerIndex < SyncMarkers.Num())
 						{
-							//if (ValidMarkers.Contains(SyncMarkers[NextMarkerIndex].GetName()))
-							if (MarkerValidityPredicate(SyncMarkers[NextMarkerIndex].GetName(), true))
+							if (ValidMarkers.Contains(SyncMarkers[NextMarkerIndex].GetName()))
 							{
 								// This marker is valid, use it
 								break;
@@ -477,40 +504,113 @@ namespace UE::AnimNext
 			OutNextMarker.TimeToMarker = NextTimeToMarker;
 		}
 
-		// Returns the boundary sync markers around the specified timeline position
-		static void GetBoundarySyncMarkersAt(const TSet<FName>& ValidMarkers, const FTimelineSyncMarkerArray& SyncMarkers, const FTimelineState& TimelineState, FMarkerPair& OutPrevMarker, FMarkerPair& OutNextMarker)
+		static FName GetMarkerName(int32 MarkerIndex, const FTimelineSyncMarkerArray& SyncMarkers)
 		{
-			const auto MarkerValidityPredicate =
-				[&ValidMarkers](FName MarkerName, bool bIsNextMarker)
-				{
-					return ValidMarkers.Contains(MarkerName);
-				};
+			return SyncMarkers.IsValidIndex(MarkerIndex) ? SyncMarkers[MarkerIndex].GetName() : NAME_None;
+		}
 
-			return GetBoundarySyncMarkersAt(MarkerValidityPredicate, SyncMarkers, TimelineState, OutPrevMarker, OutNextMarker);
+		static FName GetMarkerName(const FMarkerPair& Marker, const FTimelineSyncMarkerArray& SyncMarkers)
+		{
+			return GetMarkerName(Marker.MarkerIndex, SyncMarkers);
 		}
 
-		// Returns the boundary sync markers around the specified timeline position
-		static void GetBoundarySyncMarkersAt(const FMarkerSyncAnimPosition& SyncPosition, const FTimelineSyncMarkerArray& SyncMarkers, const FTimelineState& TimelineState, FMarkerPair& OutPrevMarker, FMarkerPair& OutNextMarker)
+		static float GetMarkerPosition(int32 MarkerIndex, const FTimelineSyncMarkerArray& SyncMarkers, float FallbackValue)
 		{
-			const auto MarkerValidityPredicate =
-				[&SyncPosition](FName MarkerName, bool bIsNextMarker)
-				{
-					return MarkerName == SyncPosition.PreviousMarkerName || MarkerName == SyncPosition.NextMarkerName;
-				};
+			return SyncMarkers.IsValidIndex(MarkerIndex) ? SyncMarkers[MarkerIndex].GetPosition() : FallbackValue;
+		}
 
-			return GetBoundarySyncMarkersAt(MarkerValidityPredicate, SyncMarkers, TimelineState, OutPrevMarker, OutNextMarker);
+		static float GetMarkerPosition(const FMarkerPair& Marker, const FTimelineSyncMarkerArray& SyncMarkers, float FallbackValue)
+		{
+			return GetMarkerPosition(Marker.MarkerIndex, SyncMarkers, FallbackValue);
 		}
 
-		// Returns the boundary sync markers around the specified timeline position matching the supplied marker pair exactly
+		// Returns the closest pair of sync markers to the specified timeline position that matches exactly the supplied sync position
 		static void GetExactBoundarySyncMarkersAt(const FMarkerSyncAnimPosition& SyncPosition, const FTimelineSyncMarkerArray& SyncMarkers, const FTimelineState& TimelineState, FMarkerPair& OutPrevMarker, FMarkerPair& OutNextMarker)
 		{
-			const auto MarkerValidityPredicate =
-				[&SyncPosition](FName MarkerName, bool bIsNextMarker)
+			const float CurrentTime = TimelineState.GetPosition();
+			const float Duration = TimelineState.GetDuration();
+			const bool bIsLooping = TimelineState.IsLooping();
+			const int32 NumSyncMarkers = SyncMarkers.Num();
+
+			float BestTimeDelta = FLT_MAX;
+
+			// Handle case for looping and sync position not being on either boundary.
+			for (int32 PrevMarkerIdx = 0; PrevMarkerIdx < NumSyncMarkers; ++PrevMarkerIdx)
+			{
+				const FTimelineSyncMarker& PrevMarker = SyncMarkers[PrevMarkerIdx];
+
+				if (PrevMarker.GetName() != SyncPosition.PreviousMarkerName)
+				{
+					continue;	// Not matching, look for the next pair candidate
+				}
+
+				const float PrevMarkerTime = PrevMarker.GetPosition();
+				const int32 EndMarkerSearchStart = PrevMarkerIdx + 1;
+				const int32 EndCount = bIsLooping ? NumSyncMarkers + EndMarkerSearchStart : NumSyncMarkers;
+
+				for (int32 NextMarkerCount = EndMarkerSearchStart; NextMarkerCount < EndCount; ++NextMarkerCount)
+				{
+					const int32 NextMarkerIdx = NextMarkerCount % NumSyncMarkers;
+					const FTimelineSyncMarker& NextMarker = SyncMarkers[NextMarkerIdx];
+
+					// We have matched the position's next marker name.
+					if (NextMarker.GetName() != SyncPosition.NextMarkerName)
+					{
+						continue;	// Not matching, look for the next pair candidate
+					}
+
+					float NextMarkerTime = NextMarker.GetPosition();
+
+					// Handle case where we need to loop to get to be able to get to the next marker.
+					if (NextMarkerTime < PrevMarkerTime)
+					{
+						NextMarkerTime += Duration;
+					}
+
+					// Get current time based of sync position.
+					float DesiredTime = FMath::LerpStable(PrevMarkerTime, NextMarkerTime, SyncPosition.PositionBetweenMarkers);
+
+					// Adjust for case where we need to loop to be able to get to the next marker.
+					if (DesiredTime > Duration)
+					{
+						DesiredTime -= Duration;
+					}
+
+					// Find marker indices closest to input time position.
+					const float TimeDelta = FMath::Abs(DesiredTime - CurrentTime);
+					if (TimeDelta < BestTimeDelta)
+					{
+						BestTimeDelta = TimeDelta;
+						OutPrevMarker.MarkerIndex = PrevMarkerIdx;
+						OutNextMarker.MarkerIndex = NextMarkerIdx;
+					}
+
+					// This marker test is done, move onto next one.
+					break;
+				}
+
+				// If we get here and we haven't found a match and we are not looping then there
+				// is no point running the rest of the loop set up something as relevant as we can and carry on
+				if (OutPrevMarker.MarkerIndex == MarkerIndexSpecialValues::Uninitialized)
 				{
-					return bIsNextMarker ? (MarkerName == SyncPosition.NextMarkerName) : (MarkerName == SyncPosition.PreviousMarkerName);
-				};
+					// Find nearest previous marker that is earlier than our current time
+					BestTimeDelta = CurrentTime - PrevMarkerTime;
+					int32 PrevMarkerToUse = PrevMarkerIdx + 1;
+					while (BestTimeDelta > 0.f && PrevMarkerToUse < NumSyncMarkers)
+					{
+						BestTimeDelta = CurrentTime - SyncMarkers[PrevMarkerToUse].GetPosition();
+						++PrevMarkerToUse;
+					}
 
-			return GetBoundarySyncMarkersAt(MarkerValidityPredicate, SyncMarkers, TimelineState, OutPrevMarker, OutNextMarker);
+					OutPrevMarker.MarkerIndex = PrevMarkerToUse - 1;	// We always go one past the marker we actually want to use
+
+					// This goes to minus one as the very fact we are here means
+					// that there is no next marker to use
+					OutNextMarker.MarkerIndex = -1;
+					
+					break; // no need to keep searching, we are done
+				}
+			}
 		}
 
 		// Returns the boundary sync markers matching the sync position pair nearest the specified timeline position and the adjusted current time based on the sync position
@@ -518,16 +618,6 @@ namespace UE::AnimNext
 		{
 			check(SyncPosition.PreviousMarkerName != NAME_None || SyncPosition.NextMarkerName != NAME_None);
 
-			// If we are not looping or if we have no usable sync position
-			//		- We scan the marker list for the matching marker pair around the current time of the follower
-			//		- Current time is unchanged (no syncing)
-			//		- Done
-			if (!TimelineState.IsLooping())
-			{
-				GetBoundarySyncMarkersAt(SyncPosition, SyncMarkers, TimelineState, OutPrevMarker, OutNextMarker);
-				return TimelineState.GetPosition();
-			}
-
 			// If the sync position's previous marker is the start boundary
 			//		- We keep the previous marker at the boundary
 			//		- We look for the next marker, the first instance of the next marker's name
@@ -540,8 +630,7 @@ namespace UE::AnimNext
 
 				OutPrevMarker.MarkerIndex = MarkerIndexSpecialValues::AnimationBoundary;
 
-				const int32 NumSyncMarkers = SyncMarkers.Num();
-				for (int32 MarkerIndex = 0; MarkerIndex < NumSyncMarkers; ++MarkerIndex)
+				for (int32 MarkerIndex = 0; MarkerIndex < SyncMarkers.Num(); ++MarkerIndex)
 				{
 					if (SyncMarkers[MarkerIndex].GetName() == SyncPosition.NextMarkerName)
 					{
@@ -571,10 +660,9 @@ namespace UE::AnimNext
 
 				OutNextMarker.MarkerIndex = MarkerIndexSpecialValues::AnimationBoundary;
 
-				const int32 NumSyncMarkers = SyncMarkers.Num();
-				for (int32 MarkerIndex = NumSyncMarkers - 1; MarkerIndex >= 0; --MarkerIndex)
+				for (int32 MarkerIndex = SyncMarkers.Num() - 1; MarkerIndex >= 0; --MarkerIndex)
 				{
-					if (SyncMarkers[MarkerIndex].GetName() == SyncPosition.NextMarkerName)
+					if (SyncMarkers[MarkerIndex].GetName() == SyncPosition.PreviousMarkerName)
 					{
 						const float PrevMarkerTime = SyncMarkers[MarkerIndex].GetPosition();
 						CurrentTime = FMath::LerpStable(PrevMarkerTime, NextMarkerTime, SyncPosition.PositionBetweenMarkers);
@@ -592,26 +680,29 @@ namespace UE::AnimNext
 
 			// Otherwise
 			//		- Scan the marker pairs that match the previous/next markers from the leader
-			//		- We look for the pair that yields an adjusted current time closest to the follower's current time, O(N^2) search
+			//		- We look for the pair that yields an adjusted current time closest to the follower's current time
 			//		- Once we find that pair, we return the current time based on the sync position (see GetCurrentTimeFromMarkers)
 			//		- If no such pair is found, the follower must not contain the pair of markers present on the leader (not possible)
 			//		- Done
 			GetExactBoundarySyncMarkersAt(SyncPosition, SyncMarkers, TimelineState, OutPrevMarker, OutNextMarker);
 
 			float CurrentTime = TimelineState.GetPosition();
-			float PrevMarkerTime = SyncMarkers[OutPrevMarker.MarkerIndex].GetPosition();
-			float NextMarkerTime = SyncMarkers[OutNextMarker.MarkerIndex].GetPosition();
+			float PrevMarkerTime = GetMarkerPosition(OutPrevMarker.MarkerIndex, SyncMarkers, 0.0f);
+			float NextMarkerTime = GetMarkerPosition(OutNextMarker.MarkerIndex, SyncMarkers, TimelineState.GetDuration());
 
-			// Handle looping
-			if (PrevMarkerTime > CurrentTime)
-			{
-				PrevMarkerTime -= TimelineState.GetDuration();
-			}
-			else if (NextMarkerTime < CurrentTime)
+			if (TimelineState.IsLooping())
 			{
-				NextMarkerTime += TimelineState.GetDuration();
+				if (PrevMarkerTime > CurrentTime)
+				{
+					PrevMarkerTime -= TimelineState.GetDuration();
+				}
+				else if (NextMarkerTime < CurrentTime)
+				{
+					NextMarkerTime += TimelineState.GetDuration();
+				}
 			}
 
+			// Map our current time in between the markers we found
 			CurrentTime = FMath::LerpStable(PrevMarkerTime, NextMarkerTime, SyncPosition.PositionBetweenMarkers);
 			return CurrentTime;
 		}
@@ -661,398 +752,999 @@ namespace UE::AnimNext
 			check(NextTime > PrevTime);
 
 			// Store the encoded current time position as a ratio between markers
-			SyncPosition.PositionBetweenMarkers = (CurrentTime - PrevTime) / (NextTime - PrevTime);
+			const float PositionBetweenMarkers = (CurrentTime - PrevTime) / (NextTime - PrevTime);
+			check(PositionBetweenMarkers >= (0.0f - SMALL_NUMBER) && PositionBetweenMarkers <= (1.0f + SMALL_NUMBER));
+
+			SyncPosition.PositionBetweenMarkers = FMath::Clamp(PositionBetweenMarkers, 0.0f, 1.0f);
 			return SyncPosition;
 		}
 
-		static float MatchLeaderMarkerPhase(const FSyncGroupContext& GroupContext, const FSyncGroupMemberContext& GroupLeaderContext, const FSyncGroupMemberContext& GroupMemberContext)
+		// Returns the desired time based on the provided markers and normalized position between
+		static float CalculateTimeFromSyncPosition(const FTimelineSyncMarkerArray& SyncMarkers, const FTimelineState& TimelineState, const FMarkerPair& PrevMarker, const FMarkerPair& NextMarker, float NormalizedPosition)
 		{
-			// Find the nearest markers matching our start position
-			FMarkerPair PrevMarker;
-			FMarkerPair NextMarker;
-			GetNearestMatchingBoundarySyncMarkers(GroupContext.SyncPositionStart, GroupMemberContext.SyncMarkers, GroupMemberContext.TimelineState, PrevMarker, NextMarker);
-
-			// Using the leader start/end sync position, skip as many markers as the leader skipped
-			int32 NumPassedMarkers;
-			if (GroupContext.LeaderStartMarker.MarkerIndex > GroupContext.LeaderEndMarker.MarkerIndex)
+			const float CurrentTime = TimelineState.GetPosition();
+			const float Duration = TimelineState.GetDuration();
+
+			float PrevMarkerTime = GetMarkerPosition(PrevMarker, SyncMarkers, 0.0f);
+			float NextMarkerTime = GetMarkerPosition(NextMarker, SyncMarkers, Duration);
+
+			if (PrevMarkerTime >= NextMarkerTime)
 			{
-				// We looped, add the trailing markers
-				NumPassedMarkers = GroupLeaderContext.SyncMarkers.Num() - GroupContext.LeaderStartMarker.MarkerIndex + GroupContext.LeaderEndMarker.MarkerIndex;
+				// We are looping around, fixup the previous marker to come before the next marker
+				// extending it past the start of the timeline (negative position)
+				PrevMarkerTime -= Duration;
 			}
-			else
+
+			float DesiredTime = FMath::LerpStable(PrevMarkerTime, NextMarkerTime, NormalizedPosition);
+
+			if (DesiredTime < 0.0f)
 			{
-				NumPassedMarkers = GroupContext.LeaderEndMarker.MarkerIndex - GroupContext.LeaderStartMarker.MarkerIndex;
+				// The markers are looping around but the desired time hasn't looped yet, wrap back
+				DesiredTime += Duration;
 			}
 
-			int32 PrevMarkerIndex = PrevMarker.MarkerIndex + NumPassedMarkers;
-			float PrevMarkerTime = 0.0f;
+			return FMath::Clamp(DesiredTime, 0.0f, Duration);
+		}
+
+		static void AdvanceAndCollectSyncMarkersPassed(
+			const FTimelineState& TimelineState, float DeltaTime, bool bIsPlayingForward,
+			const FTimelineSyncMarkerArray& SyncMarkers, FMarkerPair& PrevMarker, FMarkerPair& NextMarker,
+			TArray<FName, TMemStackAllocator<>>& OutMarkersPassed)
+		{
+			OutMarkersPassed.Reset();
+
+			const float Duration = TimelineState.GetDuration();
+			const bool bIsLooping = TimelineState.IsLooping();
+
+			// Treat delta time as being positive since we don't know if we move backwards because of the delta time (e.g. game rewind)
+			// or if it's because of the play rate
+			float RemainingDeltaTime = FMath::Abs(DeltaTime);
+			float CurrentTime = TimelineState.GetPosition();
 
+			if (bIsPlayingForward)
 			{
-				if (PrevMarkerIndex >= GroupMemberContext.SyncMarkers.Num())
+				// Progressively consume our delta time
+				while (RemainingDeltaTime > 0.0f)
 				{
-					if (GroupMemberContext.TimelineState.IsLooping())
+					// Our next marker is the end boundary. (Only possible if sequence is not looping)
+					if (NextMarker.MarkerIndex == MarkerIndexSpecialValues::AnimationBoundary)
 					{
-						PrevMarkerIndex = 0;
+						break;
+					}
 
-						while (PrevMarkerIndex < GroupMemberContext.SyncMarkers.Num())
-						{
-							if (GroupContext.SyncPositionEnd.PreviousMarkerName == GroupMemberContext.SyncMarkers[PrevMarkerIndex].GetName())
-							{
-								// This marker is valid, use it
-								break;
-							}
+					const FTimelineSyncMarker& NextSyncMarker = SyncMarkers[NextMarker.MarkerIndex];
 
-							PrevMarkerIndex++;
-						}
+					float NextMarkerTime;
+					if (CurrentTime <= NextSyncMarker.GetPosition())
+					{
+						// Next marker is ahead of us, use it as is
+						NextMarkerTime = NextSyncMarker.GetPosition();
+					}
+					else
+					{
+						// We are looping
+						check(bIsLooping);
+						NextMarkerTime = NextSyncMarker.GetPosition() + Duration;
+					}
 
-						check(PrevMarkerIndex != GroupMemberContext.SyncMarkers.Num());
+					const float TargetTime = CurrentTime + RemainingDeltaTime;
+					if (TargetTime > NextMarkerTime)
+					{
+						// We passed this marker
+						OutMarkersPassed.Add(NextSyncMarker.GetName());
 
-						// Marker lives in next loop iteration, in the 'future'
-						PrevMarkerTime = GroupMemberContext.SyncMarkers[PrevMarkerIndex].GetPosition();
+						// Update our marker tracking
+						PrevMarker = NextMarker;
+						NextMarker.MarkerIndex++;
+						if (NextMarker.MarkerIndex >= SyncMarkers.Num())
+						{
+							NextMarker.MarkerIndex = bIsLooping ? 0 : MarkerIndexSpecialValues::AnimationBoundary;
+						}
+
+						// Update our time tracking
+						const float ConsumedTime = NextMarkerTime - CurrentTime;
+						CurrentTime = NextSyncMarker.GetPosition();
+						RemainingDeltaTime -= ConsumedTime;
 					}
 					else
 					{
-						PrevMarkerIndex = MarkerIndexSpecialValues::AnimationBoundary;
-						PrevMarkerTime = GroupMemberContext.TimelineState.GetDuration();
+						break;
 					}
 				}
-				else
+			}
+			else
+			{
+				// Progressively consume our delta time
+				while (RemainingDeltaTime > 0.0f)
 				{
-					PrevMarkerTime = GroupMemberContext.SyncMarkers[PrevMarkerIndex].GetPosition();
-				}
+					// Our next marker is the end boundary. (Only possible if sequence is not looping)
+					if (PrevMarker.MarkerIndex == MarkerIndexSpecialValues::AnimationBoundary)
+					{
+						break;
+					}
 
-				check(GroupMemberContext.SyncMarkers[PrevMarkerIndex].GetName() == GroupContext.SyncPositionEnd.PreviousMarkerName);
-			}
+					const FTimelineSyncMarker& PrevSyncMarker = SyncMarkers[PrevMarker.MarkerIndex];
 
-			int32 NextMarkerIndex = PrevMarkerIndex + 1;
-			float NextMarkerTime = PrevMarkerTime;
+					float PrevMarkerTime;
+					if (CurrentTime >= PrevSyncMarker.GetPosition())
+					{
+						// Previous marker is ahead of us, use it as is
+						PrevMarkerTime = PrevSyncMarker.GetPosition();
+					}
+					else
+					{
+						// We are looping
+						check(bIsLooping);
+						PrevMarkerTime = PrevSyncMarker.GetPosition() - Duration;
+					}
 
-			{
-				if (NextMarkerIndex >= GroupMemberContext.SyncMarkers.Num())
-				{
-					if (GroupMemberContext.TimelineState.IsLooping())
+					const float TargetTime = CurrentTime - RemainingDeltaTime;
+					if (TargetTime < PrevMarkerTime)
 					{
-						NextMarkerIndex = 0;
+						// We passed this marker
+						OutMarkersPassed.Add(PrevSyncMarker.GetName());
 
-						while (NextMarkerIndex < GroupMemberContext.SyncMarkers.Num())
+						// Update our marker tracking
+						NextMarker = PrevMarker;
+						PrevMarker.MarkerIndex--;
+						if (PrevMarker.MarkerIndex == INDEX_NONE)
 						{
-							if (GroupContext.SyncPositionEnd.NextMarkerName == GroupMemberContext.SyncMarkers[NextMarkerIndex].GetName())
-							{
-								// This marker is valid, use it
-								break;
-							}
-
-							NextMarkerIndex++;
+							PrevMarker.MarkerIndex = bIsLooping ? (SyncMarkers.Num() - 1) : MarkerIndexSpecialValues::AnimationBoundary;
 						}
 
-						check(NextMarkerIndex != GroupMemberContext.SyncMarkers.Num());
-
-						// Marker lives in next loop iteration, in the 'future'
-						NextMarkerTime = GroupMemberContext.SyncMarkers[NextMarkerIndex].GetPosition() + GroupMemberContext.TimelineState.GetDuration();
+						// Update our time tracking
+						const float ConsumedTime = PrevMarkerTime - CurrentTime;
+						CurrentTime = PrevSyncMarker.GetPosition();
+						RemainingDeltaTime += ConsumedTime;
 					}
 					else
 					{
-						NextMarkerIndex = MarkerIndexSpecialValues::AnimationBoundary;
-						NextMarkerTime = GroupMemberContext.TimelineState.GetDuration();
+						break;
 					}
 				}
-				else
-				{
-					NextMarkerTime = GroupMemberContext.SyncMarkers[NextMarkerIndex].GetPosition();
-				}
-
-				check(GroupMemberContext.SyncMarkers[NextMarkerIndex].GetName() == GroupContext.SyncPositionEnd.NextMarkerName);
 			}
+		}
 
-			const float CurrentTime = FMath::LerpStable(PrevMarkerTime, NextMarkerTime, GroupContext.SyncPositionEnd.PositionBetweenMarkers);
+		static void SeekMarkerForward(int32& MarkerIndex, bool bIsLooping, FName TargetMarkerName, const FTimelineSyncMarkerArray& SyncMarkers)
+		{
+			const int32 NumSyncMarkers = SyncMarkers.Num();
 
-			float DeltaTime = CurrentTime - GroupMemberContext.TimelineState.GetPosition();
-			if (DeltaTime >= GroupMemberContext.TimelineState.GetDuration())
+			int32 MaxNumIterations = NumSyncMarkers;
+			while (MarkerIndex < NumSyncMarkers && MaxNumIterations > 0)
 			{
-				DeltaTime -= GroupMemberContext.TimelineState.GetDuration();
-			}
+				if (SyncMarkers[MarkerIndex].GetName() == TargetMarkerName)
+				{
+					// This is the marker we passed
+					break;
+				}
 
-#if UE_LOG_ENABLE_MARKER_DEBUG
-			UE_LOG(LogAnimMarkerSync, Warning,
-				TEXT("[%p] Follower PreviousTime(%0.2f) CurrentTime(%0.2f) MoveDelta(%0.2f) Looping(%d) [(Prev '%s' @ %d | Next '%s' @ %d | %0.2f)]"),
-				GroupMemberContext.State->TraitPtr.GetNodeInstance(),
-				GroupMemberContext.TimelineState.GetPosition(),
-				CurrentTime,
-				DeltaTime,
-				GroupLeaderContext.TimelineState.IsLooping() ? 1 : 0,
-				*GroupMemberContext.SyncMarkers[PrevMarkerIndex].GetName().ToString(),
-				PrevMarkerIndex,
-				*GroupMemberContext.SyncMarkers[NextMarkerIndex].GetName().ToString(),
-				NextMarkerIndex,
-				GroupContext.SyncPositionEnd.PositionBetweenMarkers);
-#endif
+				MarkerIndex++;
+				MaxNumIterations--;
 
-			return DeltaTime;
+				if (MarkerIndex == NumSyncMarkers)
+				{
+					if (bIsLooping)
+					{
+						MarkerIndex = 0;
+					}
+					else
+					{
+						MarkerIndex = MarkerIndexSpecialValues::AnimationBoundary;
+						break;
+					}
+				}
+			}
 		}
 
-		static void AdvanceLeader(FExecutionContext& Context, FSyncGroupContext& GroupContext)
+		static void SeekMarkerBackward(int32& MarkerIndex, bool bIsLooping, FName TargetMarkerName, const FTimelineSyncMarkerArray& SyncMarkers)
 		{
-			const int32 LeaderIndex = GroupContext.LeaderIndex;
-
-			FSyncGroupMemberContext& GroupLeaderContext = GroupContext.Members[LeaderIndex];
-			const FSyncGroupMember& GroupLeader = *GroupLeaderContext.State;
-
-			FTraitStackBinding& LeaderTraitStack = GroupLeaderContext.TraitStack;
-
-			TTraitBinding<IGroupSynchronization> GroupSyncTrait;
-
-			Context.BindTo(GroupLeader.TraitPtr);
-			ensure(Context.GetStack(GroupLeader.TraitPtr, LeaderTraitStack));
-			ensure(LeaderTraitStack.GetInterface(GroupSyncTrait));
-
-			// Cache our starting timeline state
-			GroupLeaderContext.TimelineState = GroupLeaderContext.TimelineTrait.GetState(Context);
-
-			float DeltaTime = GroupLeader.TraitState.GetDeltaTime();
+			const int32 NumSyncMarkers = SyncMarkers.Num();
 
-			if (GroupContext.bCanGroupUseMarkerSyncing)
+			int32 MaxNumIterations = NumSyncMarkers;
+			while (MarkerIndex >= 0 && MaxNumIterations > 0)
 			{
-				FTimelineState LeaderTimelineState;
-				FMarkerPair PrevMarkerStart;
-				FMarkerPair NextMarkerStart;
-
-				if (GroupContext.bWasActive)
+				if (SyncMarkers[MarkerIndex].GetName() == TargetMarkerName)
 				{
-					if (GroupLeaderContext.State->GroupParameters.bMatchSyncPoint)
-					{
-						// This is a soft join where we wish to match the current sync group position
-						float CurrentTime = GroupContext.PreviousLeaderTimelineState.GetPositionRatio() * GroupLeaderContext.TimelineState.GetDuration();
-
-						LeaderTimelineState = GroupLeaderContext.TimelineState.WithPosition(CurrentTime);
+					// This is the marker we passed
+					break;
+				}
 
-						if (GroupContext.PreviousLeaderSyncPosition.IsValid() &&
-							GroupContext.ValidMarkers.Contains(GroupContext.PreviousLeaderSyncPosition.PreviousMarkerName) &&
-							GroupContext.ValidMarkers.Contains(GroupContext.PreviousLeaderSyncPosition.NextMarkerName))
-						{
-							// Previous sync position markers are valid, we can use them
-							CurrentTime = GetNearestMatchingBoundarySyncMarkers(GroupContext.PreviousLeaderSyncPosition, GroupLeaderContext.SyncMarkers, LeaderTimelineState, PrevMarkerStart, NextMarkerStart);
+				MarkerIndex--;
+				MaxNumIterations--;
 
-							LeaderTimelineState = GroupLeaderContext.TimelineState.WithPosition(CurrentTime);
-						}
-						else
-						{
-							// Previous sync position markers are invalid, we can't use them
-							GetBoundarySyncMarkersAt(GroupContext.ValidMarkers, GroupLeaderContext.SyncMarkers, LeaderTimelineState, PrevMarkerStart, NextMarkerStart);
-						}
+				if (MarkerIndex < 0)
+				{
+					if (bIsLooping)
+					{
+						MarkerIndex = NumSyncMarkers - 1;
 					}
 					else
 					{
-						// This is a hard join where the leader dictates its new position
-						float CurrentTime = GroupLeaderContext.TimelineState.GetPosition();
-
-						LeaderTimelineState = GroupLeaderContext.TimelineState.WithPosition(CurrentTime);
-
-						GetBoundarySyncMarkersAt(GroupContext.ValidMarkers, GroupLeaderContext.SyncMarkers, LeaderTimelineState, PrevMarkerStart, NextMarkerStart);
+						MarkerIndex = MarkerIndexSpecialValues::AnimationBoundary;
+						break;
 					}
 				}
-				else
-				{
-					// This is a new sync group, leader dictates its position
-					float CurrentTime = GroupLeaderContext.TimelineState.GetPosition();
+			}
+		}
 
-					LeaderTimelineState = GroupLeaderContext.TimelineState.WithPosition(CurrentTime);
+		static float AdvanceSyncMarkersPassed(const FTimelineState& TimelineState, bool bIsPlayingForward,
+			const FMarkerSyncAnimPosition& LeaderEndPosition,
+			const FTimelineSyncMarkerArray& SyncMarkers, const TArray<FName, TMemStackAllocator<>>& MarkersPassed,
+			FMarkerPair& PrevMarker, FMarkerPair& NextMarker)
+		{
+			const bool bIsLooping = TimelineState.IsLooping();
+			const int32 NumMarkersPassed = MarkersPassed.Num();
 
-					GetBoundarySyncMarkersAt(GroupContext.ValidMarkers, GroupLeaderContext.SyncMarkers, LeaderTimelineState, PrevMarkerStart, NextMarkerStart);
-				}
+			if (bIsPlayingForward)
+			{
+				if (NumMarkersPassed > 0)
+				{
+					// Skip over the markers that passed
+					for (int32 PassedMarkerIndex = 0; PassedMarkerIndex < NumMarkersPassed; ++PassedMarkerIndex)
+					{
+						if (NextMarker.MarkerIndex == MarkerIndexSpecialValues::AnimationBoundary)
+						{
+							// We still have markers left to pass but the follower ran out of markers
+							break;
+						}
 
-				// Cache the start sync position
-				GroupContext.SyncPositionStart = CalculateSyncPosition(GroupLeaderContext.SyncMarkers, LeaderTimelineState, PrevMarkerStart, NextMarkerStart);
-				GroupContext.LeaderStartMarker = PrevMarkerStart;
+						PrevMarker = NextMarker;
 
-				// We advance the leader by its desired delta time
-				GroupSyncTrait.AdvanceBy(Context, DeltaTime);
+						const FName PassedMarker = MarkersPassed[PassedMarkerIndex];
 
-				// Compute our new timeline state
-				const FTimelineState NewTimelineState = GroupLeaderContext.TimelineTrait.GetState(Context);
+						// Look for the instance of our passed marker
+						SeekMarkerForward(NextMarker.MarkerIndex, bIsLooping, PassedMarker, SyncMarkers);
+					}
 
-				GroupContext.LeaderTimelineState = NewTimelineState;
-				GroupContext.LeaderDeltaTime = DeltaTime;
+					if (LeaderEndPosition.NextMarkerName == NAME_None)
+					{
+						// If our leader has reached the end boundary, make sure we reach it as well
+						NextMarker.MarkerIndex = MarkerIndexSpecialValues::AnimationBoundary;
+					}
+					else if (NextMarker.MarkerIndex != MarkerIndexSpecialValues::AnimationBoundary && NumMarkersPassed > 0)
+					{
+						PrevMarker = NextMarker;
 
-				// Find our new markers
-				FMarkerPair PrevMarkerEnd;
-				FMarkerPair NextMarkerEnd;
-				GetBoundarySyncMarkersAt(GroupContext.ValidMarkers, GroupLeaderContext.SyncMarkers, NewTimelineState, PrevMarkerEnd, NextMarkerEnd);
+						// Find the next marker to match our leader
+						SeekMarkerForward(NextMarker.MarkerIndex, bIsLooping, LeaderEndPosition.NextMarkerName, SyncMarkers);
+					}
+				}
 
-				// Cache the end sync position
-				GroupContext.SyncPositionEnd = CalculateSyncPosition(GroupLeaderContext.SyncMarkers, NewTimelineState, PrevMarkerEnd, NextMarkerEnd);
-				GroupContext.LeaderEndMarker = PrevMarkerEnd;
+				if (NextMarker.MarkerIndex != MarkerIndexSpecialValues::AnimationBoundary)
+				{
+					check(SyncMarkers[NextMarker.MarkerIndex].GetName() == LeaderEndPosition.NextMarkerName);
+				}
 
-#if UE_LOG_ENABLE_MARKER_DEBUG
-				UE_LOG(LogAnimMarkerSync, Warning,
-					TEXT("[%p] Leader CurrentTime(%0.2f, %0.2f%%) MoveDelta(%0.2f) Looping(%d) [(Prev '%s' @ %d | Next '%s' @ %d | %0.2f) -> (Prev '%s' @ %d | Next '%s' @ %d | %0.2f)]"),
-					GroupLeader.TraitPtr.GetNodeInstance(),
-					GroupContext.LeaderTimelineState.GetPosition(),
-					GroupContext.LeaderTimelineState.GetPositionRatio(),
-					DeltaTime,
-					LeaderTimelineState.IsLooping() ? 1 : 0,
-					*GroupContext.SyncPositionStart.PreviousMarkerName.ToString(),
-					PrevMarkerStart.MarkerIndex,
-					*GroupContext.SyncPositionStart.NextMarkerName.ToString(),
-					NextMarkerStart.MarkerIndex,
-					GroupContext.SyncPositionStart.PositionBetweenMarkers,
-					*GroupContext.SyncPositionEnd.PreviousMarkerName.ToString(),
-					PrevMarkerEnd.MarkerIndex,
-					*GroupContext.SyncPositionEnd.NextMarkerName.ToString(),
-					NextMarkerEnd.MarkerIndex,
-					GroupContext.SyncPositionEnd.PositionBetweenMarkers);
-#endif
+				return CalculateTimeFromSyncPosition(SyncMarkers, TimelineState, PrevMarker, NextMarker, LeaderEndPosition.PositionBetweenMarkers);
 			}
 			else
 			{
-				if (GroupContext.bWasActive)
+				if (NumMarkersPassed > 0)
 				{
-					if (GroupLeaderContext.State->GroupParameters.bMatchSyncPoint)
+					// Skip over the markers that passed
+					for (int32 PassedMarkerIndex = 0; PassedMarkerIndex < NumMarkersPassed; ++PassedMarkerIndex)
 					{
-						// This is a soft join where we wish to match the current sync group position
-						const float CurrentPosition = GroupLeaderContext.TimelineState.GetPosition();
-						const float PreviousPosition = GroupContext.PreviousLeaderTimelineState.GetPositionRatio() * GroupLeaderContext.TimelineState.GetDuration();
-						DeltaTime = (CurrentPosition + DeltaTime) - PreviousPosition;
+						if (PrevMarker.MarkerIndex == MarkerIndexSpecialValues::AnimationBoundary)
+						{
+							// We still have markers left to pass but the follower ran out of markers
+							break;
+						}
+
+						NextMarker = PrevMarker;
+
+						const FName PassedMarker = MarkersPassed[PassedMarkerIndex];
+
+						// Look for the instance of our passed marker
+						SeekMarkerBackward(PrevMarker.MarkerIndex, bIsLooping, PassedMarker, SyncMarkers);
 					}
-					else
+
+					if (LeaderEndPosition.PreviousMarkerName == NAME_None)
+					{
+						// If our leader has reached the end boundary, make sure we reach it as well
+						PrevMarker.MarkerIndex = MarkerIndexSpecialValues::AnimationBoundary;
+					}
+					else if (PrevMarker.MarkerIndex != MarkerIndexSpecialValues::AnimationBoundary && NumMarkersPassed > 0)
 					{
-						// This is a hard join where the leader dictates its new position
+						NextMarker = PrevMarker;
+
+						// Find the previous marker to match our leader
+						SeekMarkerBackward(PrevMarker.MarkerIndex, bIsLooping, LeaderEndPosition.PreviousMarkerName, SyncMarkers);
 					}
 				}
-				else
+
+				if (PrevMarker.MarkerIndex != MarkerIndexSpecialValues::AnimationBoundary)
 				{
-					// This is a new sync group, leader dictates its position
+					check(SyncMarkers[PrevMarker.MarkerIndex].GetName() == LeaderEndPosition.PreviousMarkerName);
 				}
 
-				GroupSyncTrait.AdvanceBy(Context, DeltaTime);
-
-				GroupContext.LeaderTimelineState = GroupLeaderContext.TimelineTrait.GetState(Context);
-				GroupContext.LeaderDeltaTime = DeltaTime;
-
-#if UE_LOG_ENABLE_MARKER_DEBUG
-				const bool bLooped = DeltaTime >= 0.0f
-					? (GroupContext.LeaderTimelineState.GetPosition() < GroupLeaderContext.TimelineState.GetPosition())
-					: (GroupContext.LeaderTimelineState.GetPosition() > GroupLeaderContext.TimelineState.GetPosition());
-
-				UE_LOG(LogAnimMarkerSync, Warning,
-					TEXT("[%p] Leader [Time] [%0.2f + %0.2f -> %0.2f / %0.2f (%0.2f%%)]%s"),
-					GroupLeader.TraitPtr.GetNodeInstance(),
-					GroupLeaderContext.TimelineState.GetPosition(), DeltaTime,
-					GroupContext.LeaderTimelineState.GetPosition(), GroupContext.LeaderTimelineState.GetDuration(),
-					GroupContext.LeaderTimelineState.GetPositionRatio(),
-					bLooped ? TEXT(" (looped)") : TEXT(""));
-#endif
+				return CalculateTimeFromSyncPosition(SyncMarkers, TimelineState, PrevMarker, NextMarker, LeaderEndPosition.PositionBetweenMarkers);
 			}
 		}
 
-		static void AdvanceFollowers(FExecutionContext& Context, FSyncGroupContext& GroupContext)
+		static void AdvanceLeaderTimeBased(FExecutionContext& Context, FSyncGroupContext& GroupContext)
 		{
-			const int32 NumMembers = GroupContext.Members.Num();
-			if (NumMembers == 1)
-			{
-				return;	// No followers if we have a single leader in the group
-			}
+			FSyncGroupMemberContext& GroupLeaderContext = GroupContext.Members[GroupContext.LeaderIndex];
+			const FSyncGroupMember& GroupLeader = *GroupLeaderContext.State;
 
-			const int32 LeaderIndex = GroupContext.LeaderIndex;
-			const float LeaderProgressRatio = GroupContext.LeaderTimelineState.GetPositionRatio();
-			const FSyncGroupMemberContext& GroupLeaderContext = GroupContext.Members[LeaderIndex];
+			FTraitStackBinding& LeaderTraitStack = GroupLeaderContext.TraitStack;
 
 			TTraitBinding<IGroupSynchronization> GroupSyncTrait;
 
-			// Advance every follower to the same progress ratio as the leader
-			for (int32 MemberIndex = 0; MemberIndex < NumMembers; ++MemberIndex)
+			Context.BindTo(GroupLeader.TraitPtr);
+			ensure(Context.GetStack(GroupLeader.TraitPtr, LeaderTraitStack));
+			ensure(LeaderTraitStack.GetInterface(GroupSyncTrait));
+
+			// Cache our starting timeline state
+			FTimelineState LeaderStartTimelineState = GroupLeaderContext.TimelineTrait.GetState(Context);
+
+			const float PlayRate = LeaderStartTimelineState.GetPlayRate() != 0.0f ? LeaderStartTimelineState.GetPlayRate() : 1.0f;
+
+			if (GroupLeaderContext.State->bJustJoined &&
+				GroupLeaderContext.State->GroupParameters.bMatchSyncPoint &&
+				!GroupContext.State->bJustFormed)
 			{
-				if (MemberIndex == LeaderIndex)
-				{
-					continue;	// Ignore the leader, it already advanced
-				}
+				// This is a soft join where we wish to match the current sync group position
+				const float CurrentTime = LeaderStartTimelineState.GetPosition();
+				const float DesiredTime = GroupContext.PreviousLeaderTimelineState.GetPositionRatio() * LeaderStartTimelineState.GetDuration();
 
-				FSyncGroupMemberContext& GroupMemberContext = GroupContext.Members[MemberIndex];
-				const FSyncGroupMember& GroupMember = *GroupMemberContext.State;
+				// Seek where we should be without dispatching events
+				const float DeltaTime = (DesiredTime - CurrentTime) / PlayRate;
+				GroupSyncTrait.AdvanceBy(Context, DeltaTime, false);
 
-				FTraitStackBinding& MemberTraitStack = GroupMemberContext.TraitStack;
+				// Update our cached state
+				LeaderStartTimelineState = GroupLeaderContext.TimelineTrait.GetState(Context);
+			}
 
-				Context.BindTo(GroupMember.TraitPtr);
-				ensure(Context.GetStack(GroupMember.TraitPtr, MemberTraitStack));
-				ensure(MemberTraitStack.GetInterface(GroupSyncTrait));
+			// Record where the leader started from, followers that join the group will start there
+			GroupContext.LeaderStartRatio = LeaderStartTimelineState.GetPositionRatio();
 
-				// Cache our starting timeline state
-				GroupMemberContext.TimelineState = GroupMemberContext.TimelineTrait.GetState(Context);
+			// Compute our desired delta time, accounting for the play rate
+			const float DeltaTime = GroupLeader.TraitState.GetDeltaTime();
+			GroupSyncTrait.AdvanceBy(Context, DeltaTime, true);
 
-				float DeltaTime = 0.0f;
+			// Get our new state
+			const FTimelineState LeaderEndTimelineState = GroupLeaderContext.TimelineTrait.GetState(Context);
 
-#if UE_LOG_ENABLE_MARKER_DEBUG
-				const TCHAR* SyncMode = TEXT("????");
-				bool bLooped = false;
-#endif
+			GroupContext.LeaderTimelineState = LeaderEndTimelineState;
+			GroupContext.bIsLeaderPlayingForward = (DeltaTime * PlayRate) >= 0.0f;
+			GroupContext.LeaderEndRatio = LeaderEndTimelineState.GetPositionRatio();
 
-				if (GroupMember.GroupParameters.GroupRole == EAnimGroupSynchronizationRole::ExclusiveAlwaysLeader)
-				{
-					// We asked to be a leader and we weren't picked as leader
-					// We don't want to follow anyone so we tick on our own
-					DeltaTime = GroupMember.TraitState.GetDeltaTime();
+#if UE_DEBUG_SYNC_GROUPS
+			if (CDebugVarEnableSyncLog.GetValueOnAnyThread())
+			{
+				FString DeltaTimeStr = FString::Printf(TEXT("%s %0.2f"),
+					GroupContext.bIsLeaderPlayingForward ? TEXT("+") : TEXT("-"),
+					DeltaTime * FMath::Abs(PlayRate));
 
-#if UE_LOG_ENABLE_MARKER_DEBUG
-					SyncMode = TEXT("Excl");
-#endif
-				}
-				else if (GroupMemberContext.bUseMarkerSyncing)
-				{
-					DeltaTime = MatchLeaderMarkerPhase(GroupContext, GroupLeaderContext, GroupMemberContext);
+				const bool bLooped = GroupContext.bIsLeaderPlayingForward
+					? (LeaderEndTimelineState.GetPosition() < LeaderStartTimelineState.GetPosition())
+					: (LeaderStartTimelineState.GetPosition() > LeaderEndTimelineState.GetPosition());
 
-#if UE_LOG_ENABLE_MARKER_DEBUG
-					SyncMode = TEXT("Mark");
+				UE_LOG(LogAnimMarkerSync, Warning,
+					TEXT("[%s] [%p] Leader [Time] [%0.2f %s -> %0.2f / %0.2f (%3.2f%%)] Playing [%s]%s"),
+					*GroupContext.State->GroupName.ToString(),
+					GroupLeader.TraitPtr.GetNodeInstance(),
+					LeaderStartTimelineState.GetPosition(), *DeltaTimeStr,
+					LeaderEndTimelineState.GetPosition(), LeaderEndTimelineState.GetDuration(),
+					LeaderEndTimelineState.GetPositionRatio() * 100.0f,
+					*LeaderEndTimelineState.GetDebugName().ToString(),
+					bLooped ? TEXT(" (looped)") : TEXT(""));
+			}
 #endif
+		}
+
+		static void ValidateLeaderMarkers(const FTimelineState& TimelineState, const FTimelineSyncMarkerArray& SyncMarkers, const FMarkerPair& PrevMarker, const FMarkerPair& NextMarker)
+		{
+			// Must have found some markers
+			check(PrevMarker.MarkerIndex != MarkerIndexSpecialValues::Uninitialized);
+			check(NextMarker.MarkerIndex != MarkerIndexSpecialValues::Uninitialized);
+
+			if (PrevMarker.MarkerIndex == MarkerIndexSpecialValues::AnimationBoundary)
+			{
+				// If previous is the boundary, then next must be the first marker
+				check(NextMarker.MarkerIndex == 0);
+			}
+			else if (NextMarker.MarkerIndex == MarkerIndexSpecialValues::AnimationBoundary)
+			{
+				// If next is the boundary, then previous must be the last marker
+				check(PrevMarker.MarkerIndex == SyncMarkers.Num() - 1);
+			}
+			else
+			{
+				if (TimelineState.IsLooping())
+				{
+					// If previous is a valid marker, then next must be the following marker, optionally wrapping around
+					check((PrevMarker.MarkerIndex + 1) % SyncMarkers.Num() == NextMarker.MarkerIndex);
 				}
 				else
 				{
-					const float PreviousPosition = GroupMemberContext.TimelineState.GetPosition();
-					float CurrentPosition = LeaderProgressRatio * GroupMemberContext.TimelineState.GetDuration();
+					// If previous is a valid marker, then next must be the following marker or the boundary if we run out
+					check(PrevMarker.MarkerIndex + 1 == NextMarker.MarkerIndex);
+				}
+			}
+		}
+
+		// Ensures that the follower markers we found match the ones from the leader
+		// They can mismatch if the leader/follower have mismatched looping state
+		// Returns whether or not the markers were modified
+		static bool SanitizeFollowerMarkers(const FMarkerSyncAnimPosition& SyncPosition, const FTimelineSyncMarkerArray& SyncMarkers, bool bIsPlayingForward, bool bIsLooping, FMarkerPair& PrevMarker, FMarkerPair& NextMarker)
+		{
+			// Must have found some markers
+			check(PrevMarker.MarkerIndex != MarkerIndexSpecialValues::Uninitialized);
+			check(NextMarker.MarkerIndex != MarkerIndexSpecialValues::Uninitialized);
+
+			bool bAppliedFixup = false;
 
-					if (GroupContext.LeaderDeltaTime >= 0.0f)
+			if (bIsPlayingForward)
+			{
+				if (GetMarkerName(PrevMarker, SyncMarkers) != SyncPosition.PreviousMarkerName)
+				{
+					// Our previous marker doesn't match, look for the next one that matches
+					SeekMarkerForward(PrevMarker.MarkerIndex, bIsLooping, SyncPosition.PreviousMarkerName, SyncMarkers);
+
+					// Fixup our next marker index
+					NextMarker.MarkerIndex = PrevMarker.MarkerIndex + 1;
+					if (NextMarker.MarkerIndex == SyncMarkers.Num())
 					{
-						if (CurrentPosition < PreviousPosition)
-						{
-							// We must have looped around but we still want a positive delta time to match our leader
-							CurrentPosition += GroupMemberContext.TimelineState.GetDuration();
+						NextMarker.MarkerIndex = bIsLooping ? 0 : MarkerIndexSpecialValues::AnimationBoundary;
+					}
 
-#if UE_LOG_ENABLE_MARKER_DEBUG
-							bLooped = true;
-#endif
-						}
+					bAppliedFixup = true;
+				}
+
+				if (GetMarkerName(NextMarker, SyncMarkers) != SyncPosition.NextMarkerName)
+				{
+					// Our next marker doesn't match, look for the next one that matches
+					SeekMarkerForward(NextMarker.MarkerIndex, bIsLooping, SyncPosition.NextMarkerName, SyncMarkers);
+
+					bAppliedFixup = true;
+				}
+			}
+			else
+			{
+				if (GetMarkerName(NextMarker, SyncMarkers) != SyncPosition.NextMarkerName)
+				{
+					// Our next marker doesn't match, look for the next one that matches
+					SeekMarkerBackward(NextMarker.MarkerIndex, bIsLooping, SyncPosition.NextMarkerName, SyncMarkers);
+
+					// Fixup our previous marker index
+					if (NextMarker.MarkerIndex == MarkerIndexSpecialValues::AnimationBoundary ||
+						(bIsLooping && NextMarker.MarkerIndex == 0))
+					{
+						PrevMarker.MarkerIndex = SyncMarkers.Num() - 1;
 					}
 					else
 					{
-						if (CurrentPosition > PreviousPosition)
-						{
-							// We must have looped around but we still want a negative delta time to match our leader
-							CurrentPosition -= GroupMemberContext.TimelineState.GetDuration();
+						PrevMarker.MarkerIndex = NextMarker.MarkerIndex - 1;
+					}
+
+					bAppliedFixup = true;
+				}
+
+				if (GetMarkerName(PrevMarker, SyncMarkers) != SyncPosition.PreviousMarkerName)
+				{
+					// Our previous marker doesn't match, look for the next one that matches
+					SeekMarkerBackward(PrevMarker.MarkerIndex, bIsLooping, SyncPosition.PreviousMarkerName, SyncMarkers);
+
+					bAppliedFixup = true;
+				}
+			}
+
+			// Follower markers must only match in name, they might not be siblings if the pattern on the leader doesn't match the one on the follower
+			// e.g. leader with RLRL and follower with RLLRL
+			check(GetMarkerName(PrevMarker, SyncMarkers) == SyncPosition.PreviousMarkerName);
+			check(GetMarkerName(NextMarker, SyncMarkers) == SyncPosition.NextMarkerName);
+
+			return bAppliedFixup;
+		}
+
+		static void AdvanceLeaderMarkerBased(FExecutionContext& Context, FSyncGroupContext& GroupContext)
+		{
+			FSyncGroupMemberContext& GroupLeaderContext = GroupContext.Members[GroupContext.LeaderIndex];
+			const FSyncGroupMember& GroupLeader = *GroupLeaderContext.State;
+
+			FTraitStackBinding& LeaderTraitStack = GroupLeaderContext.TraitStack;
+
+			TTraitBinding<IGroupSynchronization> GroupSyncTrait;
+
+			Context.BindTo(GroupLeader.TraitPtr);
+			ensure(Context.GetStack(GroupLeader.TraitPtr, LeaderTraitStack));
+			ensure(LeaderTraitStack.GetInterface(GroupSyncTrait));
+
+			// Cache our starting timeline state
+			FTimelineState LeaderStartTimelineState = GroupLeaderContext.TimelineTrait.GetState(Context);
+
+			const float PlayRate = LeaderStartTimelineState.GetPlayRate() != 0.0f ? LeaderStartTimelineState.GetPlayRate() : 1.0f;
+
+			const bool bIsSmoothJoining = GroupLeaderContext.State->bJustJoined &&
+				GroupLeaderContext.State->GroupParameters.bMatchSyncPoint &&
+				!GroupContext.State->bJustFormed;
+			if (bIsSmoothJoining)
+			{
+				// This is a soft join where we wish to match the current sync group position
+				const float CurrentTime = LeaderStartTimelineState.GetPosition();
+				const float DesiredTime = GroupContext.PreviousLeaderTimelineState.GetPositionRatio() * LeaderStartTimelineState.GetDuration();
+
+				// Seek where we should be without dispatching events
+				const float DeltaTime = (DesiredTime - CurrentTime) / PlayRate;
+				GroupSyncTrait.AdvanceBy(Context, DeltaTime, false);
+
+				// Update our cached state
+				LeaderStartTimelineState = GroupLeaderContext.TimelineTrait.GetState(Context);
+			}
+
+			// Find the sync markers around our current position
+			FMarkerPair PrevMarkerStart;
+			FMarkerPair NextMarkerStart;
+			GetNearestSyncMarkersAt(GroupContext.ValidMarkers, GroupLeaderContext.SyncMarkers, LeaderStartTimelineState, PrevMarkerStart, NextMarkerStart);
+
+			ValidateLeaderMarkers(LeaderStartTimelineState, GroupLeaderContext.SyncMarkers, PrevMarkerStart, NextMarkerStart);
+
+			// Cache the start sync position
+			GroupContext.SyncPositionStart = CalculateSyncPosition(GroupLeaderContext.SyncMarkers, LeaderStartTimelineState, PrevMarkerStart, NextMarkerStart);
+
+			// Record where the leader started from, followers that join the group will start there
+			GroupContext.LeaderStartRatio = LeaderStartTimelineState.GetPositionRatio();
+
+			// We advance the leader by its desired delta time
+			// Compute our desired delta time, accounting for the play rate
+			const float DeltaTime = GroupLeader.TraitState.GetDeltaTime();
+			GroupSyncTrait.AdvanceBy(Context, DeltaTime, true);
+
+			// Compute our new timeline state
+			const FTimelineState LeaderEndTimelineState = GroupLeaderContext.TimelineTrait.GetState(Context);
+
+			GroupContext.LeaderTimelineState = LeaderEndTimelineState;
+			GroupContext.bIsLeaderPlayingForward = (DeltaTime * PlayRate) >= 0.0f;
 
-#if UE_LOG_ENABLE_MARKER_DEBUG
-							bLooped = true;
+			// Advance and collect the markers we passed
+			FMarkerPair PrevMarkerEnd = PrevMarkerStart;
+			FMarkerPair NextMarkerEnd = NextMarkerStart;
+			AdvanceAndCollectSyncMarkersPassed(
+				LeaderStartTimelineState, DeltaTime, GroupContext.bIsLeaderPlayingForward,
+				GroupLeaderContext.SyncMarkers, PrevMarkerEnd, NextMarkerEnd, GroupContext.MarkersPassed);
+
+			ValidateLeaderMarkers(LeaderEndTimelineState, GroupLeaderContext.SyncMarkers, PrevMarkerEnd, NextMarkerEnd);
+
+			// Cache the end sync position
+			GroupContext.SyncPositionEnd = CalculateSyncPosition(GroupLeaderContext.SyncMarkers, LeaderEndTimelineState, PrevMarkerEnd, NextMarkerEnd);
+			GroupContext.LeaderEndRatio = LeaderEndTimelineState.GetPositionRatio();
+
+#if DO_CHECK
+			if (GroupContext.bIsLeaderPlayingForward)
+			{
+				if (PrevMarkerStart.MarkerIndex == PrevMarkerEnd.MarkerIndex)
+				{
+					check(NextMarkerStart.MarkerIndex == NextMarkerEnd.MarkerIndex);
+					check(GroupContext.SyncPositionStart.PositionBetweenMarkers <= GroupContext.SyncPositionEnd.PositionBetweenMarkers);
+				}
+			}
+			else
+			{
+				if (PrevMarkerStart.MarkerIndex == PrevMarkerEnd.MarkerIndex)
+				{
+					check(NextMarkerStart.MarkerIndex == NextMarkerEnd.MarkerIndex);
+					check(GroupContext.SyncPositionStart.PositionBetweenMarkers >= GroupContext.SyncPositionEnd.PositionBetweenMarkers);
+				}
+			}
 #endif
-						}
+
+#if UE_DEBUG_SYNC_GROUPS
+			if (CDebugVarEnableSyncLog.GetValueOnAnyThread())
+			{
+				FString DeltaTimeStr = FString::Printf(TEXT("%s %0.2f"),
+					GroupContext.bIsLeaderPlayingForward ? TEXT("+") : TEXT("-"),
+					DeltaTime * FMath::Abs(PlayRate));
+
+				const bool bLooped = GroupContext.bIsLeaderPlayingForward ?
+					(LeaderEndTimelineState.GetPosition() < LeaderStartTimelineState.GetPosition()) :
+					(LeaderEndTimelineState.GetPosition() > LeaderStartTimelineState.GetPosition());
+
+				FString PassedMarkers;
+				if (!GroupContext.MarkersPassed.IsEmpty())
+				{
+					TArray<FString> MarkersPassed;
+					for (FName MarkerName : GroupContext.MarkersPassed)
+					{
+						MarkersPassed.Add(MarkerName.ToString());
 					}
 
-					DeltaTime = CurrentPosition - PreviousPosition;
+					PassedMarkers = FString::Printf(TEXT(" (passed [%s])"), *FString::Join(MarkersPassed, TEXT(", ")));
+				}
+
+				FString JoinedStatus;
+				if (GroupLeaderContext.State->bJustJoined)
+				{
+					JoinedStatus = FString::Printf(TEXT(" (joined from %0.2f)"), LeaderStartTimelineState.GetPosition());
+				}
+
+				UE_LOG(LogAnimMarkerSync, Warning,
+					TEXT("[%s] [%p] Leader [Mark] [%0.2f %s -> %0.2f / %0.2f (%3.2f%%)] [%s@%0.2f | %s@%0.2f (%3.2f%%) -> %s@%0.2f | %s@%0.2f (%3.2f%%)] Playing [%s]%s%s%s%s"),
+					*GroupContext.State->GroupName.ToString(),
+					GroupLeader.TraitPtr.GetNodeInstance(),
+					LeaderStartTimelineState.GetPosition(), *DeltaTimeStr,
+					LeaderEndTimelineState.GetPosition(), LeaderEndTimelineState.GetDuration(),
+					LeaderEndTimelineState.GetPositionRatio() * 100.0f,
+					*GroupContext.SyncPositionStart.PreviousMarkerName.ToString(),
+					GetMarkerPosition(PrevMarkerStart, GroupLeaderContext.SyncMarkers, 0.0f),
+					*GroupContext.SyncPositionStart.NextMarkerName.ToString(),
+					GetMarkerPosition(NextMarkerStart, GroupLeaderContext.SyncMarkers, LeaderEndTimelineState.GetDuration()),
+					GroupContext.SyncPositionStart.PositionBetweenMarkers * 100.0f,
+					*GroupContext.SyncPositionEnd.PreviousMarkerName.ToString(),
+					GetMarkerPosition(PrevMarkerEnd, GroupLeaderContext.SyncMarkers, 0.0f),
+					*GroupContext.SyncPositionEnd.NextMarkerName.ToString(),
+					GetMarkerPosition(NextMarkerEnd, GroupLeaderContext.SyncMarkers, LeaderEndTimelineState.GetDuration()),
+					GroupContext.SyncPositionEnd.PositionBetweenMarkers * 100.0f,
+					*LeaderEndTimelineState.GetDebugName().ToString(),
+					bLooped ? TEXT(" (looped)") : TEXT(""), *PassedMarkers,
+					bIsSmoothJoining ? TEXT(" (smooth join)") : TEXT(""),
+					*JoinedStatus);
+			}
+#endif
+		}
+
+		static void AdvanceLeader(FExecutionContext& Context, FSyncGroupContext& GroupContext)
+		{
+			bool bUseMarkerSyncing = GroupContext.bCanGroupUseMarkerSyncing;
+
+#if !(UE_BUILD_SHIPPING || UE_BUILD_TEST)
+			const int32 DebugSyncMode = CDebugVarSyncMode.GetValueOnAnyThread();
+			if (DebugSyncMode == 1)
+			{
+				bUseMarkerSyncing = false;
+			}
+#endif
+
+			if (bUseMarkerSyncing)
+			{
+				AdvanceLeaderMarkerBased(Context, GroupContext);
+			}
+			else
+			{
+				AdvanceLeaderTimeBased(Context, GroupContext);
+			}
+		}
+
+		static void AdvanceFollowerTimeBased(FExecutionContext& Context, FSyncGroupContext& GroupContext, FSyncGroupMemberContext& GroupMemberContext)
+		{
+			TTraitBinding<IGroupSynchronization> GroupSyncTrait;
 
-#if UE_LOG_ENABLE_MARKER_DEBUG
-					SyncMode = TEXT("Time");
+			const FSyncGroupMember& GroupMember = *GroupMemberContext.State;
+
+			FTraitStackBinding& MemberTraitStack = GroupMemberContext.TraitStack;
+
+			Context.BindTo(GroupMember.TraitPtr);
+			ensure(Context.GetStack(GroupMember.TraitPtr, MemberTraitStack));
+			ensure(MemberTraitStack.GetInterface(GroupSyncTrait));
+
+			// Cache our starting timeline state
+			FTimelineState FollowerStartTimelineState = GroupMemberContext.TimelineTrait.GetState(Context);
+
+			const bool bIsPlayingForward = GroupContext.bIsLeaderPlayingForward;
+			const float PlayRate = FollowerStartTimelineState.GetPlayRate() != 0.0f ? FollowerStartTimelineState.GetPlayRate() : 1.0f;
+
+			if (GroupMemberContext.State->bJustJoined)
+			{
+				// If we just joined the group as a follower then our current time is not relevant
+				// We assume that we start at the same position in normalized time as our leader
+				const float CurrentTime = FollowerStartTimelineState.GetPosition();
+				const float DesiredTime = GroupContext.LeaderStartRatio * FollowerStartTimelineState.GetDuration();
+
+				// Seek where we should be without dispatching events
+				const float DeltaTime = (DesiredTime - CurrentTime) / PlayRate;
+				GroupSyncTrait.AdvanceBy(Context, DeltaTime, false);
+
+				// Update our cached state
+				FollowerStartTimelineState = GroupMemberContext.TimelineTrait.GetState(Context);
+			}
+
+#if UE_DEBUG_SYNC_GROUPS
+			bool bLooped = false;
+#endif
+
+			const float PreviousPosition = FollowerStartTimelineState.GetPosition();
+			float CurrentPosition = GroupContext.LeaderEndRatio * FollowerStartTimelineState.GetDuration();
+
+			if (GroupContext.bIsLeaderPlayingForward)
+			{
+				if (CurrentPosition < PreviousPosition)
+				{
+					// We must have looped around but we still want a positive delta time to match our leader
+					CurrentPosition += FollowerStartTimelineState.GetDuration();
+
+#if UE_DEBUG_SYNC_GROUPS
+					bLooped = true;
+#endif
+				}
+			}
+			else
+			{
+				if (CurrentPosition > PreviousPosition)
+				{
+					// We must have looped around but we still want a negative delta time to match our leader
+					CurrentPosition -= FollowerStartTimelineState.GetDuration();
+
+#if UE_DEBUG_SYNC_GROUPS
+					bLooped = true;
 #endif
 				}
+			}
 
-				GroupSyncTrait.AdvanceBy(Context, DeltaTime);
+			// Compute our desired delta time, accounting for the follower play rate
+			const float DeltaTime = (CurrentPosition - PreviousPosition) / PlayRate;
 
-#if UE_LOG_ENABLE_MARKER_DEBUG
-				const FTimelineState CurrTimelineState = GroupMemberContext.TimelineTrait.GetState(Context);
+			GroupSyncTrait.AdvanceBy(Context, DeltaTime, true);
+
+#if UE_DEBUG_SYNC_GROUPS
+			if (CDebugVarEnableSyncLog.GetValueOnAnyThread())
+			{
+				FString DeltaTimeStr = FString::Printf(TEXT("%s %0.2f"),
+					GroupContext.bIsLeaderPlayingForward ? TEXT("+") : TEXT("-"),
+					DeltaTime * FMath::Abs(PlayRate));
+
+				const FTimelineState FollowerEndTimelineState = GroupMemberContext.TimelineTrait.GetState(Context);
 
 				UE_LOG(LogAnimMarkerSync, Warning,
-					TEXT("[%p] Follower [%s] [%0.2f + %0.2f -> %0.2f / %0.2f (%0.2f%%)]%s"),
-					GroupMember.TraitPtr.GetNodeInstance(), SyncMode,
-					GroupMemberContext.TimelineState.GetPosition(), DeltaTime,
-					CurrTimelineState.GetPosition(), CurrTimelineState.GetDuration(),
-					CurrTimelineState.GetPositionRatio(),
+					TEXT("[%s] [%p] Follow [Time] [%0.2f %s -> %0.2f / %0.2f (%3.2f%%)] Playing [%s]%s"),
+					*GroupContext.State->GroupName.ToString(),
+					GroupMember.TraitPtr.GetNodeInstance(),
+					FollowerStartTimelineState.GetPosition(), *DeltaTimeStr,
+					FollowerEndTimelineState.GetPosition(), FollowerEndTimelineState.GetDuration(),
+					FollowerEndTimelineState.GetPositionRatio() * 100.0f,
+					*FollowerEndTimelineState.GetDebugName().ToString(),
 					bLooped ? TEXT(" (looped)") : TEXT(""));
+			}
 #endif
+		}
+
+		static void AdvanceFollowerMarkerBased(FExecutionContext& Context, FSyncGroupContext& GroupContext, FSyncGroupMemberContext& GroupMemberContext)
+		{
+			const FSyncGroupMemberContext& GroupLeaderContext = GroupContext.Members[GroupContext.LeaderIndex];
+
+			const FSyncGroupMember& GroupMember = *GroupMemberContext.State;
+			FTraitStackBinding& MemberTraitStack = GroupMemberContext.TraitStack;
+
+			Context.BindTo(GroupMember.TraitPtr);
+			ensure(Context.GetStack(GroupMember.TraitPtr, MemberTraitStack));
+
+			TTraitBinding<IGroupSynchronization> GroupSyncTrait;
+			ensure(MemberTraitStack.GetInterface(GroupSyncTrait));
+
+			// Cache our starting timeline state
+			FTimelineState FollowerStartTimelineState = GroupMemberContext.TimelineTrait.GetState(Context);
+
+			const bool bIsPlayingForward = GroupContext.bIsLeaderPlayingForward;
+			const float PlayRate = FollowerStartTimelineState.GetPlayRate() != 0.0f ? FollowerStartTimelineState.GetPlayRate() : 1.0f;
+			const float StartTime = FollowerStartTimelineState.GetPosition();
+
+			const bool bIsSmoothJoining = GroupMemberContext.State->bJustJoined && GroupMemberContext.State->GroupParameters.bMatchSyncPoint;
+			if (bIsSmoothJoining)
+			{
+				// If we just joined the group as a follower then our current time is not relevant
+				// We assume that we start at the same position in normalized time as our leader
+				// This is where we'll start searching for valid markers, we'll perform a final fix-up afterwards
+				const float CurrentTime = FollowerStartTimelineState.GetPosition();
+				const float DesiredTime = GroupContext.LeaderStartRatio * FollowerStartTimelineState.GetDuration();
+
+				// Seek where we should be without dispatching events
+				FollowerStartTimelineState = FollowerStartTimelineState.WithPosition(DesiredTime);
+			}
+
+#if UE_DEBUG_SYNC_GROUPS
+			const float MarkerSearchTime = FollowerStartTimelineState.GetPosition();
+#endif
+
+			// Find the nearest markers matching our current time
+			FMarkerPair PrevMarkerStart;
+			FMarkerPair NextMarkerStart;
+			GetNearestMatchingBoundarySyncMarkers(GroupContext.SyncPositionStart, GroupMemberContext.SyncMarkers, FollowerStartTimelineState, PrevMarkerStart, NextMarkerStart);
+
+			FMarkerPair PrevMarkerSanitized = PrevMarkerStart;
+			FMarkerPair NextMarkerSanitized = NextMarkerStart;
+			bool bHadInvalidMarkers = SanitizeFollowerMarkers(GroupContext.SyncPositionStart, GroupMemberContext.SyncMarkers, bIsPlayingForward, FollowerStartTimelineState.IsLooping(), PrevMarkerSanitized, NextMarkerSanitized);
+
+			if (GroupMemberContext.State->bJustJoined)
+			{
+				// If we just joined, we started searching for our markers approximately where the leader was
+				// However, we might end up finding markers before/after where we should be and as a result can
+				// end up with a slightly positive/negative delta time moving opposite the desired direction
+				// To avoid this, now that we've found good markers, we seek again to the position the leader
+				// started at between them
+				const float CurrentTime = StartTime;
+				const float DesiredTime = CalculateTimeFromSyncPosition(GroupMemberContext.SyncMarkers, FollowerStartTimelineState,
+					PrevMarkerSanitized, NextMarkerSanitized, GroupContext.SyncPositionStart.PositionBetweenMarkers);;
+
+				// Seek where we should be without dispatching events
+				const float DeltaTime = (DesiredTime - CurrentTime) / PlayRate;
+				GroupSyncTrait.AdvanceBy(Context, DeltaTime, false);
+
+				// Update our cached state
+				FollowerStartTimelineState = GroupMemberContext.TimelineTrait.GetState(Context);
+			}
+
+			const float CurrentTime = FollowerStartTimelineState.GetPosition();
+
+			FMarkerPair PrevMarkerEnd = PrevMarkerSanitized;
+			FMarkerPair NextMarkerEnd = NextMarkerSanitized;
+			float DesiredTime = AdvanceSyncMarkersPassed(
+				FollowerStartTimelineState, bIsPlayingForward, GroupContext.SyncPositionEnd,
+				GroupMemberContext.SyncMarkers, GroupContext.MarkersPassed, PrevMarkerEnd, NextMarkerEnd);
+
+			if (bIsPlayingForward)
+			{
+				if (DesiredTime < CurrentTime)
+				{
+					// We are looping around, wrap around to ensure our delta time takes us past the end of the timeline
+					// The player will handle looping internally
+					DesiredTime += FollowerStartTimelineState.GetDuration();
+				}
+			}
+			else
+			{
+				if (DesiredTime > CurrentTime)
+				{
+					// We are looping around, wrap around to ensure our delta time takes us past the start of the timeline
+					// The played will handle looping internally
+					DesiredTime -= FollowerStartTimelineState.GetDuration();
+				}
+			}
+
+			// Compute our desired delta time, accounting for the follower play rate
+			const float DeltaTime = (DesiredTime - CurrentTime) / PlayRate;
+			GroupSyncTrait.AdvanceBy(Context, DeltaTime, true);
+
+#if UE_DEBUG_SYNC_GROUPS
+			if (CDebugVarEnableSyncLog.GetValueOnAnyThread())
+			{
+				FString DeltaTimeStr = FString::Printf(TEXT("%s %0.2f"),
+					GroupContext.bIsLeaderPlayingForward ? TEXT("+") : TEXT("-"),
+					DeltaTime * FMath::Abs(PlayRate));
+
+				const FTimelineState FollowerEndTimelineState = GroupMemberContext.TimelineTrait.GetState(Context);
+
+				const bool bLooped = GroupContext.bIsLeaderPlayingForward ?
+					(FollowerEndTimelineState.GetPosition() < CurrentTime) :
+					(FollowerEndTimelineState.GetPosition() > CurrentTime);
+
+				const bool bUnexpectedDeltaTime = DeltaTime < 0.0f || DeltaTime > 0.3f;
+
+				FString SanitizedResult;
+				if (bHadInvalidMarkers)
+				{
+					SanitizedResult = FString::Printf(TEXT(" (fixed invalid markers [%s@%d | %s@%d] -> [%s@%d | %s@%d])"),
+						*GetMarkerName(PrevMarkerStart, GroupMemberContext.SyncMarkers).ToString(),
+						PrevMarkerStart.MarkerIndex,
+						*GetMarkerName(NextMarkerStart, GroupMemberContext.SyncMarkers).ToString(),
+						NextMarkerStart.MarkerIndex,
+						*GetMarkerName(PrevMarkerSanitized, GroupMemberContext.SyncMarkers).ToString(),
+						PrevMarkerSanitized.MarkerIndex,
+						*GetMarkerName(NextMarkerSanitized, GroupMemberContext.SyncMarkers).ToString(),
+						NextMarkerSanitized.MarkerIndex);
+				}
+
+				FString JoinedStatus;
+				if (GroupMemberContext.State->bJustJoined)
+				{
+					JoinedStatus = FString::Printf(TEXT(" (joined from %0.2f)"), MarkerSearchTime);
+				}
+
+				UE_LOG(LogAnimMarkerSync, Warning,
+					TEXT("[%s] [%p] Follow [Mark] [%0.2f %s -> %0.2f / %0.2f (%3.2f%%)] [%s@%0.2f | %s@%0.2f (%3.2f%%) -> %s@%0.2f | %s@%0.2f (%3.2f%%)] Playing [%s]%s%s%s%s%s"),
+					*GroupContext.State->GroupName.ToString(),
+					GroupMember.TraitPtr.GetNodeInstance(),
+					CurrentTime, *DeltaTimeStr,
+					FollowerEndTimelineState.GetPosition(), FollowerEndTimelineState.GetDuration(),
+					FollowerEndTimelineState.GetPositionRatio() * 100.0f,
+					*GetMarkerName(PrevMarkerStart, GroupMemberContext.SyncMarkers).ToString(),
+					GetMarkerPosition(PrevMarkerStart, GroupMemberContext.SyncMarkers, 0.0f),
+					*GetMarkerName(NextMarkerStart, GroupMemberContext.SyncMarkers).ToString(),
+					GetMarkerPosition(NextMarkerStart, GroupMemberContext.SyncMarkers, FollowerEndTimelineState.GetDuration()),
+					GroupContext.SyncPositionStart.PositionBetweenMarkers * 100.0f,
+					*GetMarkerName(PrevMarkerEnd, GroupMemberContext.SyncMarkers).ToString(),
+					GetMarkerPosition(PrevMarkerEnd, GroupMemberContext.SyncMarkers, 0.0f),
+					*GetMarkerName(NextMarkerEnd, GroupMemberContext.SyncMarkers).ToString(),
+					GetMarkerPosition(NextMarkerEnd, GroupMemberContext.SyncMarkers, FollowerEndTimelineState.GetDuration()),
+					GroupContext.SyncPositionEnd.PositionBetweenMarkers * 100.0f,
+					*FollowerEndTimelineState.GetDebugName().ToString(),
+					bLooped ? TEXT(" (looped)") : TEXT(""),
+					bUnexpectedDeltaTime ? TEXT(" (unusual delta time)") : TEXT(""),
+					bHadInvalidMarkers ? *SanitizedResult : TEXT(""),
+					bIsSmoothJoining ? TEXT(" (smooth join)") : TEXT(""),
+					*JoinedStatus);
+			}
+#endif
+		}
+
+		static void AdvanceFollowerUngrouped(FExecutionContext& Context, FSyncGroupContext& GroupContext, FSyncGroupMemberContext& GroupMemberContext)
+		{
+			const FSyncGroupMember& GroupMember = *GroupMemberContext.State;
+
+			FTraitStackBinding& MemberTraitStack = GroupMemberContext.TraitStack;
+
+			Context.BindTo(GroupMember.TraitPtr);
+			ensure(Context.GetStack(GroupMember.TraitPtr, MemberTraitStack));
+
+			TTraitBinding<IGroupSynchronization> GroupSyncTrait;
+			ensure(MemberTraitStack.GetInterface(GroupSyncTrait));
+
+#if UE_DEBUG_SYNC_GROUPS
+			const FTimelineState FollowerStartTimelineState = GroupMemberContext.TimelineTrait.GetState(Context);
+#endif
+
+			const float DeltaTime = GroupMember.TraitState.GetDeltaTime();
+			GroupSyncTrait.AdvanceBy(Context, DeltaTime, true);
+
+#if UE_DEBUG_SYNC_GROUPS
+			if (CDebugVarEnableSyncLog.GetValueOnAnyThread())
+			{
+				FString DeltaTimeStr = FString::Printf(TEXT("%s %0.2f"),
+					GroupContext.bIsLeaderPlayingForward ? TEXT("+") : TEXT("-"),
+					DeltaTime * FMath::Abs(FollowerStartTimelineState.GetPlayRate()));
+
+				const FTimelineState FollowerEndTimelineState = GroupMemberContext.TimelineTrait.GetState(Context);
+
+				UE_LOG(LogAnimMarkerSync, Warning,
+					TEXT("[%s] [%p] Follow [Solo] [%0.2f %s -> %0.2f / %0.2f (%3.2f%%)] Playing [%s]"),
+					*GroupContext.State->GroupName.ToString(),
+					GroupMember.TraitPtr.GetNodeInstance(),
+					FollowerStartTimelineState.GetPosition(), *DeltaTimeStr,
+					FollowerEndTimelineState.GetPosition(), FollowerEndTimelineState.GetDuration(),
+					FollowerEndTimelineState.GetPositionRatio() * 100.0f,
+					*FollowerEndTimelineState.GetDebugName().ToString());
+			}
+#endif
+		}
+
+		static void AdvanceFollowers(FExecutionContext& Context, FSyncGroupContext& GroupContext)
+		{
+			const int32 NumMembers = GroupContext.Members.Num();
+			if (NumMembers == 1)
+			{
+				return;	// No followers if we have a single leader in the group
+			}
+
+			const int32 LeaderIndex = GroupContext.LeaderIndex;
+
+			// Advance every follower to the same progress ratio as the leader
+			for (int32 MemberIndex = 0; MemberIndex < NumMembers; ++MemberIndex)
+			{
+				if (MemberIndex == LeaderIndex)
+				{
+					continue;	// Ignore the leader, it already advanced
+				}
+
+				FSyncGroupMemberContext& GroupMemberContext = GroupContext.Members[MemberIndex];
+				const FSyncGroupMember& GroupMember = *GroupMemberContext.State;
+
+				bool bUseMarkerSyncing = GroupMemberContext.bUseMarkerSyncing;
+				EAnimGroupSynchronizationRole GroupRole = GroupMember.GroupParameters.GroupRole;
+
+#if !(UE_BUILD_SHIPPING || UE_BUILD_TEST)
+				const int32 DebugSyncMode = CDebugVarSyncMode.GetValueOnAnyThread();
+				if (DebugSyncMode == 1)
+				{
+					bUseMarkerSyncing = false;
+				}
+				else if (DebugSyncMode == 2)
+				{
+					GroupRole = EAnimGroupSynchronizationRole::ExclusiveAlwaysLeader;
+				}
+#endif
+
+				if (GroupRole == EAnimGroupSynchronizationRole::ExclusiveAlwaysLeader)
+				{
+					// We asked to be a leader and we weren't picked as leader
+					// We don't want to follow anyone so we tick on our own
+					AdvanceFollowerUngrouped(Context, GroupContext, GroupMemberContext);
+				}
+				else if (bUseMarkerSyncing)
+				{
+					AdvanceFollowerMarkerBased(Context, GroupContext, GroupMemberContext);
+				}
+				else
+				{
+					AdvanceFollowerTimeBased(Context, GroupContext, GroupMemberContext);
+				}
 			}
 		}
 	}
@@ -1082,11 +1774,45 @@ namespace UE::AnimNext
 
 	void FSyncGroupGraphInstanceComponent::RegisterWithGroup(const FSyncGroupParameters& GroupParameters, const FWeakTraitPtr& TraitPtr, const FTraitUpdateState& TraitState)
 	{
+		using namespace Private;
+
 		const int32 GroupIndex = SyncGroupMap.FindOrAdd(GroupParameters.GroupName, SyncGroups.Num());
-		Private::FSyncGroupState& GroupState = GroupIndex == SyncGroups.Num() ? SyncGroups.AddDefaulted_GetRef() : SyncGroups[GroupIndex];
+		const bool bJustFormed = GroupIndex == SyncGroups.Num();
+
+		FSyncGroupState& GroupState = bJustFormed ? SyncGroups.AddDefaulted_GetRef() : SyncGroups[GroupIndex];
 		GroupState.GroupName = GroupParameters.GroupName;
-		GroupState.Members.Add(Private::FSyncGroupMember{ TraitState, TraitPtr, GroupParameters });
-		MaxNumMembers = FMath::Max<uint32>(MaxNumMembers, GroupState.Members.Num());
+		GroupState.bIsActive = true;
+		GroupState.bJustFormed = bJustFormed;
+
+		bool bIsMemberNew = true;
+		for (FSyncGroupMember& GroupMember : GroupState.Members)
+		{
+			if (GroupMember.TraitPtr == TraitPtr)
+			{
+				bIsMemberNew = false;
+				check(!GroupMember.bIsActive);
+
+				GroupMember.TraitState = TraitState;
+				GroupMember.GroupParameters = GroupParameters;
+				GroupMember.bIsActive = true;
+				break;
+			}
+		}
+
+		if (bIsMemberNew)
+		{
+			GroupState.Members.Add(FSyncGroupMember{ TraitState, FTraitPtr(TraitPtr), GroupParameters, true, true });
+
+#if UE_DEBUG_SYNC_GROUPS
+			if (CDebugVarEnableSyncLog.GetValueOnAnyThread())
+			{
+				UE_LOG(LogAnimMarkerSync, Warning,
+					TEXT("[%s] [%p] Joined"),
+					*GroupParameters.GroupName.ToString(),
+					TraitPtr.GetNodeInstance());
+			}
+#endif
+		}
 	}
 
 	FName FSyncGroupGraphInstanceComponent::CreateUniqueGroupName()
@@ -1135,47 +1861,79 @@ namespace UE::AnimNext
 
 	void FSyncGroupGraphInstanceComponent::PreUpdate(FExecutionContext& Context)
 	{
-		bool bRebuildGroupIndexMap = false;
+		using namespace Private;
 
-		const int32 NumPreviousSyncGroups = SyncGroups.Num();
-		for (int32 SyncGroupIndex = NumPreviousSyncGroups - 1; SyncGroupIndex >= 0; --SyncGroupIndex)
+		for (FSyncGroupState& GroupState : SyncGroups)
 		{
-			Private::FSyncGroupState& GroupState = SyncGroups[SyncGroupIndex];
-
-			const bool bWasActive = GroupState.Members.IsEmpty();
-			if (!GroupState.bWasActive && !bWasActive)
+			for (FSyncGroupMember& GroupMember : GroupState.Members)
 			{
-				// We have been inactive for two consecutive updates, time to purge this group
-				SyncGroups.RemoveAtSwap(SyncGroupIndex, EAllowShrinking::No);
-				bRebuildGroupIndexMap = true;
+				GroupMember.bIsActive = false;
+				GroupMember.bJustJoined = false;
 			}
-			else
-			{
-				GroupState.bWasActive = bWasActive;
-				GroupState.Members.Reset();
-			}
-		}
-
-		if (bRebuildGroupIndexMap)
-		{
-			// We purged one or more inactive groups, rebuild our map
-			SyncGroupMap.Reset();
 
-			const int32 NumSyncGroups = SyncGroups.Num();
-			for (int32 SyncGroupIndex = 0; SyncGroupIndex < NumSyncGroups; ++SyncGroupIndex)
-			{
-				const Private::FSyncGroupState& GroupState = SyncGroups[SyncGroupIndex];
-				SyncGroupMap.Add(GroupState.GroupName, SyncGroupIndex);
-			}
+			GroupState.bIsActive = false;
+			GroupState.bJustFormed = false;
 		}
-
-		MaxNumMembers = 0;
 	}
 
 	void FSyncGroupGraphInstanceComponent::PostUpdate(FExecutionContext& Context)
 	{
 		using namespace Private;
 
+		// First purge stale entries and update our bookkeeping
+		int32 MaxNumMembers = 0;
+		{
+			bool bRebuildGroupIndexMap = false;
+			for (int32 SyncGroupIndex = SyncGroups.Num() - 1; SyncGroupIndex >= 0; --SyncGroupIndex)
+			{
+				FSyncGroupState& GroupState = SyncGroups[SyncGroupIndex];
+
+				if (GroupState.bIsActive)
+				{
+					for (int32 GroupMemberIndex = GroupState.Members.Num() - 1; GroupMemberIndex >= 0; GroupMemberIndex--)
+					{
+						FSyncGroupMember& GroupMember = GroupState.Members[GroupMemberIndex];
+						if (!GroupMember.bIsActive)
+						{
+#if UE_DEBUG_SYNC_GROUPS
+							if (CDebugVarEnableSyncLog.GetValueOnAnyThread())
+							{
+								UE_LOG(LogAnimMarkerSync, Warning,
+									TEXT("[%s] [%p] Left"),
+									*GroupState.GroupName.ToString(),
+									GroupMember.TraitPtr.GetNodeInstance());
+							}
+#endif
+
+							// This member is no longer active, remove it
+							GroupState.Members.RemoveAtSwap(GroupMemberIndex, EAllowShrinking::No);
+						}
+					}
+
+					MaxNumMembers = FMath::Max<uint32>(MaxNumMembers, GroupState.Members.Num());
+				}
+				else
+				{
+					// This group is no longer active, remove it
+					SyncGroups.RemoveAtSwap(SyncGroupIndex, EAllowShrinking::No);
+					bRebuildGroupIndexMap = true;
+				}
+			}
+
+			if (bRebuildGroupIndexMap)
+			{
+				// We purged one or more inactive groups, rebuild our map
+				SyncGroupMap.Reset();
+
+				const int32 NumSyncGroups = SyncGroups.Num();
+				for (int32 SyncGroupIndex = 0; SyncGroupIndex < NumSyncGroups; ++SyncGroupIndex)
+				{
+					const FSyncGroupState& GroupState = SyncGroups[SyncGroupIndex];
+					SyncGroupMap.Add(GroupState.GroupName, SyncGroupIndex);
+				}
+			}
+		}
+
 		TArray<FSyncGroupMemberContext, TMemStackAllocator<>> MemberContexts;
 		MemberContexts.AddDefaulted(MaxNumMembers);
 
@@ -1184,14 +1942,6 @@ namespace UE::AnimNext
 		// Now that we have discovered all groups and their memberships, we can perform synchronization
 		for (FSyncGroupState& GroupState : SyncGroups)
 		{
-			if (GroupState.Members.IsEmpty())
-			{
-				// Skip empty groups
-				// A group can be empty if it became inactive in the previous frame
-				check(!GroupState.bWasActive);
-				continue;
-			}
-
 			// Initialize our context for this new group
 			InitGroup(GroupState, MemberContexts, GroupContext);
 
@@ -1210,7 +1960,12 @@ namespace UE::AnimNext
 
 			// Retain the leader progress as we might need it if a new leader joins the group during the next update
 			GroupState.PreviousLeaderTimelineState = GroupContext.LeaderTimelineState;
-			GroupState.PreviousLeaderSyncPosition = GroupContext.SyncPositionEnd;
+
+#if UE_DEBUG_SYNC_GROUPS
+			GroupState.PreviousLeaderIndex = GroupContext.LeaderIndex;
+#endif
 		}
 	}
 }
+
+#undef UE_DEBUG_SYNC_GROUPS
diff --git a/Engine/Plugins/Experimental/AnimNextAnimGraph/Source/AnimNextAnimGraph/Private/Graph/SyncGroup_GraphInstanceComponent.h b/Engine/Plugins/Experimental/AnimNextAnimGraph/Source/AnimNextAnimGraph/Private/Graph/SyncGroup_GraphInstanceComponent.h
--- a/Engine/Plugins/Experimental/AnimNextAnimGraph/Source/AnimNextAnimGraph/Private/Graph/SyncGroup_GraphInstanceComponent.h
+++ b/Engine/Plugins/Experimental/AnimNextAnimGraph/Source/AnimNextAnimGraph/Private/Graph/SyncGroup_GraphInstanceComponent.h
@@ -57,9 +57,6 @@ namespace UE::AnimNext
 		// The map of currently used group names
 		TMap<FName, Private::FSyncGroupUniqueName*> UsedUniqueGroupNames;
 
-		// The number of members in the largest group
-		uint32 MaxNumMembers = 0;
-
 		// A counter tracking the next unique group name to allocate
 		int32 UniqueGroupNameCounter = 0;
 	};
diff --git a/Engine/Plugins/Experimental/AnimNextAnimGraph/Source/AnimNextAnimGraph/Private/TraitInterfaces/IGroupSynchronization.cpp b/Engine/Plugins/Experimental/AnimNextAnimGraph/Source/AnimNextAnimGraph/Private/TraitInterfaces/IGroupSynchronization.cpp
--- a/Engine/Plugins/Experimental/AnimNextAnimGraph/Source/AnimNextAnimGraph/Private/TraitInterfaces/IGroupSynchronization.cpp
+++ b/Engine/Plugins/Experimental/AnimNextAnimGraph/Source/AnimNextAnimGraph/Private/TraitInterfaces/IGroupSynchronization.cpp
@@ -32,12 +32,12 @@ namespace UE::AnimNext
 		return FSyncGroupParameters();
 	}
 
-	void IGroupSynchronization::AdvanceBy(FExecutionContext& Context, const TTraitBinding<IGroupSynchronization>& Binding, float DeltaTime) const
+	void IGroupSynchronization::AdvanceBy(FExecutionContext& Context, const TTraitBinding<IGroupSynchronization>& Binding, float DeltaTime, bool bDispatchEvents) const
 	{
 		TTraitBinding<IGroupSynchronization> SuperBinding;
 		if (Binding.GetStackInterfaceSuper(SuperBinding))
 		{
-			SuperBinding.AdvanceBy(Context, DeltaTime);
+			SuperBinding.AdvanceBy(Context, DeltaTime, bDispatchEvents);
 		}
 	}
 }
diff --git a/Engine/Plugins/Experimental/AnimNextAnimGraph/Source/AnimNextAnimGraph/Private/TraitInterfaces/ITimelinePlayer.cpp b/Engine/Plugins/Experimental/AnimNextAnimGraph/Source/AnimNextAnimGraph/Private/TraitInterfaces/ITimelinePlayer.cpp
--- a/Engine/Plugins/Experimental/AnimNextAnimGraph/Source/AnimNextAnimGraph/Private/TraitInterfaces/ITimelinePlayer.cpp
+++ b/Engine/Plugins/Experimental/AnimNextAnimGraph/Source/AnimNextAnimGraph/Private/TraitInterfaces/ITimelinePlayer.cpp
@@ -21,12 +21,12 @@ namespace UE::AnimNext
 	}
 #endif // WITH_EDITOR
 
-	void ITimelinePlayer::AdvanceBy(FExecutionContext& Context, const TTraitBinding<ITimelinePlayer>& Binding, float DeltaTime) const
+	void ITimelinePlayer::AdvanceBy(FExecutionContext& Context, const TTraitBinding<ITimelinePlayer>& Binding, float DeltaTime, bool bDispatchEvents) const
 	{
 		TTraitBinding<ITimelinePlayer> SuperBinding;
 		if (Binding.GetStackInterfaceSuper(SuperBinding))
 		{
-			SuperBinding.AdvanceBy(Context, DeltaTime);
+			SuperBinding.AdvanceBy(Context, DeltaTime, bDispatchEvents);
 		}
 	}
 }
diff --git a/Engine/Plugins/Experimental/AnimNextAnimGraph/Source/AnimNextAnimGraph/Private/Traits/NotifyDispatcher.cpp b/Engine/Plugins/Experimental/AnimNextAnimGraph/Source/AnimNextAnimGraph/Private/Traits/NotifyDispatcher.cpp
--- a/Engine/Plugins/Experimental/AnimNextAnimGraph/Source/AnimNextAnimGraph/Private/Traits/NotifyDispatcher.cpp
+++ b/Engine/Plugins/Experimental/AnimNextAnimGraph/Source/AnimNextAnimGraph/Private/Traits/NotifyDispatcher.cpp
@@ -29,7 +29,7 @@ namespace UE::AnimNext
 	#undef TRAIT_REQUIRED_INTERFACE_ENUMERATOR
 	#undef TRAIT_INTERFACE_ENUMERATOR
 
-	void FNotifyDispatcherTrait::AdvanceBy(FExecutionContext& Context, const TTraitBinding<ITimelinePlayer>& Binding, float DeltaTime) const
+	void FNotifyDispatcherTrait::AdvanceBy(FExecutionContext& Context, const TTraitBinding<ITimelinePlayer>& Binding, float DeltaTime, bool bDispatchEvents) const
 	{
 		TTraitBinding<ITimeline> TimelineTrait;
 		ensure(Binding.GetStackInterfaceSuper(TimelineTrait));
@@ -37,12 +37,18 @@ namespace UE::AnimNext
 		TTraitBinding<ITimelinePlayer> TimelinePlayerTrait;
 		ensure(Binding.GetStackInterfaceSuper(TimelinePlayerTrait));
 
+		// Get current state from the stack, advance time, then get the delta state
+		const FTimelineState PreAdvanceState = TimelineTrait.GetState(Context);
+		TimelinePlayerTrait.AdvanceBy(Context, DeltaTime, bDispatchEvents);
+
+		if (!bDispatchEvents)
+		{
+			return;
+		}
+
 		TTraitBinding<INotifySource> NotifySourceTrait;
 		ensure(Binding.GetStackInterfaceSuper(NotifySourceTrait));
 
-		// Get current state from the stack, advance time, then get the delta state
-		const FTimelineState PreAdvanceState = TimelineTrait.GetState(Context);
-		TimelinePlayerTrait.AdvanceBy(Context, DeltaTime);
 		const FTimelineDelta Delta = TimelineTrait.GetDelta(Context);
 
 		// Query for notifies
diff --git a/Engine/Plugins/Experimental/AnimNextAnimGraph/Source/AnimNextAnimGraph/Private/Traits/NotifyDispatcher.h b/Engine/Plugins/Experimental/AnimNextAnimGraph/Source/AnimNextAnimGraph/Private/Traits/NotifyDispatcher.h
--- a/Engine/Plugins/Experimental/AnimNextAnimGraph/Source/AnimNextAnimGraph/Private/Traits/NotifyDispatcher.h
+++ b/Engine/Plugins/Experimental/AnimNextAnimGraph/Source/AnimNextAnimGraph/Private/Traits/NotifyDispatcher.h
@@ -79,7 +79,7 @@ namespace UE::AnimNext
 		};
 
 		// ITimelinePlayer impl
-		virtual void AdvanceBy(FExecutionContext& Context, const TTraitBinding<ITimelinePlayer>& Binding, float DeltaTime) const override;
+		virtual void AdvanceBy(FExecutionContext& Context, const TTraitBinding<ITimelinePlayer>& Binding, float DeltaTime, bool bDispatchEvents) const override;
 	};
 
 	// Event that dispatches to the module
diff --git a/Engine/Plugins/Experimental/AnimNextAnimGraph/Source/AnimNextAnimGraph/Private/Traits/SequencePlayer.cpp b/Engine/Plugins/Experimental/AnimNextAnimGraph/Source/AnimNextAnimGraph/Private/Traits/SequencePlayer.cpp
--- a/Engine/Plugins/Experimental/AnimNextAnimGraph/Source/AnimNextAnimGraph/Private/Traits/SequencePlayer.cpp
+++ b/Engine/Plugins/Experimental/AnimNextAnimGraph/Source/AnimNextAnimGraph/Private/Traits/SequencePlayer.cpp
@@ -100,7 +100,7 @@ namespace UE::AnimNext
 				InstanceData->InternalTimeAccumulator,
 				AnimSeq->GetPlayLength(),
 				SharedData->GetPlayRate(Binding),
-				SharedData->GetbLoop(Binding));
+				SharedData->GetbLoop(Binding)).WithDebugName(AnimSeq->GetFName());
 		}
 
 		return FTimelineState();
@@ -119,7 +119,7 @@ namespace UE::AnimNext
 		return FTimelineDelta();
 	}
 
-	void FSequencePlayerTrait::AdvanceBy(FExecutionContext& Context, const TTraitBinding<ITimelinePlayer>& Binding, float DeltaTime) const
+	void FSequencePlayerTrait::AdvanceBy(FExecutionContext& Context, const TTraitBinding<ITimelinePlayer>& Binding, float DeltaTime, bool bDispatchEvents) const
 	{
 		FInstanceData* InstanceData = Binding.GetInstanceData<FInstanceData>();
 
@@ -169,7 +169,7 @@ namespace UE::AnimNext
 		TTraitBinding<ITimelinePlayer> TimelinePlayerTrait;
 		Binding.GetStackInterface(TimelinePlayerTrait);
 
-		TimelinePlayerTrait.AdvanceBy(Context, TraitState.GetDeltaTime());
+		TimelinePlayerTrait.AdvanceBy(Context, TraitState.GetDeltaTime(), true);
 	}
 
 	void FSequencePlayerTrait::AddReferencedObjects(const FExecutionContext& Context, const TTraitBinding<IGarbageCollection>& Binding, FReferenceCollector& Collector) const
diff --git a/Engine/Plugins/Experimental/AnimNextAnimGraph/Source/AnimNextAnimGraph/Private/Traits/SequencePlayer.h b/Engine/Plugins/Experimental/AnimNextAnimGraph/Source/AnimNextAnimGraph/Private/Traits/SequencePlayer.h
--- a/Engine/Plugins/Experimental/AnimNextAnimGraph/Source/AnimNextAnimGraph/Private/Traits/SequencePlayer.h
+++ b/Engine/Plugins/Experimental/AnimNextAnimGraph/Source/AnimNextAnimGraph/Private/Traits/SequencePlayer.h
@@ -90,7 +90,7 @@ namespace UE::AnimNext
 		virtual FTimelineDelta GetDelta(const FExecutionContext& Context, const TTraitBinding<ITimeline>& Binding) const override;
 
 		// ITimelinePlayer impl
-		virtual void AdvanceBy(FExecutionContext& Context, const TTraitBinding<ITimelinePlayer>& Binding, float DeltaTime) const override;
+		virtual void AdvanceBy(FExecutionContext& Context, const TTraitBinding<ITimelinePlayer>& Binding, float DeltaTime, bool bDispatchEvents) const override;
 
 		// IUpdate impl
 		virtual void OnBecomeRelevant(FUpdateTraversalContext& Context, const TTraitBinding<IUpdate>& Binding, const FTraitUpdateState& TraitState) const override;
diff --git a/Engine/Plugins/Experimental/AnimNextAnimGraph/Source/AnimNextAnimGraph/Private/Traits/SynchronizeUsingGroups.cpp b/Engine/Plugins/Experimental/AnimNextAnimGraph/Source/AnimNextAnimGraph/Private/Traits/SynchronizeUsingGroups.cpp
--- a/Engine/Plugins/Experimental/AnimNextAnimGraph/Source/AnimNextAnimGraph/Private/Traits/SynchronizeUsingGroups.cpp
+++ b/Engine/Plugins/Experimental/AnimNextAnimGraph/Source/AnimNextAnimGraph/Private/Traits/SynchronizeUsingGroups.cpp
@@ -129,7 +129,7 @@ namespace UE::AnimNext
 		return Parameters;
 	}
 
-	void FSynchronizeUsingGroupsTrait::AdvanceBy(FExecutionContext& Context, const TTraitBinding<IGroupSynchronization>& Binding, float DeltaTime) const
+	void FSynchronizeUsingGroupsTrait::AdvanceBy(FExecutionContext& Context, const TTraitBinding<IGroupSynchronization>& Binding, float DeltaTime, bool bDispatchEvents) const
 	{
 		FInstanceData* InstanceData = Binding.GetInstanceData<FInstanceData>();
 
@@ -144,12 +144,12 @@ namespace UE::AnimNext
 		TTraitBinding<ITimelinePlayer> TimelinePlayerTrait;
 		Binding.GetStackInterface(TimelinePlayerTrait);
 
-		TimelinePlayerTrait.AdvanceBy(Context, DeltaTime);
+		TimelinePlayerTrait.AdvanceBy(Context, DeltaTime, bDispatchEvents);
 
 		InstanceData->bFreezeTimeline = true;
 	}
 
-	void FSynchronizeUsingGroupsTrait::AdvanceBy(FExecutionContext& Context, const TTraitBinding<ITimelinePlayer>& Binding, float DeltaTime) const
+	void FSynchronizeUsingGroupsTrait::AdvanceBy(FExecutionContext& Context, const TTraitBinding<ITimelinePlayer>& Binding, float DeltaTime, bool bDispatchEvents) const
 	{
 		const FInstanceData* InstanceData = Binding.GetInstanceData<FInstanceData>();
 		if (InstanceData->bFreezeTimeline)
@@ -157,7 +157,7 @@ namespace UE::AnimNext
 			return;	// If the timeline is frozen, we don't advance
 		}
 
-		ITimelinePlayer::AdvanceBy(Context, Binding, DeltaTime);
+		ITimelinePlayer::AdvanceBy(Context, Binding, DeltaTime, bDispatchEvents);
 	}
 
 	void FSynchronizeUsingGroupsTrait::CreatePayloadForObject(FExecutionContext& Context, const TTraitBinding<IGraphFactory>& Binding, const UObject* InObject, FAnimNextDataInterfacePayload& InOutPayload) const
diff --git a/Engine/Plugins/Experimental/AnimNextAnimGraph/Source/AnimNextAnimGraph/Private/Traits/SynchronizeUsingGroups.h b/Engine/Plugins/Experimental/AnimNextAnimGraph/Source/AnimNextAnimGraph/Private/Traits/SynchronizeUsingGroups.h
--- a/Engine/Plugins/Experimental/AnimNextAnimGraph/Source/AnimNextAnimGraph/Private/Traits/SynchronizeUsingGroups.h
+++ b/Engine/Plugins/Experimental/AnimNextAnimGraph/Source/AnimNextAnimGraph/Private/Traits/SynchronizeUsingGroups.h
@@ -29,8 +29,8 @@ struct FAnimNextSynchronizeUsingGroupsTraitSharedData : public FAnimNextTraitSha
 	UPROPERTY(EditAnywhere, Category = "Default")
 	EAnimGroupSynchronizationMode SyncMode = EAnimGroupSynchronizationMode::NoSynchronization;
 
-	// Whether or not to match the group sync point when joining as leader
-	// When disabled, the leader will snap the position of the group to its current position
+	// Whether or not to match the group sync point when joining as leader or follower with markers
+	// When disabled, the start position of synced sequences must be properly set to avoid pops
 	UPROPERTY(EditAnywhere, Category = "Default")
 	bool bMatchSyncPoint = true;
 
@@ -87,10 +87,10 @@ namespace UE::AnimNext
 
 		// IGroupSynchronization impl
 		virtual FSyncGroupParameters GetGroupParameters(FExecutionContext& Context, const TTraitBinding<IGroupSynchronization>& Binding) const override;
-		virtual void AdvanceBy(FExecutionContext& Context, const TTraitBinding<IGroupSynchronization>& Binding, float DeltaTime) const override;
+		virtual void AdvanceBy(FExecutionContext& Context, const TTraitBinding<IGroupSynchronization>& Binding, float DeltaTime, bool bDispatchEvents) const override;
 
 		// ITimelinePlayer impl
-		virtual void AdvanceBy(FExecutionContext& Context, const TTraitBinding<ITimelinePlayer>& Binding, float DeltaTime) const override;
+		virtual void AdvanceBy(FExecutionContext& Context, const TTraitBinding<ITimelinePlayer>& Binding, float DeltaTime, bool bDispatchEvents) const override;
 
 		// IGraphFactory impl
 		virtual void CreatePayloadForObject(FExecutionContext& Context, const TTraitBinding<IGraphFactory>& Binding, const UObject* InObject, FAnimNextDataInterfacePayload& InOutPayload) const override;
diff --git a/Engine/Plugins/Experimental/AnimNextAnimGraph/Source/AnimNextAnimGraph/Public/TraitInterfaces/IGroupSynchronization.h b/Engine/Plugins/Experimental/AnimNextAnimGraph/Source/AnimNextAnimGraph/Public/TraitInterfaces/IGroupSynchronization.h
--- a/Engine/Plugins/Experimental/AnimNextAnimGraph/Source/AnimNextAnimGraph/Public/TraitInterfaces/IGroupSynchronization.h
+++ b/Engine/Plugins/Experimental/AnimNextAnimGraph/Source/AnimNextAnimGraph/Public/TraitInterfaces/IGroupSynchronization.h
@@ -84,8 +84,8 @@ namespace UE::AnimNext
 		// The synchronization mode
 		EAnimGroupSynchronizationMode SyncMode = EAnimGroupSynchronizationMode::NoSynchronization;
 
-		// Whether or not to match the group sync point when joining as leader
-		// When disabled, the leader will snap the position of the group to its current position
+		// Whether or not to match the group sync point when joining as leader or follower with markers
+		// When disabled, the start position of synced sequences must be properly set to avoid pops
 		bool bMatchSyncPoint = true;
 	};
 
@@ -102,7 +102,7 @@ namespace UE::AnimNext
 		virtual FSyncGroupParameters GetGroupParameters(FExecutionContext& Context, const TTraitBinding<IGroupSynchronization>& Binding) const;
 
 		// Called by the sync group graph instance component to advance time
-		virtual void AdvanceBy(FExecutionContext& Context, const TTraitBinding<IGroupSynchronization>& Binding, float DeltaTime) const;
+		virtual void AdvanceBy(FExecutionContext& Context, const TTraitBinding<IGroupSynchronization>& Binding, float DeltaTime, bool bDispatchEvents) const;
 
 #if WITH_EDITOR
 		virtual const FText& GetDisplayName() const override;
@@ -123,9 +123,9 @@ namespace UE::AnimNext
 		}
 
 		// @see IGroupSynchronization::AdvanceBy
-		void AdvanceBy(FExecutionContext& Context, float DeltaTime) const
+		void AdvanceBy(FExecutionContext& Context, float DeltaTime, bool bDispatchEvents) const
 		{
-			GetInterface()->AdvanceBy(Context, *this, DeltaTime);
+			GetInterface()->AdvanceBy(Context, *this, DeltaTime, bDispatchEvents);
 		}
 
 	protected:
diff --git a/Engine/Plugins/Experimental/AnimNextAnimGraph/Source/AnimNextAnimGraph/Public/TraitInterfaces/ITimeline.h b/Engine/Plugins/Experimental/AnimNextAnimGraph/Source/AnimNextAnimGraph/Public/TraitInterfaces/ITimeline.h
--- a/Engine/Plugins/Experimental/AnimNextAnimGraph/Source/AnimNextAnimGraph/Public/TraitInterfaces/ITimeline.h
+++ b/Engine/Plugins/Experimental/AnimNextAnimGraph/Source/AnimNextAnimGraph/Public/TraitInterfaces/ITimeline.h
@@ -62,12 +62,26 @@ namespace UE::AnimNext
 		// Resets the timeline state to its default state
 		void Reset()
 		{
+#if !(UE_BUILD_SHIPPING || UE_BUILD_TEST)
+			DebugName = NAME_None;
+#endif
+
 			Position = 0.0f;
 			Duration = 0.0f;
 			PlayRate = 0.0f;
 			bIsLooping = false;
 		}
 
+		// Returns the timeline name
+		FName GetDebugName() const
+		{
+#if !(UE_BUILD_SHIPPING || UE_BUILD_TEST)
+			return DebugName;
+#else
+			return NAME_None;
+#endif
+		}
+
 		// Returns the timeline position in seconds
 		float GetPosition() const { return Position; }
 
@@ -106,6 +120,18 @@ namespace UE::AnimNext
 		// Returns whether or not this timeline is finite
 		bool IsFinite() const { return Duration != INFINITY; }
 
+		// Creates a new instance of the timeline state with the supplied debug name
+		[[nodiscard]] FTimelineState WithDebugName(FName InDebugName) const
+		{
+#if !(UE_BUILD_SHIPPING || UE_BUILD_TEST)
+			FTimelineState Result(*this);
+			Result.DebugName = InDebugName;
+			return Result;
+#else
+			return *this;
+#endif
+		}
+
 		// Creates a new instance of the timeline state with the supplied position
 		[[nodiscard]] FTimelineState WithPosition(float InPosition) const
 		{
@@ -139,6 +165,11 @@ namespace UE::AnimNext
 		}
 
 	private:
+#if !(UE_BUILD_SHIPPING || UE_BUILD_TEST)
+		// A debug name associated with this timeline
+		FName DebugName;
+#endif
+
 		// Timeline position in seconds
 		float Position = 0.0f;
 
diff --git a/Engine/Plugins/Experimental/AnimNextAnimGraph/Source/AnimNextAnimGraph/Public/TraitInterfaces/ITimelinePlayer.h b/Engine/Plugins/Experimental/AnimNextAnimGraph/Source/AnimNextAnimGraph/Public/TraitInterfaces/ITimelinePlayer.h
--- a/Engine/Plugins/Experimental/AnimNextAnimGraph/Source/AnimNextAnimGraph/Public/TraitInterfaces/ITimelinePlayer.h
+++ b/Engine/Plugins/Experimental/AnimNextAnimGraph/Source/AnimNextAnimGraph/Public/TraitInterfaces/ITimelinePlayer.h
@@ -18,7 +18,7 @@ namespace UE::AnimNext
 		DECLARE_ANIM_TRAIT_INTERFACE(ITimelinePlayer)
 
 		// Advances time by the provided delta time (positive or negative) on this timeline player
-		virtual void AdvanceBy(FExecutionContext& Context, const TTraitBinding<ITimelinePlayer>& Binding, float DeltaTime) const;
+		virtual void AdvanceBy(FExecutionContext& Context, const TTraitBinding<ITimelinePlayer>& Binding, float DeltaTime, bool bDispatchEvents) const;
 
 #if WITH_EDITOR
 		virtual const FText& GetDisplayName() const override;
@@ -33,9 +33,9 @@ namespace UE::AnimNext
 	struct TTraitBinding<ITimelinePlayer> : FTraitBinding
 	{
 		// @see ITimelinePlayer::AdvanceBy
-		void AdvanceBy(FExecutionContext& Context, float DeltaTime) const
+		void AdvanceBy(FExecutionContext& Context, float DeltaTime, bool bDispatchEvents) const
 		{
-			GetInterface()->AdvanceBy(Context, *this, DeltaTime);
+			GetInterface()->AdvanceBy(Context, *this, DeltaTime, bDispatchEvents);
 		}
 
 	protected:
diff --git a/Engine/Plugins/Experimental/EvaluationNotifies/Source/Runtime/Private/AnimNext/EvaluationNotifiesTrait.cpp b/Engine/Plugins/Experimental/EvaluationNotifies/Source/Runtime/Private/AnimNext/EvaluationNotifiesTrait.cpp
--- a/Engine/Plugins/Experimental/EvaluationNotifies/Source/Runtime/Private/AnimNext/EvaluationNotifiesTrait.cpp
+++ b/Engine/Plugins/Experimental/EvaluationNotifies/Source/Runtime/Private/AnimNext/EvaluationNotifiesTrait.cpp
@@ -33,7 +33,7 @@ namespace UE::AnimNext
 	GENERATE_ANIM_TRAIT_IMPLEMENTATION(FEvaluationNotifiesTrait, TRAIT_INTERFACE_ENUMERATOR, NULL_ANIM_TRAIT_INTERFACE_ENUMERATOR, NULL_ANIM_TRAIT_EVENT_ENUMERATOR)
 	#undef TRAIT_INTERFACE_ENUMERATOR
 
-	void FEvaluationNotifiesTrait::AdvanceBy(FExecutionContext& Context, const TTraitBinding<ITimelinePlayer>& Binding, float DeltaTime) const
+	void FEvaluationNotifiesTrait::AdvanceBy(FExecutionContext& Context, const TTraitBinding<ITimelinePlayer>& Binding, float DeltaTime, bool bDispatchEvents) const
 	{
 		TTraitBinding<ITimeline> TimelineTrait;
 		ensure(Binding.GetStackInterfaceSuper(TimelineTrait));
@@ -41,12 +41,18 @@ namespace UE::AnimNext
 		TTraitBinding<ITimelinePlayer> TimelinePlayerTrait;
 		ensure(Binding.GetStackInterfaceSuper(TimelinePlayerTrait));
 
+		// Get current state from the stack, advance time, then get the delta state
+		const FTimelineState PreAdvanceState = TimelineTrait.GetState(Context);
+		TimelinePlayerTrait.AdvanceBy(Context, DeltaTime, bDispatchEvents);
+
+		if (!bDispatchEvents)
+		{
+			return;
+		}
+
 		TTraitBinding<INotifySource> NotifySourceTrait;
 		ensure(Binding.GetStackInterfaceSuper(NotifySourceTrait));
 
-		// Get current state from the stack, advance time, then get the delta state
-		const FTimelineState PreAdvanceState = TimelineTrait.GetState(Context);
-		TimelinePlayerTrait.AdvanceBy(Context, DeltaTime);
 		const FTimelineDelta Delta = TimelineTrait.GetDelta(Context);
 
 		// Query for notifies
diff --git a/Engine/Plugins/Experimental/EvaluationNotifies/Source/Runtime/Public/AnimNext/EvaluationNotifiesTrait.h b/Engine/Plugins/Experimental/EvaluationNotifies/Source/Runtime/Public/AnimNext/EvaluationNotifiesTrait.h
--- a/Engine/Plugins/Experimental/EvaluationNotifies/Source/Runtime/Public/AnimNext/EvaluationNotifiesTrait.h
+++ b/Engine/Plugins/Experimental/EvaluationNotifies/Source/Runtime/Public/AnimNext/EvaluationNotifiesTrait.h
@@ -66,7 +66,7 @@ struct EVALUATIONNOTIFIESRUNTIME_API FEvaluationNotifiesTrait : FAdditiveTrait,
 	};
 
 	// ITimelinePlayer impl
-	virtual void AdvanceBy(FExecutionContext& Context, const TTraitBinding<ITimelinePlayer>& Binding, float DeltaTime) const override;
+	virtual void AdvanceBy(FExecutionContext& Context, const TTraitBinding<ITimelinePlayer>& Binding, float DeltaTime, bool bDispatchEvents) const override;
 	
 	// IUpdate impl 
 	virtual void PreUpdate(FUpdateTraversalContext& Context, const TTraitBinding<IUpdate>& Binding, const FTraitUpdateState& TraitState) const override;
diff --git a/Engine/Source/Runtime/Engine/Classes/Animation/AnimationAsset.h b/Engine/Source/Runtime/Engine/Classes/Animation/AnimationAsset.h
--- a/Engine/Source/Runtime/Engine/Classes/Animation/AnimationAsset.h
+++ b/Engine/Source/Runtime/Engine/Classes/Animation/AnimationAsset.h
@@ -58,6 +58,7 @@ struct FMarkerPair
 	float TimeToMarker;
 
 	FMarkerPair() : MarkerIndex(MarkerIndexSpecialValues::Uninitialized) {}
+	FMarkerPair(int32 InMarkerIndex, float InTimeToMarker) : MarkerIndex(InMarkerIndex), TimeToMarker(InTimeToMarker) {}
 
 	void Reset() { MarkerIndex = MarkerIndexSpecialValues::Uninitialized; }
 };
