Fix incorrect loop detection in logging code
#rnx

diff --git a/Engine/Plugins/Experimental/AnimNextAnimGraph/Source/AnimNextAnimGraph/Private/Graph/SyncGroup_GraphInstanceComponent.cpp b/Engine/Plugins/Experimental/AnimNextAnimGraph/Source/AnimNextAnimGraph/Private/Graph/SyncGroup_GraphInstanceComponent.cpp
--- a/Engine/Plugins/Experimental/AnimNextAnimGraph/Source/AnimNextAnimGraph/Private/Graph/SyncGroup_GraphInstanceComponent.cpp
+++ b/Engine/Plugins/Experimental/AnimNextAnimGraph/Source/AnimNextAnimGraph/Private/Graph/SyncGroup_GraphInstanceComponent.cpp
@@ -1116,9 +1116,9 @@ namespace UE::AnimNext
 					GroupContext.bIsLeaderPlayingForward ? TEXT("+") : TEXT("-"),
 					DeltaTime * FMath::Abs(PlayRate));
 
-				const bool bLooped = GroupContext.bIsLeaderPlayingForward
-					? (LeaderEndTimelineState.GetPosition() < LeaderStartTimelineState.GetPosition())
-					: (LeaderStartTimelineState.GetPosition() > LeaderEndTimelineState.GetPosition());
+				const bool bLooped = GroupContext.bIsLeaderPlayingForward ?
+					(LeaderEndTimelineState.GetPosition() < LeaderStartTimelineState.GetPosition()) :
+					(LeaderStartTimelineState.GetPosition() < LeaderEndTimelineState.GetPosition());
 
 				UE_LOG(LogAnimMarkerSync, Warning,
 					TEXT("[%s] [%p] Leader [Time] [%0.2f %s -> %0.2f / %0.2f (%3.2f%%)] Playing [%s]%s"),
@@ -1338,7 +1338,7 @@ namespace UE::AnimNext
 
 				const bool bLooped = GroupContext.bIsLeaderPlayingForward ?
 					(LeaderEndTimelineState.GetPosition() < LeaderStartTimelineState.GetPosition()) :
-					(LeaderEndTimelineState.GetPosition() > LeaderStartTimelineState.GetPosition());
+					(LeaderStartTimelineState.GetPosition() < LeaderEndTimelineState.GetPosition());
 
 				FString PassedMarkers;
 				if (!GroupContext.MarkersPassed.IsEmpty())
@@ -1603,8 +1603,8 @@ namespace UE::AnimNext
 				const FTimelineState FollowerEndTimelineState = GroupMemberContext.TimelineTrait.GetState(Context);
 
 				const bool bLooped = GroupContext.bIsLeaderPlayingForward ?
-					(FollowerEndTimelineState.GetPosition() < CurrentTime) :
-					(FollowerEndTimelineState.GetPosition() > CurrentTime);
+					(FollowerEndTimelineState.GetPosition() < FollowerStartTimelineState.GetPosition()) :
+					(FollowerStartTimelineState.GetPosition() < FollowerEndTimelineState.GetPosition());
 
 				const bool bUnexpectedDeltaTime = DeltaTime < 0.0f || DeltaTime > 0.3f;
 
