AnimNext: Fix notify states getting stuck active

We were earlying-out if the queue was empty, but this meant that active states did not get processed. Changed early out condition to also check active states

diff --git a/Engine/Plugins/Experimental/AnimNextAnimGraph/Source/AnimNextAnimGraph/Private/Traits/NotifyDispatcher.cpp b/Engine/Plugins/Experimental/AnimNextAnimGraph/Source/AnimNextAnimGraph/Private/Traits/NotifyDispatcher.cpp
--- a/Engine/Plugins/Experimental/AnimNextAnimGraph/Source/AnimNextAnimGraph/Private/Traits/NotifyDispatcher.cpp
+++ b/Engine/Plugins/Experimental/AnimNextAnimGraph/Source/AnimNextAnimGraph/Private/Traits/NotifyDispatcher.cpp
@@ -143,8 +143,9 @@ void FAnimNextNotifyDispatcherComponent::OnTraitEvent(FAnimNextTraitEvent& Event
 void FAnimNextNotifyDispatcherComponent::OnEndExecution(float InDeltaTime)
 {
 	using namespace UE::AnimNext;
-	
-	if(NotifyQueue.AnimNotifies.Num() == 0)
+
+	// Early out if we dont need to do any work
+	if(NotifyQueue.AnimNotifies.Num() == 0 && ActiveAnimNotifyState.Num() == 0)
 	{
 		return;
 	}
