AnimNext: Fix crash when dispatching notifies from a skel mesh component that has been removed

Between worker thread dispatch and game thread dispatch the skel mesh component may have been removed

diff --git a/Engine/Plugins/Experimental/AnimNextAnimGraph/Source/AnimNextAnimGraph/Private/Traits/NotifyDispatcher.cpp b/Engine/Plugins/Experimental/AnimNextAnimGraph/Source/AnimNextAnimGraph/Private/Traits/NotifyDispatcher.cpp
--- a/Engine/Plugins/Experimental/AnimNextAnimGraph/Source/AnimNextAnimGraph/Private/Traits/NotifyDispatcher.cpp
+++ b/Engine/Plugins/Experimental/AnimNextAnimGraph/Source/AnimNextAnimGraph/Private/Traits/NotifyDispatcher.cpp
@@ -89,6 +89,10 @@ namespace UE::AnimNext
 	void FNotifyQueueDispatchEvent::Execute() const
 	{
 		USkeletalMeshComponent* SkeletalMeshComponent = WeakSkeletalMeshComponent.Get();
+		if (SkeletalMeshComponent == nullptr)
+		{
+			return;
+		}
 
 		// Order of dispatch here mimics that of UAnimInstance::TriggerAnimNotifies
 
