Defer loading of animation graph factory assets during engine init

diff --git a/Engine/Plugins/Experimental/AnimNextAnimGraph/Source/AnimNextAnimGraph/Internal/AnimNextAnimGraphSettings.h b/Engine/Plugins/Experimental/AnimNextAnimGraph/Source/AnimNextAnimGraph/Internal/AnimNextAnimGraphSettings.h
--- a/Engine/Plugins/Experimental/AnimNextAnimGraph/Source/AnimNextAnimGraph/Internal/AnimNextAnimGraphSettings.h
+++ b/Engine/Plugins/Experimental/AnimNextAnimGraph/Source/AnimNextAnimGraph/Internal/AnimNextAnimGraphSettings.h
@@ -93,6 +93,10 @@ private:
 #endif
 
 	TArray<UClass*> GetAllowedAssetClassesImpl() const;
+
+	void OnDefaultRunGraphHostLoaded(const UAnimNextAnimationGraph* AnimationGraph);
+	void OnMappingAnimationGraphLoaded(const UAnimNextAnimationGraph* AnimationGraph, FAnimNextAssetGraphMapping& Mapping, int32 MappingIndex);
+	void FinalizeAsyncLoad();
 	
 private:
 	// The animation graph we run by default when hosting in a module
@@ -107,6 +111,9 @@ private:
 	// Maps an asset type (e.g. anim sequence) to the index in the mappings array for the first entry for that object
 	TMap<TObjectKey<const UClass>, int32> AssetGraphMap;
 
+	// Counter to keep track of how mappings have finished loading
+	int32 NumMappingsLoaded = 0;
+
 	// Flag to verify correct mapping rebuild behavior during startup
-	bool bMappingsBuiltAtLeastOnce = false;
+	std::atomic<bool> bMappingsBuiltAtLeastOnce = false;
 };
\ No newline at end of file
diff --git a/Engine/Plugins/Experimental/AnimNextAnimGraph/Source/AnimNextAnimGraph/Private/AnimNextAnimGraphSettings.cpp b/Engine/Plugins/Experimental/AnimNextAnimGraph/Source/AnimNextAnimGraph/Private/AnimNextAnimGraphSettings.cpp
--- a/Engine/Plugins/Experimental/AnimNextAnimGraph/Source/AnimNextAnimGraph/Private/AnimNextAnimGraphSettings.cpp
+++ b/Engine/Plugins/Experimental/AnimNextAnimGraph/Source/AnimNextAnimGraph/Private/AnimNextAnimGraphSettings.cpp
@@ -36,7 +36,6 @@ const UAnimNextAnimationGraph* UAnimNextAnimGraphSettings::GetGraphFromObject(co
 const UAnimNextAnimationGraph* UAnimNextAnimGraphSettings::GetGraphFromObject(const UObject* InObject, const FAnimNextDataInterfacePayload& InGraphPayload) const
 {
 	check(InObject);
-	ensure(bMappingsBuiltAtLeastOnce);
 
 	const UAnimNextAnimationGraph* AnimationGraph = Cast<UAnimNextAnimationGraph>(InObject);
 	if(AnimationGraph)
@@ -45,6 +44,11 @@ const UAnimNextAnimationGraph* UAnimNextAnimGraphSettings::GetGraphFromObject(co
 		return AnimationGraph;
 	}
 
+	if (!ensure(bMappingsBuiltAtLeastOnce.load(std::memory_order_acquire)))
+	{
+		return nullptr;
+	}
+
 	// Check this class and its parent classes in turn
 	const UClass* Class = InObject->GetClass();
 	int32 Index = INDEX_NONE;
@@ -125,8 +129,6 @@ const UAnimNextAnimationGraph* UAnimNextAnimGraphSettings::GetGraphFromObject(co
 
 bool UAnimNextAnimGraphSettings::CanGetGraphFromAssetClass(const UClass* InClass) const
 {
-	ensure(bMappingsBuiltAtLeastOnce);
-
 	if(InClass == nullptr)
 	{
 		return false;
@@ -137,6 +139,11 @@ bool UAnimNextAnimGraphSettings::CanGetGraphFromAssetClass(const UClass* InClass
 		return true;
 	}
 
+	if (!ensure(bMappingsBuiltAtLeastOnce.load(std::memory_order_acquire)))
+	{
+		return false;
+	}
+
 	const UClass* Class = InClass;
 	while(Class)
 	{
@@ -154,7 +161,6 @@ bool UAnimNextAnimGraphSettings::CanGetGraphFromAssetClass(const UClass* InClass
 FName UAnimNextAnimGraphSettings::GetInjectedVariableNameFromObject(const UObject* InObject) const
 {
 	check(InObject);
-	ensure(bMappingsBuiltAtLeastOnce);
 
 	const UAnimNextAnimationGraph* AnimationGraph = Cast<UAnimNextAnimationGraph>(InObject);
 	if(AnimationGraph)
@@ -163,6 +169,11 @@ FName UAnimNextAnimGraphSettings::GetInjectedVariableNameFromObject(const UObjec
 		return NAME_None;
 	}
 
+	if (!ensure(bMappingsBuiltAtLeastOnce.load(std::memory_order_acquire)))
+	{
+		return NAME_None;
+	}
+
 	// Check this class and its parent classes in turn
 	const UClass* Class = InObject->GetClass();
 	int32 Index = INDEX_NONE;
@@ -273,95 +284,146 @@ void UAnimNextAnimGraphSettings::GetNonNativePayloadFromGraph(const UObject* InO
 	ObjectProperty->SetObjectPropertyValue(ValuePtr, const_cast<UObject*>(InObject));
 }
 
-void UAnimNextAnimGraphSettings::LoadAndRebuildMappings()
+void UAnimNextAnimGraphSettings::OnDefaultRunGraphHostLoaded(const UAnimNextAnimationGraph* AnimationGraph)
 {
+	if (AnimationGraph == nullptr)
+	{
+		UE_LOGFMT(LogAnimation, Error, "UAnimNextAnimGraphSettings::LoadAndRebuildMappings: Could not load default host animation graph '{GraphName}'", DefaultRunGraphHost.ToString());
+		return;
+	}
+
 	using namespace UE::AnimNext::AnimGraph;
-	
 	FAnimNextAnimGraphModule& AnimGraphModule = FModuleManager::GetModuleChecked<FAnimNextAnimGraphModule>("AnimNextAnimGraph");
-	AnimGraphModule.LoadedGraphs.Reset();
-	AssetGraphMap.Reset();
+	AnimGraphModule.LoadedGraphs.Add(AnimationGraph);
 
-	const UAnimNextAnimationGraph* DefaultAnimationGraph = DefaultRunGraphHost.LoadSynchronous();
-	if(DefaultAnimationGraph == nullptr)
+	if (AssetGraphMappings.IsEmpty())
 	{
-		UE_LOGFMT(LogAnimation, Error, "UAnimNextAnimGraphSettings::LoadAndRebuildMappings: Could not load default host animation graph '{GraphName}'", DefaultRunGraphHost.ToString());
+		FinalizeAsyncLoad();
+		return;
+	}
+
+	for (TEnumerateRef<FAnimNextAssetGraphMapping> Mapping : EnumerateRange(AssetGraphMappings))
+	{
+		Mapping->AnimationGraph.ToSoftObjectPath().LoadAsync(FLoadSoftObjectPathAsyncDelegate::CreateLambda(
+			[WeakThis = TWeakObjectPtr<UAnimNextAnimGraphSettings>(this),
+			&Mapping = *Mapping,
+			MappingIndex = Mapping.GetIndex()]
+			(const FSoftObjectPath& InSoftObjectPath, UObject* InObject)
+			{
+				if (UAnimNextAnimGraphSettings* ThisPin = WeakThis.Get())
+				{
+					ThisPin->NumMappingsLoaded++;
+
+					ThisPin->OnMappingAnimationGraphLoaded(CastChecked<const UAnimNextAnimationGraph>(InObject), Mapping, MappingIndex);
+
+					if (ThisPin->NumMappingsLoaded == ThisPin->AssetGraphMappings.Num())
+					{
+						ThisPin->FinalizeAsyncLoad();
+					}
+				}
+			}));
+	}
+}
+
+void UAnimNextAnimGraphSettings::OnMappingAnimationGraphLoaded(const UAnimNextAnimationGraph* AnimationGraph, FAnimNextAssetGraphMapping& Mapping, int32 MappingIndex)
+{
+	if (AnimationGraph == nullptr)
+	{
+		UE_LOGFMT(LogAnimation, Error, "UAnimNextAnimGraphSettings::LoadAndRebuildMappings: Could not load animation graph '{GraphName}'", Mapping.AnimationGraph.ToString());
+		return;
 	}
-	else
+
+	using namespace UE::AnimNext::AnimGraph;
+	FAnimNextAnimGraphModule& AnimGraphModule = FModuleManager::GetModuleChecked<FAnimNextAnimGraphModule>("AnimNextAnimGraph");
+	AnimGraphModule.LoadedGraphs.Add(AnimationGraph);
+
+	const UClass* AssetType = Mapping.AssetType.LoadSynchronous();
+	if (AssetType == nullptr)
 	{
-		AnimGraphModule.LoadedGraphs.Add(DefaultAnimationGraph);
+		UE_LOGFMT(LogAnimation, Error, "UAnimNextAnimGraphSettings::LoadAndRebuildMappings: Could not load asset class '{ClassName}'", Mapping.AssetType.ToString());
+		return;
 	}
 
-	for(TEnumerateRef<FAnimNextAssetGraphMapping> Mapping : EnumerateRange(AssetGraphMappings))
+	bool bAreRequiredInterfacesValid = true;
+	for (const TSoftObjectPtr<const UScriptStruct>& DataInterfacePtr : Mapping.RequiredDataInterfaces)
 	{
-		const UAnimNextAnimationGraph* AnimationGraph = Mapping->AnimationGraph.LoadSynchronous();
-		if(AnimationGraph == nullptr)
+		const UScriptStruct* DataInterface = DataInterfacePtr.LoadSynchronous();
+		if (DataInterface == nullptr)
 		{
-			UE_LOGFMT(LogAnimation, Error, "UAnimNextAnimGraphSettings::LoadAndRebuildMappings: Could not load animation graph '{GraphName}'", Mapping->AnimationGraph.ToString());
+			UE_LOGFMT(LogAnimation, Error, "UAnimNextAnimGraphSettings::LoadAndRebuildMappings: Could not load required data interface class '{ClassName}'", DataInterfacePtr.ToString());
+			bAreRequiredInterfacesValid = false;
 			continue;
 		}
+	}
 
-		AnimGraphModule.LoadedGraphs.Add(AnimationGraph);
+	if (!bAreRequiredInterfacesValid)
+	{
+		return;
+	}
 
-		const UClass* AssetType = Mapping->AssetType.LoadSynchronous();
-		if(AssetType == nullptr)
+	// Find or add our new entry
+	int32& FirstAssetMappingIndex = AssetGraphMap.FindOrAdd(AssetType, INDEX_NONE);
+
+	// Update the next mapping index to point to our old entry (or INDEX_NONE if we are the first)
+	Mapping.NextMappingIndex = FirstAssetMappingIndex;
+
+	// Our new first mapping index is the current one
+	FirstAssetMappingIndex = MappingIndex;
+
+	// Warn if the variable is not settable in the public interface of the graph
+	if (AnimationGraph != nullptr && AssetType != nullptr && !Mapping.Variable.IsNone())
+	{
+		const FInstancedPropertyBag& PublicVariableDefaults = AnimationGraph->GetPublicVariableDefaults();
+		const FPropertyBagPropertyDesc* Desc = PublicVariableDefaults.FindPropertyDescByName(Mapping.Variable);
+		if (Desc == nullptr)
 		{
-			UE_LOGFMT(LogAnimation, Error, "UAnimNextAnimGraphSettings::LoadAndRebuildMappings: Could not load asset class '{ClassName}'", Mapping->AssetType.ToString());
-			continue;
+			UE_LOGFMT(LogAnimation, Error, "UAnimNextAnimGraphSettings::LoadAndRebuildMappings: Could not find public variable '{VariableName}' in graph '{GraphName}'", Mapping.Variable, AnimationGraph->GetPathName());
+			return;
 		}
 
-		bool bAreRequiredInterfacesValid = true;
-		for (const TSoftObjectPtr<const UScriptStruct>& DataInterfacePtr : Mapping->RequiredDataInterfaces)
+		const FObjectProperty* ObjectProperty = CastField<FObjectProperty>(Desc->CachedProperty);
+		if (ObjectProperty == nullptr)
 		{
-			const UScriptStruct* DataInterface = DataInterfacePtr.LoadSynchronous();
-			if (DataInterface == nullptr)
-			{
-				UE_LOGFMT(LogAnimation, Error, "UAnimNextAnimGraphSettings::LoadAndRebuildMappings: Could not load required data interface class '{ClassName}'", DataInterfacePtr.ToString());
-				bAreRequiredInterfacesValid = false;
-				continue;
-			}
+			UE_LOGFMT(LogAnimation, Error, "UAnimNextAnimGraphSettings::LoadAndRebuildMappings: Variable '{VariableName}' in graph '{GraphName}' is not of object type", Mapping.Variable, AnimationGraph->GetPathName());
+			return;
 		}
 
-		if (!bAreRequiredInterfacesValid)
+		if (!AssetType->IsChildOf(ObjectProperty->PropertyClass))
 		{
-			continue;
+			UE_LOGFMT(LogAnimation, Error, "UAnimNextAnimGraphSettings::LoadAndRebuildMappings: Variable '{VariableName}' in graph '{GraphName}' is not of a compatible object type: '{TargetType}' vs '{SourceType}'", Mapping.Variable, AnimationGraph->GetPathName(), ObjectProperty->PropertyClass->GetFullName(), AssetType->GetFullName());
+			return;
 		}
+	}
+}
 
-		// Find or add our new entry
-		int32& FirstAssetMappingIndex = AssetGraphMap.FindOrAdd(AssetType, INDEX_NONE);
+void UAnimNextAnimGraphSettings::FinalizeAsyncLoad()
+{
+	// Release memory order to ensure all previous writes complete before other threads can see this value change
+	// Queries to this value must use the acquire memory order to ensure loads are not re-ordered before the value
+	// is read and tested
+	bMappingsBuiltAtLeastOnce.store(true, std::memory_order_release);
+}
 
-		// Update the next mapping index to point to our old entry (or INDEX_NONE if we are the first)
-		Mapping->NextMappingIndex = FirstAssetMappingIndex;
+void UAnimNextAnimGraphSettings::LoadAndRebuildMappings()
+{
+	using namespace UE::AnimNext::AnimGraph;
 
-		// Our new first mapping index is the current one
-		FirstAssetMappingIndex = Mapping.GetIndex();
+	// To avoid loading every asset synchronously at engine init, we load them async in multiple steps
+	//   - We first load the default run graph host
+	//   - Followed by the mapping graph assets concurrently
+	//   - Once everything has loaded, the graph factories are usable
+	
+	FAnimNextAnimGraphModule& AnimGraphModule = FModuleManager::GetModuleChecked<FAnimNextAnimGraphModule>("AnimNextAnimGraph");
+	AnimGraphModule.LoadedGraphs.Reset();
+	AssetGraphMap.Reset();
 
-		// Warn if the variable is not settable in the public interface of the graph
-		if(AnimationGraph != nullptr && AssetType != nullptr && !Mapping->Variable.IsNone())
+	DefaultRunGraphHost.ToSoftObjectPath().LoadAsync(FLoadSoftObjectPathAsyncDelegate::CreateLambda([WeakThis = TWeakObjectPtr<UAnimNextAnimGraphSettings>(this)](const FSoftObjectPath& InSoftObjectPath, UObject* InObject)
 		{
-			const FInstancedPropertyBag& PublicVariableDefaults = AnimationGraph->GetPublicVariableDefaults();
-			const FPropertyBagPropertyDesc* Desc = PublicVariableDefaults.FindPropertyDescByName(Mapping->Variable);
-			if(Desc == nullptr)
-			{
-				UE_LOGFMT(LogAnimation, Error, "UAnimNextAnimGraphSettings::LoadAndRebuildMappings: Could not find public variable '{VariableName}' in graph '{GraphName}'", Mapping->Variable, AnimationGraph->GetPathName());
-				continue;
-			}
-
-			const FObjectProperty* ObjectProperty = CastField<FObjectProperty>(Desc->CachedProperty);
-			if(ObjectProperty == nullptr)
+			if (TStrongObjectPtr<UAnimNextAnimGraphSettings> ThisPin = WeakThis.Pin())
 			{
-				UE_LOGFMT(LogAnimation, Error, "UAnimNextAnimGraphSettings::LoadAndRebuildMappings: Variable '{VariableName}' in graph '{GraphName}' is not of object type", Mapping->Variable, AnimationGraph->GetPathName());
-				return;
+				ThisPin->OnDefaultRunGraphHostLoaded(CastChecked<const UAnimNextAnimationGraph>(InObject));
 			}
-
-			if(!AssetType->IsChildOf(ObjectProperty->PropertyClass))
-			{
-				UE_LOGFMT(LogAnimation, Error, "UAnimNextAnimGraphSettings::LoadAndRebuildMappings: Variable '{VariableName}' in graph '{GraphName}' is not of a compatible object type: '{TargetType}' vs '{SourceType}'", Mapping->Variable, AnimationGraph->GetPathName(), ObjectProperty->PropertyClass->GetFullName(), AssetType->GetFullName());
-				return;
-			}
-		}
-	}
-
-	bMappingsBuiltAtLeastOnce = true;
+		}));
 }
 
 TArray<UClass*> UAnimNextAnimGraphSettings::GetAllowedAssetClasses()
@@ -371,10 +433,13 @@ TArray<UClass*> UAnimNextAnimGraphSettings::GetAllowedAssetClasses()
 
 TArray<UClass*> UAnimNextAnimGraphSettings::GetAllowedAssetClassesImpl() const
 {
-	ensure(bMappingsBuiltAtLeastOnce);
-
 	TArray<UClass*> AllowedClasses;
 
+	if (!ensure(bMappingsBuiltAtLeastOnce.load(std::memory_order_acquire)))
+	{
+		return AllowedClasses;
+	}
+
 	// Add all derived classes of UAnimNextAnimationGraph
 	AllowedClasses.Add(UAnimNextAnimationGraph::StaticClass());
 	GetDerivedClasses(UAnimNextAnimationGraph::StaticClass(), AllowedClasses);
diff --git a/Engine/Plugins/Runtime/RigVM/Source/RigVM/Private/RigVMCore/RigVMRegistry.cpp b/Engine/Plugins/Runtime/RigVM/Source/RigVM/Private/RigVMCore/RigVMRegistry.cpp
--- a/Engine/Plugins/Runtime/RigVM/Source/RigVM/Private/RigVMCore/RigVMRegistry.cpp
+++ b/Engine/Plugins/Runtime/RigVM/Source/RigVM/Private/RigVMCore/RigVMRegistry.cpp
@@ -371,6 +371,8 @@ void FRigVMRegistry_NoLock::RefreshEngineTypesIfRequired_NoLock()
 
 void FRigVMRegistry_NoLock::RefreshEngineTypes_NoLock()
 {
+	TRACE_CPUPROFILER_EVENT_SCOPE(FRigVMRegistry::RefreshEngineTypes);
+
 	TGuardValue<bool> EnableGuardRefresh(bIsRefreshingEngineTypes, true);
 
 	const int32 NumTypesBefore = Types.Num(); 
