Fix Animcrash standalone due to property pins not existing in RigVMTrait_AnimNextPublicVariables

#ForUE5

diff --git a/Engine/Plugins/Experimental/AnimNext/Source/AnimNext/Private/Graph/RigVMTrait_AnimNextPublicVariables.cpp b/Engine/Plugins/Experimental/AnimNext/Source/AnimNext/Private/Graph/RigVMTrait_AnimNextPublicVariables.cpp
--- a/Engine/Plugins/Experimental/AnimNext/Source/AnimNext/Private/Graph/RigVMTrait_AnimNextPublicVariables.cpp
+++ b/Engine/Plugins/Experimental/AnimNext/Source/AnimNext/Private/Graph/RigVMTrait_AnimNextPublicVariables.cpp
@@ -11,20 +11,19 @@ bool (*FRigVMTrait_AnimNextPublicVariables::ShouldCreatePinForPropertyFunc)(cons
 
 FString FRigVMTrait_AnimNextPublicVariables::GetDisplayName() const
 {
+	check(GetDisplayNameFunc);
 	return GetDisplayNameFunc(*this);
 }
 
 void FRigVMTrait_AnimNextPublicVariables::GetProgrammaticPins(URigVMController* InController, int32 InParentPinIndex, const URigVMPin* InTraitPin, const FString& InDefaultValue, FRigVMPinInfoArray& OutPinArray) const
 {
+	check(GetProgrammaticPinsFunc);
 	GetProgrammaticPinsFunc(*this, InController, InParentPinIndex, InDefaultValue, OutPinArray);
 }
 
 bool FRigVMTrait_AnimNextPublicVariables::ShouldCreatePinForProperty(const FProperty* InProperty) const
 {
-	if(!Super::ShouldCreatePinForProperty(InProperty))
-	{
-		return false;
-	}
+	check(ShouldCreatePinForPropertyFunc);
 	return ShouldCreatePinForPropertyFunc(*this, InProperty);
 }
 #endif
diff --git a/Engine/Plugins/Experimental/AnimNext/Source/AnimNextUncookedOnly/Internal/AnimNextRigVMAssetEditorData.h b/Engine/Plugins/Experimental/AnimNext/Source/AnimNextUncookedOnly/Internal/AnimNextRigVMAssetEditorData.h
--- a/Engine/Plugins/Experimental/AnimNext/Source/AnimNextUncookedOnly/Internal/AnimNextRigVMAssetEditorData.h
+++ b/Engine/Plugins/Experimental/AnimNext/Source/AnimNextUncookedOnly/Internal/AnimNextRigVMAssetEditorData.h
@@ -25,6 +25,7 @@ struct FAnimNextProcessGraphCompileContext;
 
 namespace UE::AnimNext::UncookedOnly
 {
+	struct FPublicVariablesImpl;
 	struct FUtils;
 	struct FUtilsPrivate;
 	class FScopedCompilerResults;
@@ -51,7 +52,6 @@ namespace UE::AnimNext::Editor
 	class SVariablesOutliner;
 	class SAddVariablesDialog;
 	class FVariableProxyCustomization;
-	struct FPublicVariablesImpl;
 	class FAnimNextAnimGraphEditorModule;
 	class FCallFunctionSharedDataDetails;
 	class FAssetCompilationHandler;
@@ -198,7 +198,7 @@ protected:
 	friend class UE::AnimNext::Editor::SAddVariablesDialog;
 	friend class UE::AnimNext::Editor::FVariableProxyCustomization;
 	friend class UAnimNextDataInterfaceEntry;
-	friend struct UE::AnimNext::Editor::FPublicVariablesImpl;
+	friend struct UE::AnimNext::UncookedOnly::FPublicVariablesImpl;
 	friend UE::AnimNext::Editor::FAnimNextAnimGraphEditorModule;
 	friend class UE::AnimNext::Editor::FCallFunctionSharedDataDetails;
 	friend UAnimNextWorkspaceEditorMode;
diff --git a/Engine/Plugins/Experimental/AnimNext/Source/AnimNextUncookedOnly/Internal/Entries/AnimNextVariableEntry.h b/Engine/Plugins/Experimental/AnimNext/Source/AnimNextUncookedOnly/Internal/Entries/AnimNextVariableEntry.h
--- a/Engine/Plugins/Experimental/AnimNext/Source/AnimNextUncookedOnly/Internal/Entries/AnimNextVariableEntry.h
+++ b/Engine/Plugins/Experimental/AnimNext/Source/AnimNextUncookedOnly/Internal/Entries/AnimNextVariableEntry.h
@@ -29,7 +29,6 @@ namespace UE::AnimNext::Editor
 	class FVariablesOutlinerHierarchy;
 	class SVariablesOutlinerValue;
 	class FVariableProxyCustomization;
-	struct FPublicVariablesImpl;
 	class FAnimNextEditorModule;
 }
 
@@ -43,6 +42,7 @@ namespace UE::AnimNext::Tests
 
 namespace UE::AnimNext::UncookedOnly
 {
+	struct FPublicVariablesImpl;
 	struct FUtils;
 }
 
@@ -70,7 +70,7 @@ class UAnimNextVariableEntry : public UAnimNextRigVMAssetEntry, public IAnimNext
 	friend class UE::AnimNext::Editor::FVariablesOutlinerHierarchy;
 	friend struct UE::AnimNext::UncookedOnly::FUtils;
 	friend class UE::AnimNext::Editor::FVariableProxyCustomization;
-	friend struct UE::AnimNext::Editor::FPublicVariablesImpl;
+	friend struct UE::AnimNext::UncookedOnly::FPublicVariablesImpl;
 	friend class UE::AnimNext::Editor::FAnimNextEditorModule;
 
 	UAnimNextVariableEntry();
diff --git a/Engine/Plugins/Experimental/AnimNextAnimGraph/Source/AnimNextAnimGraphEditor/Private/AnimNextAnimGraphEditorModule.cpp b/Engine/Plugins/Experimental/AnimNextAnimGraph/Source/AnimNextAnimGraphEditor/Private/AnimNextAnimGraphEditorModule.cpp
--- a/Engine/Plugins/Experimental/AnimNextAnimGraph/Source/AnimNextAnimGraphEditor/Private/AnimNextAnimGraphEditorModule.cpp
+++ b/Engine/Plugins/Experimental/AnimNextAnimGraph/Source/AnimNextAnimGraphEditor/Private/AnimNextAnimGraphEditorModule.cpp
@@ -20,7 +20,6 @@
 #include "Graph/AnimNextAnimationGraph_EditorData.h"
 #include "Graph/AnimNextGraphPanelNodeFactory.h"
 #include "Graph/RigUnit_AnimNextGraphRoot.h"
-#include "Graph/RigVMTrait_AnimNextPublicVariablesEditor.h"
 #include "Graph/TraitEditorTabSummoner.h"
 #include "Traits/AnimNextCallFunctionSharedDataDetails.h"
 #include "Traits/CallFunction.h"
@@ -64,8 +63,6 @@ class FAnimNextAnimGraphEditorModule : public IModuleInterface
 			LOCTEXT("SettingsDescription", "Configure options for AnimNext animation graphs."),
 			GetMutableDefault<UAnimNextAnimGraphSettings>());
 
-		FPublicVariablesImpl::Register();
-
 		FPropertyEditorModule& PropertyModule = FModuleManager::LoadModuleChecked<FPropertyEditorModule>("PropertyEditor");
 		PropertyModule.RegisterCustomPropertyTypeLayout(
 			FAnimNextCallFunctionSharedData::StaticStruct()->GetFName(),
diff --git a/Engine/Plugins/Experimental/AnimNextAnimGraph/Source/AnimNextAnimGraphEditor/Private/Graph/RigVMTrait_AnimNextPublicVariablesEditor.cpp b/Engine/Plugins/Experimental/AnimNextAnimGraph/Source/AnimNextAnimGraphEditor/Private/Graph/RigVMTrait_AnimNextPublicVariablesEditor.cpp
deleted file mode 100644
diff --git a/Engine/Plugins/Experimental/AnimNextAnimGraph/Source/AnimNextAnimGraphEditor/Private/Graph/RigVMTrait_AnimNextPublicVariablesEditor.h b/Engine/Plugins/Experimental/AnimNextAnimGraph/Source/AnimNextAnimGraphEditor/Private/Graph/RigVMTrait_AnimNextPublicVariablesEditor.h
deleted file mode 100644
diff --git a/Engine/Plugins/Experimental/AnimNextAnimGraph/Source/AnimNextAnimGraphUncookedOnly/Private/AnimNextAnimGraphUncookedOnlyModule.cpp b/Engine/Plugins/Experimental/AnimNextAnimGraph/Source/AnimNextAnimGraphUncookedOnly/Private/AnimNextAnimGraphUncookedOnlyModule.cpp
new file mode 100644
--- /dev/null
+++ b/Engine/Plugins/Experimental/AnimNextAnimGraph/Source/AnimNextAnimGraphUncookedOnly/Private/AnimNextAnimGraphUncookedOnlyModule.cpp
@@ -0,0 +1,27 @@
+// Copyright Epic Games, Inc. All Rights Reserved.
+
+#include "CoreMinimal.h"
+#include "Modules/ModuleInterface.h"
+#include "Modules/ModuleManager.h"
+#include "Graph/RigVMTrait_AnimNextPublicVariablesUncookedOnly.h"
+
+#define LOCTEXT_NAMESPACE "FAnimNextAnimGraphUncookedOnlyModule"
+
+namespace UE::AnimNext::AnimGraph::UncookedOnly
+{
+
+class FAnimNextAnimGraphUncookedOnlyModule : public IModuleInterface
+{
+	virtual void StartupModule() override
+	{
+		using namespace UE::AnimNext::UncookedOnly;
+
+		FPublicVariablesImpl::Register();
+	}
+};
+
+}
+
+#undef LOCTEXT_NAMESPACE
+
+IMPLEMENT_MODULE(UE::AnimNext::AnimGraph::UncookedOnly::FAnimNextAnimGraphUncookedOnlyModule, AnimNextAnimGraphUncookedOnly);
diff --git a/Engine/Plugins/Experimental/AnimNextAnimGraph/Source/AnimNextAnimGraphUncookedOnly/Private/Graph/RigVMTrait_AnimNextPublicVariablesUncookedOnly.cpp b/Engine/Plugins/Experimental/AnimNextAnimGraph/Source/AnimNextAnimGraphUncookedOnly/Private/Graph/RigVMTrait_AnimNextPublicVariablesUncookedOnly.cpp
new file mode 100644
--- /dev/null
+++ b/Engine/Plugins/Experimental/AnimNextAnimGraph/Source/AnimNextAnimGraphUncookedOnly/Private/Graph/RigVMTrait_AnimNextPublicVariablesUncookedOnly.cpp
@@ -0,0 +1,101 @@
+// Copyright Epic Games, Inc. All Rights Reserved.
+
+#include "RigVMTrait_AnimNextPublicVariablesUncookedOnly.h"
+#include "AnimNextRigVMAssetEditorData.h"
+#include "UncookedOnlyUtils.h"
+#include "Entries/AnimNextVariableEntry.h"
+#include "Graph/RigVMTrait_AnimNextPublicVariables.h"
+#include "RigVMModel/RigVMController.h"
+#include "RigVMModel/RigVMPin.h"
+
+namespace UE::AnimNext::UncookedOnly
+{
+
+void FPublicVariablesImpl::Register()
+{
+	FRigVMTrait_AnimNextPublicVariables::GetDisplayNameFunc = GetDisplayName;
+	FRigVMTrait_AnimNextPublicVariables::GetProgrammaticPinsFunc = GetProgrammaticPins;
+	FRigVMTrait_AnimNextPublicVariables::ShouldCreatePinForPropertyFunc = ShouldCreatePinForProperty;
+}
+
+FString FPublicVariablesImpl::GetDisplayName(const FRigVMTrait_AnimNextPublicVariables& InTrait)
+{
+	TStringBuilder<256> StringBuilder;
+	StringBuilder.Appendf(TEXT("Variables: %s"), InTrait.InternalAsset ? *InTrait.InternalAsset->GetFName().ToString() : TEXT("None"));
+	return StringBuilder.ToString();
+}
+
+void FPublicVariablesImpl::GetProgrammaticPins(const FRigVMTrait_AnimNextPublicVariables& InTrait, URigVMController* InController, int32 InParentPinIndex, const FString& InDefaultValue, FRigVMPinInfoArray& OutPinArray)
+{
+	if (InTrait.InternalAsset == nullptr)
+	{
+		return;
+	}
+
+	UAnimNextRigVMAssetEditorData* EditorData = UE::AnimNext::UncookedOnly::FUtils::GetEditorData<UAnimNextRigVMAssetEditorData>(InTrait.InternalAsset.Get());
+	if(EditorData == nullptr)
+	{
+		return;
+	}
+
+	TArray<UAnimNextVariableEntry*> PublicVariables;
+	EditorData->GetPublicVariables(PublicVariables);
+	if(PublicVariables.Num() == 0)
+	{
+		return;
+	}
+
+	TArray<FPropertyBagPropertyDesc> Descs;
+	Descs.Reserve(PublicVariables.Num());
+	TArray<TConstArrayView<uint8>> Values;
+	Values.Reserve(PublicVariables.Num());
+
+	// Maintain trait sorted order
+	for (FName VariableName : InTrait.InternalVariableNames)
+	{
+		auto FindMatchingName = [&](const UAnimNextVariableEntry* InVariable)
+		{
+			if (!InVariable)
+			{
+				return false;
+			}
+
+			return InVariable->GetEntryName() == VariableName;
+		};
+
+		if (UAnimNextVariableEntry** PublicVariableItr = PublicVariables.FindByPredicate(FindMatchingName))
+		{
+			if (UAnimNextVariableEntry* PublicVariable = *PublicVariableItr)
+			{
+				const FProperty* Property;
+				TConstArrayView<uint8> Value;
+				PublicVariable->GetDefaultValue(Property, Value);
+				Descs.Emplace(PublicVariable->GetEntryName(), Property);
+				Values.Emplace(Value);
+			}
+		}
+
+	}
+
+	ensure(Values.Num() == InTrait.InternalVariableNames.Num());
+
+	FInstancedPropertyBag Defaults;
+	Defaults.ReplaceAllPropertiesAndValues(Descs, Values);
+
+	const TFunction<ERigVMPinDefaultValueType(const FName&)> DefaultValueTypeGetter = [](const FName& InPropertyName)
+	{
+		return ERigVMPinDefaultValueType::AutoDetect;
+	};
+
+	OutPinArray.AddPins(const_cast<UPropertyBag*>(Defaults.GetPropertyBagStruct()), InController, ERigVMPinDirection::Input, InParentPinIndex, DefaultValueTypeGetter, Defaults.GetValue().GetMemory(), true);
+}
+
+bool FPublicVariablesImpl::ShouldCreatePinForProperty(const FRigVMTrait_AnimNextPublicVariables& InTrait, const FProperty* InProperty)
+{
+	return
+		InProperty->GetFName() == GET_MEMBER_NAME_CHECKED(FRigVMTrait_AnimNextPublicVariables, InternalAsset) ||
+		InProperty->GetFName() == GET_MEMBER_NAME_CHECKED(FRigVMTrait_AnimNextPublicVariables, InternalVariableNames) ||
+		InTrait.InternalVariableNames.Contains(InProperty->GetFName());
+}
+
+}
\ No newline at end of file
diff --git a/Engine/Plugins/Experimental/AnimNextAnimGraph/Source/AnimNextAnimGraphUncookedOnly/Private/Graph/RigVMTrait_AnimNextPublicVariablesUncookedOnly.h b/Engine/Plugins/Experimental/AnimNextAnimGraph/Source/AnimNextAnimGraphUncookedOnly/Private/Graph/RigVMTrait_AnimNextPublicVariablesUncookedOnly.h
new file mode 100644
--- /dev/null
+++ b/Engine/Plugins/Experimental/AnimNextAnimGraph/Source/AnimNextAnimGraphUncookedOnly/Private/Graph/RigVMTrait_AnimNextPublicVariablesUncookedOnly.h
@@ -0,0 +1,23 @@
+﻿// Copyright Epic Games, Inc. All Rights Reserved.
+
+#pragma once
+
+#include "CoreMinimal.h"
+
+struct FRigVMTrait_AnimNextPublicVariables;
+class URigVMController;
+struct FRigVMPinInfoArray;
+
+namespace UE::AnimNext::UncookedOnly
+{
+
+// UncookedOnly-side impl for FRigVMTrait_AnimNextPublicVariables
+struct ANIMNEXTANIMGRAPHUNCOOKEDONLY_API FPublicVariablesImpl
+{
+	static void Register();
+	static FString GetDisplayName(const FRigVMTrait_AnimNextPublicVariables& InTrait);
+	static void GetProgrammaticPins(const FRigVMTrait_AnimNextPublicVariables& InTrait, URigVMController* InController, int32 InParentPinIndex, const FString& InDefaultValue, FRigVMPinInfoArray& OutPinArray);
+	static bool ShouldCreatePinForProperty(const FRigVMTrait_AnimNextPublicVariables& InTrait, const FProperty* InProperty);
+};
+
+}
\ No newline at end of file
diff --git a/Engine/Plugins/Experimental/AnimNextAnimGraph/Source/AnimNextAnimGraphUncookedOnly/Private/Module.cpp b/Engine/Plugins/Experimental/AnimNextAnimGraph/Source/AnimNextAnimGraphUncookedOnly/Private/Module.cpp
deleted file mode 100644
