Nonunity fixes

#rnx

diff --git a/Engine/Plugins/Experimental/AnimNextAnimGraph/Source/AnimNextAnimGraph/Private/Graph/RigUnit_AnimNextSetNotifyContext.cpp b/Engine/Plugins/Experimental/AnimNextAnimGraph/Source/AnimNextAnimGraph/Private/Graph/RigUnit_AnimNextSetNotifyContext.cpp
--- a/Engine/Plugins/Experimental/AnimNextAnimGraph/Source/AnimNextAnimGraph/Private/Graph/RigUnit_AnimNextSetNotifyContext.cpp
+++ b/Engine/Plugins/Experimental/AnimNextAnimGraph/Source/AnimNextAnimGraph/Private/Graph/RigUnit_AnimNextSetNotifyContext.cpp
@@ -1,8 +1,9 @@
-﻿// Copyright Epic Games, Inc. All Rights Reserved.
+// Copyright Epic Games, Inc. All Rights Reserved.
 
 #include "RigUnit_AnimNextSetNotifyContext.h"
 
 #include "Components/SkeletalMeshComponent.h"
+#include "Module/AnimNextModuleInstance.h"
 #include "Traits/NotifyDispatcher.h"
 
 FRigUnit_AnimNextSetNotifyContext_Execute()
diff --git a/Engine/Plugins/Experimental/AnimNextAnimGraph/Source/AnimNextAnimGraph/Private/Graph/RigUnit_GetPostProcessAnimation.cpp b/Engine/Plugins/Experimental/AnimNextAnimGraph/Source/AnimNextAnimGraph/Private/Graph/RigUnit_GetPostProcessAnimation.cpp
--- a/Engine/Plugins/Experimental/AnimNextAnimGraph/Source/AnimNextAnimGraph/Private/Graph/RigUnit_GetPostProcessAnimation.cpp
+++ b/Engine/Plugins/Experimental/AnimNextAnimGraph/Source/AnimNextAnimGraph/Private/Graph/RigUnit_GetPostProcessAnimation.cpp
@@ -5,6 +5,7 @@
 #include "Components/SkeletalMeshComponent.h"
 #include "Engine/SkeletalMesh.h"
 #include "Graph/PostProcessAnimationAssetUserData.h"
+#include "Module/AnimNextModuleInstance.h"
 #include "Module/AnimNextSkeletalMeshComponentReferenceComponent.h"
 
 FRigUnit_GetPostProcessAnimation_Execute()
