Avoid blend cost at Anim Next Passthrough Blend, when blend type is float or bool, and the calculated weight is either full relevant (1) or not relevant (0)

diff --git a/Engine/Plugins/Experimental/AnimNextAnimGraph/Source/AnimNextAnimGraph/Private/Traits/PassthroughBlendTrait.cpp b/Engine/Plugins/Experimental/AnimNextAnimGraph/Source/AnimNextAnimGraph/Private/Traits/PassthroughBlendTrait.cpp
--- a/Engine/Plugins/Experimental/AnimNextAnimGraph/Source/AnimNextAnimGraph/Private/Traits/PassthroughBlendTrait.cpp
+++ b/Engine/Plugins/Experimental/AnimNextAnimGraph/Source/AnimNextAnimGraph/Private/Traits/PassthroughBlendTrait.cpp
@@ -25,7 +25,7 @@ void FPassthroughBlendTrait::FInstanceData::Construct(const FExecutionContext& C
 {
 	const FSharedData* SharedData = Binding.GetSharedData<FSharedData>();
 
-	ComputedAlphaValue = SharedData->GetAlpha(Binding);;
+	ComputedAlphaValue = SharedData->GetAlpha(Binding);
 	AlphaBoolBlend = SharedData->GetAlphaBoolBlend(Binding);
 	AlphaScaleBiasClamp = SharedData->AlphaScaleBiasClamp;
 }
@@ -59,18 +59,18 @@ void FPassthroughBlendTrait::PreUpdate(FUpdateTraversalContext& Context, const T
 // IEvaluate impl
 void FPassthroughBlendTrait::PostEvaluate(FEvaluateTraversalContext& Context, const TTraitBinding<IEvaluate>& Binding) const
 {
-	// Duplicate existing pose in the stack
-	Context.AppendTask(FAnimNextDuplicateTopKeyframeTask::Make());
-
-	// Let children perform their tasks on the duplicated pose
-	IEvaluate::PostEvaluate(Context, Binding);
-
 	const FSharedData* SharedData = Binding.GetSharedData<FSharedData>();
 	FInstanceData* InstanceData = Binding.GetInstanceData<FInstanceData>();
 
 	const EAnimAlphaInputType AlphaInputType = SharedData->GetAlphaInputType(Binding);
 	if (AlphaInputType == EAnimAlphaInputType::Curve)
 	{
+		// Duplicate existing pose in the stack
+		Context.AppendTask(FAnimNextDuplicateTopKeyframeTask::Make());
+
+		// Let children perform their tasks on the duplicated pose
+		IEvaluate::PostEvaluate(Context, Binding);
+
 		// Add a BlendTwoKey with the original pose and the children modified copy, using alpha curve data (value will be calculated / clamped at the task)
 		const float DeltaTime = InstanceData->DeltaTime;
 		const FName AlphaCurveName = SharedData->GetAlphaCurveName(Binding);
@@ -83,13 +83,28 @@ void FPassthroughBlendTrait::PostEvaluate(FEvaluateTraversalContext& Context, co
 	}
 	else
 	{
-		// Add a BlendTwoKey with the original pose and the children modified copy, using alpha value as float blend weight
 		const float BlendWeight = InstanceData->ComputedAlphaValue;
-		Context.AppendTask(FAnimNextBlendTwoKeyframesTask::Make(BlendWeight));
+		if (FAnimWeight::IsFullWeight(BlendWeight))
+		{
+			IEvaluate::PostEvaluate(Context, Binding);	// full means children result without any modification
+		}
+		else if (FAnimWeight::IsRelevant(BlendWeight)) // if weight is relevant (> 0) we duplicate keyframe, call children to add the task (if any) and blend
+		{
+			// Duplicate existing pose in the stack
+			Context.AppendTask(FAnimNextDuplicateTopKeyframeTask::Make());
+
+			IEvaluate::PostEvaluate(Context, Binding);
+
+			// Add a BlendTwoKey with the original pose and the children modified copy, using alpha value as float blend weight
+			Context.AppendTask(FAnimNextBlendTwoKeyframesTask::Make(BlendWeight));
+		}
+		//else 
+		//{
+			// for BlendWeight == 0 we just need the input keyframe from children as is, so not calling base post evaluate 
+		//}
 	}
 }
 
-
 float FPassthroughBlendTrait::ComputeAlphaValue(const EAnimAlphaInputType AlphaInputType, const FPassthroughBlendTraitSharedData* SharedData, FInstanceData* InstanceData, const TTraitBinding<IUpdate>& Binding, float DeltaTime)
 {
 	// alpha handlers
diff --git a/Engine/Plugins/Experimental/AnimNextAnimGraph/Source/AnimNextAnimGraph/Public/Traits/PassthroughBlendTrait.h b/Engine/Plugins/Experimental/AnimNextAnimGraph/Source/AnimNextAnimGraph/Public/Traits/PassthroughBlendTrait.h
--- a/Engine/Plugins/Experimental/AnimNextAnimGraph/Source/AnimNextAnimGraph/Public/Traits/PassthroughBlendTrait.h
+++ b/Engine/Plugins/Experimental/AnimNextAnimGraph/Source/AnimNextAnimGraph/Public/Traits/PassthroughBlendTrait.h
@@ -11,6 +11,27 @@
 // --- ---
 #include "PassthroughBlendTrait.generated.h"
 
+/*
+* Passthrough Blend :
+* It will generate an output using the base trait input keyframe and the selected alpha blend type (float, bool or curve)
+* Requires a base trait with a single input and PostEvaluate task creation (currently used to blend Control Rig Trait)
+* 
+* If alpha type is Float or Bool
+* Depending on the weight value
+* - 0 (not relevant)          : The input "as is" will be the output (passthrough). No call to PostEvaluate on the base will be performed.
+* - 1 (full weight)           : The output of the base will be used as the output without blending by calling base PostEvaluate
+* - Any other relevant weight : Input will be duplicated (with a task) in Post Evaluate
+*								The base Post Evaluate will be called to execute on the copy (normally will create a task)
+*								Both keyframes (original and copy with base modifications) will be blended using a two keyframes blend task, passing the calculated alpha blend weight
+*
+* If alpha type is Curve
+* The curve value is not known until we are in the task (because we need the keyframe to extract the curve value)
+* Currently, the process will generate a full blend :	A task will added to duplicate input keyframe at PostEvaluate
+*														Base PostEvaluate will be called to allow task creation
+*														Resulting keyframes will be blended with a two keyframes blend task, passing the curve data so the alpha can be obtained
+* If we try to get current alpha value of a Curve blend type, the returned value will be 1
+*/
+
 USTRUCT(meta = (DisplayName = "Passthrough Blend"))
 struct FPassthroughBlendTraitSharedData : public FAnimNextTraitSharedData
 {
