Fix immediate crash in editor due to AnimNext casting redirector into anim graph

diff --git a/Engine/Plugins/Experimental/AnimNextAnimGraph/Source/AnimNextAnimGraph/AnimNextAnimGraph.Build.cs b/Engine/Plugins/Experimental/AnimNextAnimGraph/Source/AnimNextAnimGraph/AnimNextAnimGraph.Build.cs
--- a/Engine/Plugins/Experimental/AnimNextAnimGraph/Source/AnimNextAnimGraph/AnimNextAnimGraph.Build.cs
+++ b/Engine/Plugins/Experimental/AnimNextAnimGraph/Source/AnimNextAnimGraph/AnimNextAnimGraph.Build.cs
@@ -32,6 +32,7 @@ namespace UnrealBuildTool.Rules
 				PrivateDependencyModuleNames.AddRange(
 					new string[]
 					{
+						"AssetRegistry",
 						"RigVMDeveloper",
 					}
 				);
diff --git a/Engine/Plugins/Experimental/AnimNextAnimGraph/Source/AnimNextAnimGraph/Private/AnimNextAnimGraphSettings.cpp b/Engine/Plugins/Experimental/AnimNextAnimGraph/Source/AnimNextAnimGraph/Private/AnimNextAnimGraphSettings.cpp
--- a/Engine/Plugins/Experimental/AnimNextAnimGraph/Source/AnimNextAnimGraph/Private/AnimNextAnimGraphSettings.cpp
+++ b/Engine/Plugins/Experimental/AnimNextAnimGraph/Source/AnimNextAnimGraph/Private/AnimNextAnimGraphSettings.cpp
@@ -10,6 +10,10 @@
 #include "Misc/EnumerateRange.h"
 #include "Modules/ModuleManager.h"
 
+#if WITH_EDITOR
+#include "AssetRegistry/AssetRegistryHelpers.h"
+#endif // WITH_EDITOR
+
 #if WITH_EDITOR
 void UAnimNextAnimGraphSettings::PostEditChangeProperty(FPropertyChangedEvent& PropertyChangedEvent)
 {
@@ -304,7 +308,13 @@ void UAnimNextAnimGraphSettings::OnDefaultRunGraphHostLoaded(const UAnimNextAnim
 
 	for (TEnumerateRef<FAnimNextAssetGraphMapping> Mapping : EnumerateRange(AssetGraphMappings))
 	{
-		Mapping->AnimationGraph.ToSoftObjectPath().LoadAsync(FLoadSoftObjectPathAsyncDelegate::CreateLambda(
+		FSoftObjectPath AnimationGraphSoftPath = Mapping->AnimationGraph.ToSoftObjectPath();
+
+#if WITH_EDITOR
+		UAssetRegistryHelpers::FixupRedirectedAssetPath(IN OUT AnimationGraphSoftPath);
+#endif
+
+		AnimationGraphSoftPath.LoadAsync(FLoadSoftObjectPathAsyncDelegate::CreateLambda(
 			[WeakThis = TWeakObjectPtr<UAnimNextAnimGraphSettings>(this),
 			&Mapping = *Mapping,
 			MappingIndex = Mapping.GetIndex()]
