Demote check in sync code to an ensure to avoid crashing while we investigate the source issue

diff --git a/Engine/Plugins/Experimental/AnimNextAnimGraph/Source/AnimNextAnimGraph/Private/Graph/SyncGroup_GraphInstanceComponent.cpp b/Engine/Plugins/Experimental/AnimNextAnimGraph/Source/AnimNextAnimGraph/Private/Graph/SyncGroup_GraphInstanceComponent.cpp
--- a/Engine/Plugins/Experimental/AnimNextAnimGraph/Source/AnimNextAnimGraph/Private/Graph/SyncGroup_GraphInstanceComponent.cpp
+++ b/Engine/Plugins/Experimental/AnimNextAnimGraph/Source/AnimNextAnimGraph/Private/Graph/SyncGroup_GraphInstanceComponent.cpp
@@ -749,11 +749,12 @@ namespace UE::AnimNext
 				PrevTime -= Duration;
 			}
 
-			check(NextTime > PrevTime);
+			ensure(NextTime > PrevTime);
+			const float TimeBetweenMarkers = FMath::Max(NextTime - PrevTime, SMALL_NUMBER);
 
 			// Store the encoded current time position as a ratio between markers
-			const float PositionBetweenMarkers = (CurrentTime - PrevTime) / (NextTime - PrevTime);
-			check(PositionBetweenMarkers >= (0.0f - SMALL_NUMBER) && PositionBetweenMarkers <= (1.0f + SMALL_NUMBER));
+			const float PositionBetweenMarkers = (CurrentTime - PrevTime) / TimeBetweenMarkers;
+			ensure(PositionBetweenMarkers >= (0.0f - SMALL_NUMBER) && PositionBetweenMarkers <= (1.0f + SMALL_NUMBER));
 
 			SyncPosition.PositionBetweenMarkers = FMath::Clamp(PositionBetweenMarkers, 0.0f, 1.0f);
 			return SyncPosition;
