AnimNext: Fix crash when a factory anim graph mapped to an asset cannot be loaded

diff --git a/Engine/Plugins/Experimental/AnimNextAnimGraph/Source/AnimNextAnimGraph/Private/AnimNextAnimGraphSettings.cpp b/Engine/Plugins/Experimental/AnimNextAnimGraph/Source/AnimNextAnimGraph/Private/AnimNextAnimGraphSettings.cpp
--- a/Engine/Plugins/Experimental/AnimNextAnimGraph/Source/AnimNextAnimGraph/Private/AnimNextAnimGraphSettings.cpp
+++ b/Engine/Plugins/Experimental/AnimNextAnimGraph/Source/AnimNextAnimGraph/Private/AnimNextAnimGraphSettings.cpp
@@ -324,7 +324,7 @@ void UAnimNextAnimGraphSettings::OnDefaultRunGraphHostLoaded(const UAnimNextAnim
 				{
 					ThisPin->NumMappingsLoaded++;
 
-					ThisPin->OnMappingAnimationGraphLoaded(CastChecked<const UAnimNextAnimationGraph>(InObject), Mapping, MappingIndex);
+					ThisPin->OnMappingAnimationGraphLoaded(CastChecked<const UAnimNextAnimationGraph>(InObject, ECastCheckedType::NullAllowed), Mapping, MappingIndex);
 
 					if (ThisPin->NumMappingsLoaded == ThisPin->AssetGraphMappings.Num())
 					{
