AnimNext: Optional auto-expiry of injection requests

Plus additional looping options for PlayAnim requests

diff --git a/Engine/Plugins/Experimental/AnimNextAnimGraph/Source/AnimNextAnimGraph/Private/Injection/InjectionSiteTrait.cpp b/Engine/Plugins/Experimental/AnimNextAnimGraph/Source/AnimNextAnimGraph/Private/Injection/InjectionSiteTrait.cpp
--- a/Engine/Plugins/Experimental/AnimNextAnimGraph/Source/AnimNextAnimGraph/Private/Injection/InjectionSiteTrait.cpp
+++ b/Engine/Plugins/Experimental/AnimNextAnimGraph/Source/AnimNextAnimGraph/Private/Injection/InjectionSiteTrait.cpp
@@ -309,7 +309,10 @@ namespace UE::AnimNext
 				const FAnimNextInjectionRequestArgs& RequestArgs = ActiveInjectionTracker.Request->GetArgs();
 
 				const float BlendOutTime = RequestArgs.BlendOutSettings.Blend.BlendTime;
-				if (CurrentRequestTimeLeft <= BlendOutTime)
+
+				// Only do this if we are automatically blending out
+				if (ActiveInjectionTracker.Request->GetArgs().LifetimeType == EAnimNextInjectionLifetimeType::Auto &&
+					CurrentRequestTimeLeft <= BlendOutTime)
 				{
 					// We are ready to start blending out
 					{
diff --git a/Engine/Plugins/Experimental/AnimNextAnimGraph/Source/AnimNextAnimGraph/Private/Injection/InjectionUtils.cpp b/Engine/Plugins/Experimental/AnimNextAnimGraph/Source/AnimNextAnimGraph/Private/Injection/InjectionUtils.cpp
--- a/Engine/Plugins/Experimental/AnimNextAnimGraph/Source/AnimNextAnimGraph/Private/Injection/InjectionUtils.cpp
+++ b/Engine/Plugins/Experimental/AnimNextAnimGraph/Source/AnimNextAnimGraph/Private/Injection/InjectionUtils.cpp
@@ -47,24 +47,46 @@ namespace UE::AnimNext
 		const FInjectionBlendSettings& InBlendOutSettings,
 		FInjectionLifetimeEvents&& InLifetimeEvents)
 	{
-		return PlayAnim(
+		return PlayAnimHandle(
 			InComponent,
 			InComponent->GetModuleHandle(),
 			InSite,
 			InAnimSequence,
-			InPlayRate,
-			InStartPosition,
+			{
+				.PlayRate = InPlayRate,
+				.StartPosition = InStartPosition
+			},
 			InBlendInSettings,
 			InBlendOutSettings,
 			MoveTemp(InLifetimeEvents));
 	}
 
 	FInjectionRequestPtr FInjectionUtils::PlayAnim(
+		UAnimNextComponent* InComponent,
+		const FInjectionSite& InSite,
+		UAnimSequence* InAnimSequence,
+		const FPlayAnimArgs& InArgs,
+		const FInjectionBlendSettings& InBlendInSettings,
+		const FInjectionBlendSettings& InBlendOutSettings,
+		FInjectionLifetimeEvents&& InLifetimeEvents)
+	{
+		return PlayAnimHandle(
+			InComponent,
+			InComponent->GetModuleHandle(),
+			InSite,
+			InAnimSequence,
+			InArgs,
+			InBlendInSettings,
+			InBlendOutSettings,
+			MoveTemp(InLifetimeEvents));
+	}
+
+	FInjectionRequestPtr FInjectionUtils::PlayAnimHandle(
 		UObject* InHost,
 		FModuleHandle InModuleHandle,
 		const FInjectionSite& InSite,
 		UAnimSequence* InAnimSequence,
-		float InPlayRate, float InStartPosition,
+		const FPlayAnimArgs& InArgs,
 		const FInjectionBlendSettings& InBlendInSettings,
 		const FInjectionBlendSettings& InBlendOutSettings,
 		FInjectionLifetimeEvents&& InLifetimeEvents)
@@ -75,9 +97,20 @@ namespace UE::AnimNext
 		Payload.InitializeAs<FAnimNextNativeDataInterface_AnimSequencePlayer>();
 		FAnimNextNativeDataInterface_AnimSequencePlayer& PlayNativeInterface = Payload.GetMutable<FAnimNextNativeDataInterface_AnimSequencePlayer>();
 		PlayNativeInterface.AnimSequence = InAnimSequence;
-		PlayNativeInterface.PlayRate = InPlayRate;
-		PlayNativeInterface.StartPosition = InStartPosition;
-		PlayNativeInterface.Loop = InAnimSequence->bLoop;
+		PlayNativeInterface.PlayRate = InArgs.PlayRate;
+		PlayNativeInterface.StartPosition = InArgs.StartPosition;
+		switch (InArgs.LoopMode)
+		{
+		case ELoopMode::Auto:
+			PlayNativeInterface.Loop = InAnimSequence->bLoop;
+			break;
+		case ELoopMode::ForceLoop:
+			PlayNativeInterface.Loop = true;
+			break;
+		case ELoopMode::ForceNonLoop:
+			PlayNativeInterface.Loop = false;
+			break;
+		}
 
 		FInjectionRequestArgs RequestArgs;
 		RequestArgs.Site = InSite;
@@ -86,6 +119,7 @@ namespace UE::AnimNext
 		RequestArgs.BlendInSettings = InBlendInSettings;
 		RequestArgs.BlendOutSettings = InBlendOutSettings;
 		RequestArgs.Payload.AddNative(MoveTemp(Payload));
+		RequestArgs.LifetimeType = InArgs.LifetimeType;
 
 		return Inject(InHost, InModuleHandle, MoveTemp(RequestArgs), MoveTemp(InLifetimeEvents));
 	}
diff --git a/Engine/Plugins/Experimental/AnimNextAnimGraph/Source/AnimNextAnimGraph/Private/PlayAnim/PlayAnimCallbackProxy.cpp b/Engine/Plugins/Experimental/AnimNextAnimGraph/Source/AnimNextAnimGraph/Private/PlayAnim/PlayAnimCallbackProxy.cpp
--- a/Engine/Plugins/Experimental/AnimNextAnimGraph/Source/AnimNextAnimGraph/Private/PlayAnim/PlayAnimCallbackProxy.cpp
+++ b/Engine/Plugins/Experimental/AnimNextAnimGraph/Source/AnimNextAnimGraph/Private/PlayAnim/PlayAnimCallbackProxy.cpp
@@ -48,12 +48,17 @@ bool UPlayAnimCallbackProxy::Play(
 		LifetimeEvents.OnInterrupted.BindUObject(this, &UPlayAnimCallbackProxy::OnPlayAnimInterrupted);
 		LifetimeEvents.OnBlendingOut.BindUObject(this, &UPlayAnimCallbackProxy::OnPlayAnimBlendingOut);
 
+		UE::AnimNext::FInjectionUtils::FPlayAnimArgs Args;
+		Args.StartPosition = StartPosition;
+		Args.PlayRate = PlayRate;
 		PlayingRequest = UE::AnimNext::FInjectionUtils::PlayAnim(
 			AnimNextComponent,
 			UE::AnimNext::FInjectionSite(SiteName),
 			AnimSequence,
-			PlayRate,
-			StartPosition,
+			{ 
+				.PlayRate = PlayRate,
+				.StartPosition = StartPosition
+			},
 			BlendInSettings,
 			BlendOutSettings,
 			MoveTemp(LifetimeEvents));
diff --git a/Engine/Plugins/Experimental/AnimNextAnimGraph/Source/AnimNextAnimGraph/Public/Injection/InjectionRequest.h b/Engine/Plugins/Experimental/AnimNextAnimGraph/Source/AnimNextAnimGraph/Public/Injection/InjectionRequest.h
--- a/Engine/Plugins/Experimental/AnimNextAnimGraph/Source/AnimNextAnimGraph/Public/Injection/InjectionRequest.h
+++ b/Engine/Plugins/Experimental/AnimNextAnimGraph/Source/AnimNextAnimGraph/Public/Injection/InjectionRequest.h
@@ -77,6 +77,20 @@ enum class EAnimNextInjectionType : uint8
 	EvaluationModifier,
 };
 
+// Lifetime behavior for injection requests
+UENUM()
+enum class EAnimNextInjectionLifetimeType : uint8
+{
+	// Automatically un-inject this request when the timeline of the injection request expires.
+	// Looping/infinite timelines will never auto-expire and must be uninjected manually.
+	Auto,
+
+	// Injection request persists and must be un-injected manually.
+	// Blending out to the source animation at the injection site does not occur.
+	// Useful when you want to just push a chain of animations to an injection site, never seeing the source/passthrough pose.
+	ForcePersistent,
+};
+
 /**
  * Injection Request Arguments
  *
@@ -121,6 +135,10 @@ private:
 	UPROPERTY()
 	EAnimNextInjectionType Type = EAnimNextInjectionType::InjectObject;
 
+	// Lifetime behavior for the request
+	UPROPERTY()
+	EAnimNextInjectionLifetimeType LifetimeType = EAnimNextInjectionLifetimeType::Auto;
+
 	// Payload that will be applied to the animation graph's variables via its data interfaces.
 	UPROPERTY()
 	FAnimNextDataInterfacePayload Payload;
diff --git a/Engine/Plugins/Experimental/AnimNextAnimGraph/Source/AnimNextAnimGraph/Public/Injection/InjectionUtils.h b/Engine/Plugins/Experimental/AnimNextAnimGraph/Source/AnimNextAnimGraph/Public/Injection/InjectionUtils.h
--- a/Engine/Plugins/Experimental/AnimNextAnimGraph/Source/AnimNextAnimGraph/Public/Injection/InjectionUtils.h
+++ b/Engine/Plugins/Experimental/AnimNextAnimGraph/Source/AnimNextAnimGraph/Public/Injection/InjectionUtils.h
@@ -19,14 +19,48 @@ namespace UE::AnimNext
 	// A collection of wrapper functions for common injection operations
 	struct FInjectionUtils
 	{
+		UE_DEPRECATED(5.6, "Use PlayAnim that takes FPlayAnimArgs")
+		ANIMNEXTANIMGRAPH_API static FInjectionRequestPtr PlayAnim(
+			UAnimNextComponent* InComponent,
+			const FInjectionSite& InSite,
+			UAnimSequence* InAnimSequence,
+			float InPlayRate = 1.0f,
+			float InStartPosition = 0.0f,
+			const FInjectionBlendSettings& InBlendInSettings = FInjectionBlendSettings(),
+			const FInjectionBlendSettings& InBlendOutSettings = FInjectionBlendSettings(),
+			FInjectionLifetimeEvents&& InLifetimeEvents = FInjectionLifetimeEvents());
+
+		// How to loop an animation when playing back
+		enum class ELoopMode : uint8
+		{
+			// Use looping flag on the animation
+			Auto,
+
+			// Force looping
+			ForceLoop,
+
+			// Force non-looping
+			ForceNonLoop,
+
+			// TODO: Specific fractional loop count
+		};
+
+		// Optional arguments to PlayAnim
+		struct FPlayAnimArgs
+		{
+			float PlayRate = 1.0f;
+			float StartPosition = 0.0f;
+			ELoopMode LoopMode = ELoopMode::Auto;
+			EAnimNextInjectionLifetimeType LifetimeType = EAnimNextInjectionLifetimeType::Auto;
+		};
+
 		// Injects an animation sequence into a running AnimNext module
 		// @param InComponent          The UAnimNextComponent that hosts the module being played on
 		// @param InHost               The object that hosts the module (e.g. UAnimNextComponent) being played on
 		// @param InModuleHandle       The module handle being injected into
 		// @param InSite               The injection site to 'play' the animation sequence on
 		// @param InAnimSequence       The animation sequence to play
-		// @param InPlayRate           The play rate multiplier to use
-		// @param InStartPosition      The start position to use
+		// @param InArgs               Optional args to use
 		// @param InBlendInSettings    The blend in settings
 		// @param InBlendOutSettings   The blend in settings
 		// @param InLifetimeEvents     Delegates called for various request lifetime events
@@ -35,20 +69,18 @@ namespace UE::AnimNext
 			UAnimNextComponent* InComponent,
 			const FInjectionSite& InSite,
 			UAnimSequence* InAnimSequence,
-			float InPlayRate = 1.0f,
-			float InStartPosition = 0.0f,
+			const FPlayAnimArgs& InArgs = FPlayAnimArgs(),
 			const FInjectionBlendSettings& InBlendInSettings = FInjectionBlendSettings(),
 			const FInjectionBlendSettings& InBlendOutSettings = FInjectionBlendSettings(),
 			FInjectionLifetimeEvents&& InLifetimeEvents = FInjectionLifetimeEvents());
 
 		// @see PlayAnim
-		ANIMNEXTANIMGRAPH_API static FInjectionRequestPtr PlayAnim(
+		ANIMNEXTANIMGRAPH_API static FInjectionRequestPtr PlayAnimHandle(
 			UObject* InHost,
 			FModuleHandle InModuleHandle,
 			const FInjectionSite& InSite,
 			UAnimSequence* InAnimSequence,
-			float InPlayRate = 1.0f,
-			float InStartPosition = 0.0f,
+			const FPlayAnimArgs& InArgs = FPlayAnimArgs(),
 			const FInjectionBlendSettings& InBlendInSettings = FInjectionBlendSettings(),
 			const FInjectionBlendSettings& InBlendOutSettings = FInjectionBlendSettings(),
 			FInjectionLifetimeEvents&& InLifetimeEvents = FInjectionLifetimeEvents());
