AnimNext: Fix Clang/PS5 designated initailzer/default constructor issue

Moving nested FPlayAnimArgs struct outside of the FInjectionUtils struct allows Clang to compile. Pesky compiler passes!

diff --git a/Engine/Plugins/Experimental/AnimNextAnimGraph/Source/AnimNextAnimGraph/Private/PlayAnim/PlayAnimCallbackProxy.cpp b/Engine/Plugins/Experimental/AnimNextAnimGraph/Source/AnimNextAnimGraph/Private/PlayAnim/PlayAnimCallbackProxy.cpp
--- a/Engine/Plugins/Experimental/AnimNextAnimGraph/Source/AnimNextAnimGraph/Private/PlayAnim/PlayAnimCallbackProxy.cpp
+++ b/Engine/Plugins/Experimental/AnimNextAnimGraph/Source/AnimNextAnimGraph/Private/PlayAnim/PlayAnimCallbackProxy.cpp
@@ -48,9 +48,6 @@ bool UPlayAnimCallbackProxy::Play(
 		LifetimeEvents.OnInterrupted.BindUObject(this, &UPlayAnimCallbackProxy::OnPlayAnimInterrupted);
 		LifetimeEvents.OnBlendingOut.BindUObject(this, &UPlayAnimCallbackProxy::OnPlayAnimBlendingOut);
 
-		UE::AnimNext::FInjectionUtils::FPlayAnimArgs Args;
-		Args.StartPosition = StartPosition;
-		Args.PlayRate = PlayRate;
 		PlayingRequest = UE::AnimNext::FInjectionUtils::PlayAnim(
 			AnimNextComponent,
 			UE::AnimNext::FInjectionSite(SiteName),
diff --git a/Engine/Plugins/Experimental/AnimNextAnimGraph/Source/AnimNextAnimGraph/Public/Injection/InjectionUtils.h b/Engine/Plugins/Experimental/AnimNextAnimGraph/Source/AnimNextAnimGraph/Public/Injection/InjectionUtils.h
--- a/Engine/Plugins/Experimental/AnimNextAnimGraph/Source/AnimNextAnimGraph/Public/Injection/InjectionUtils.h
+++ b/Engine/Plugins/Experimental/AnimNextAnimGraph/Source/AnimNextAnimGraph/Public/Injection/InjectionUtils.h
@@ -1,4 +1,4 @@
-﻿// Copyright Epic Games, Inc. All Rights Reserved.
+// Copyright Epic Games, Inc. All Rights Reserved.
 
 #pragma once
 
@@ -16,6 +16,30 @@ namespace UE::AnimNext
 
 namespace UE::AnimNext
 {
+	// How to loop an animation when playing back
+	enum class ELoopMode : uint8
+	{
+		// Use looping flag on the animation
+		Auto,
+
+		// Force looping
+		ForceLoop,
+
+		// Force non-looping
+		ForceNonLoop,
+
+		// TODO: Specific fractional loop count
+	};
+
+	// Optional arguments to PlayAnim
+	struct FPlayAnimArgs
+	{
+		float PlayRate = 1.0f;
+		float StartPosition = 0.0f;
+		ELoopMode LoopMode = ELoopMode::Auto;
+		EAnimNextInjectionLifetimeType LifetimeType = EAnimNextInjectionLifetimeType::Auto;
+	};
+
 	// A collection of wrapper functions for common injection operations
 	struct FInjectionUtils
 	{
@@ -30,30 +54,6 @@ namespace UE::AnimNext
 			const FInjectionBlendSettings& InBlendOutSettings = FInjectionBlendSettings(),
 			FInjectionLifetimeEvents&& InLifetimeEvents = FInjectionLifetimeEvents());
 
-		// How to loop an animation when playing back
-		enum class ELoopMode : uint8
-		{
-			// Use looping flag on the animation
-			Auto,
-
-			// Force looping
-			ForceLoop,
-
-			// Force non-looping
-			ForceNonLoop,
-
-			// TODO: Specific fractional loop count
-		};
-
-		// Optional arguments to PlayAnim
-		struct FPlayAnimArgs
-		{
-			float PlayRate = 1.0f;
-			float StartPosition = 0.0f;
-			ELoopMode LoopMode = ELoopMode::Auto;
-			EAnimNextInjectionLifetimeType LifetimeType = EAnimNextInjectionLifetimeType::Auto;
-		};
-
 		// Injects an animation sequence into a running AnimNext module
 		// @param InComponent          The UAnimNextComponent that hosts the module being played on
 		// @param InHost               The object that hosts the module (e.g. UAnimNextComponent) being played on
