Added override for GraphItemDetails icon
Changed ::GetItemIcon to have explicit paths for loaded vs unloaded assets (previously would generate full new FAssetData for alive objects through AssetRegistry API)
Deleted some dead-code

diff --git a/Engine/Plugins/Experimental/AnimNext/Source/AnimNextEditor/Internal/Common/AnimNextAssetItemDetails.cpp b/Engine/Plugins/Experimental/AnimNext/Source/AnimNextEditor/Internal/Common/AnimNextAssetItemDetails.cpp
--- a/Engine/Plugins/Experimental/AnimNext/Source/AnimNextEditor/Internal/Common/AnimNextAssetItemDetails.cpp
+++ b/Engine/Plugins/Experimental/AnimNext/Source/AnimNextEditor/Internal/Common/AnimNextAssetItemDetails.cpp
@@ -35,10 +35,23 @@ namespace UE::AnimNext::Editor
 const FSlateBrush* FAnimNextAssetItemDetails::GetItemIcon(const FWorkspaceOutlinerItemExport& Export) const
 {
 	const FSoftObjectPath& AssetPath = Export.GetTopLevelAssetPath();
-	IAssetRegistry& AssetRegistry = FModuleManager::LoadModuleChecked<FAssetRegistryModule>(TEXT("AssetRegistry")).Get();
-	const FAssetData AssetData = AssetRegistry.GetAssetByObjectPath(AssetPath);
-	const UClass* AssetClass = FClassIconFinder::GetIconClassForAssetData(AssetData);
-	return FSlateIconFinder::FindIconForClass(AssetClass).GetIcon();
+
+	// Deal with in-memory objects directly
+	if (UObject* LoadedObject = AssetPath.ResolveObject())
+	{
+		return FSlateIconFinder::FindIconForClass(LoadedObject->GetClass()).GetIcon();
+	}
+
+	// Otherwise force on-disk only query of cached AssetRegistry state
+	IAssetRegistry& AssetRegistry = FModuleManager::GetModuleChecked<FAssetRegistryModule>(TEXT("AssetRegistry")).Get();
+	FAssetData AssetData;
+	if (AssetRegistry.TryGetAssetByObjectPath(AssetPath, AssetData) == UE::AssetRegistry::EExists::Exists)
+	{
+		const UClass* AssetClass = FClassIconFinder::GetIconClassForAssetData(AssetData);
+		return FSlateIconFinder::FindIconForClass(AssetClass).GetIcon();
+	}
+
+	return nullptr;
 }
 
 void FAnimNextAssetItemDetails::RegisterToolMenuExtensions()
diff --git a/Engine/Plugins/Experimental/AnimNextAnimGraph/Source/AnimNextAnimGraphEditor/Private/Graph/AnimNextAnimationGraphItemDetails.cpp b/Engine/Plugins/Experimental/AnimNextAnimGraph/Source/AnimNextAnimGraphEditor/Private/Graph/AnimNextAnimationGraphItemDetails.cpp
--- a/Engine/Plugins/Experimental/AnimNextAnimGraph/Source/AnimNextAnimGraphEditor/Private/Graph/AnimNextAnimationGraphItemDetails.cpp
+++ b/Engine/Plugins/Experimental/AnimNextAnimGraph/Source/AnimNextAnimGraphEditor/Private/Graph/AnimNextAnimationGraphItemDetails.cpp
@@ -19,82 +19,9 @@
 
 namespace UE::AnimNext::Editor
 {
-
-void FAnimNextAnimationGraphItemDetails::RegisterToolMenuExtensions()
-{
-	FToolMenuOwnerScoped OwnerScoped(TEXT("FAnimNextAnimationGraphItemDetails"));
-	UToolMenu* Menu = UToolMenus::Get()->ExtendMenu("WorkspaceOutliner.ItemContextMenu");
-	if (Menu == nullptr)
-	{
-		return;
-	}
-
-	Menu->AddDynamicSection(TEXT("AnimNextAssetItem"), FNewToolMenuDelegate::CreateLambda([](UToolMenu* InMenu)
-	{
-		UWorkspaceItemMenuContext* WorkspaceItemContext = InMenu->FindContext<UWorkspaceItemMenuContext>();
-		const UAssetEditorToolkitMenuContext* AssetEditorContext = InMenu->FindContext<UAssetEditorToolkitMenuContext>();
-		if (WorkspaceItemContext == nullptr || AssetEditorContext == nullptr)
-		{
-			return;
-		}
-
-		const TSharedPtr<UE::Workspace::IWorkspaceEditor> WorkspaceEditor = StaticCastSharedPtr<UE::Workspace::IWorkspaceEditor>(AssetEditorContext->Toolkit.Pin());
-		if(!WorkspaceEditor.IsValid())
-		{
-			return;
-		}
-
-		if (WorkspaceItemContext->SelectedExports.Num() != 1)
-		{
-			return;
-		}
-
-		TInstancedStruct<FWorkspaceOutlinerItemData>& Data = WorkspaceItemContext->SelectedExports[0].GetData();
-		if (!Data.IsValid() || !Data.GetScriptStruct()->IsChildOf(FAnimNextAnimationGraphOutlinerData::StaticStruct()))
-		{
-			return;
-		}
-
-		const FAnimNextAnimationGraphOutlinerData& OutlinerData = Data.Get<FAnimNextAnimationGraphOutlinerData>();
-		UAnimNextAnimationGraph* Asset = OutlinerData.GetAnimationGraph();
-		if(Asset == nullptr)
-		{
-			return;
-		}
-
-		UAnimNextAnimationGraph_EditorData* EditorData = UncookedOnly::FUtils::GetEditorData<UAnimNextAnimationGraph_EditorData>(Asset);
-		if(EditorData == nullptr)
-		{
-			return;
-		}
-
-		FToolMenuSection& AssetSection = InMenu->AddSection("AnimNextAsset", LOCTEXT("AnimNextAssetSectionLabel", "AnimNext Asset"));
-		const UClass* SubEntryClass = UAnimNextAnimationGraphEntry::StaticClass();
-
-		static FTextFormat AddEntryLabelFormat(LOCTEXT("AddEntryLabelFormat", "Add {0}"));
-		static FTextFormat AddEntryTooltipFormat(LOCTEXT("AddEntryTooltipFormat", "Adds a new {0} to this asset"));
-
-		AssetSection.AddMenuEntry(
-			SubEntryClass->GetFName(),
-			FText::Format(AddEntryLabelFormat, SubEntryClass->GetDisplayNameText()),
-			FText::Format(AddEntryTooltipFormat, SubEntryClass->GetDisplayNameText()),
-			FSlateIconFinder::FindIconForClass(SubEntryClass, "ClassIcon.Object"),
-			FUIAction(
-				FExecuteAction::CreateWeakLambda(WorkspaceItemContext, [&OutlinerData, EditorData]()
-				{
-					FScopedTransaction Transaction(LOCTEXT("AddAnimationGraph", "Add Animation Graph"));
-					EditorData->AddAnimationGraph(TEXT("Root"));
-				})
-			));
-	}));
-}
-
-void FAnimNextAnimationGraphItemDetails::UnregisterToolMenuExtensions()
+const FSlateBrush* FAnimNextAnimationGraphItemDetails::GetItemIcon(const FWorkspaceOutlinerItemExport& Export) const
 {
-	if(UToolMenus* ToolMenus = UToolMenus::Get())
-	{
-		ToolMenus->UnregisterOwnerByName("FAnimNextAnimationGraphItemDetails");
-	}
+	return FAppStyle::GetBrush(TEXT("GraphEditor.EventGraph_24x"));
 }
 
 }
diff --git a/Engine/Plugins/Experimental/AnimNextAnimGraph/Source/AnimNextAnimGraphEditor/Private/Graph/AnimNextAnimationGraphItemDetails.h b/Engine/Plugins/Experimental/AnimNextAnimGraph/Source/AnimNextAnimGraphEditor/Private/Graph/AnimNextAnimationGraphItemDetails.h
--- a/Engine/Plugins/Experimental/AnimNextAnimGraph/Source/AnimNextAnimGraphEditor/Private/Graph/AnimNextAnimationGraphItemDetails.h
+++ b/Engine/Plugins/Experimental/AnimNextAnimGraph/Source/AnimNextAnimGraphEditor/Private/Graph/AnimNextAnimationGraphItemDetails.h
@@ -11,9 +11,9 @@ class FAnimNextAnimationGraphItemDetails : public FAnimNextAssetItemDetails
 {
 public:
 	FAnimNextAnimationGraphItemDetails() = default;
+	
+	ANIMNEXTANIMGRAPHEDITOR_API virtual const FSlateBrush* GetItemIcon(const FWorkspaceOutlinerItemExport& Export) const override;
 
-	static void RegisterToolMenuExtensions();
-	static void UnregisterToolMenuExtensions();
 };
 
 }
