Fix invalid property offset being returned when only controls are exposed in a ControlRig trait

diff --git a/Engine/Plugins/Experimental/AnimNextAnimGraph/Source/AnimNextAnimGraph/Private/TraitCore/TraitReader.cpp b/Engine/Plugins/Experimental/AnimNextAnimGraph/Source/AnimNextAnimGraph/Private/TraitCore/TraitReader.cpp
--- a/Engine/Plugins/Experimental/AnimNextAnimGraph/Source/AnimNextAnimGraph/Private/TraitCore/TraitReader.cpp
+++ b/Engine/Plugins/Experimental/AnimNextAnimGraph/Source/AnimNextAnimGraph/Private/TraitCore/TraitReader.cpp
@@ -177,6 +177,8 @@ namespace UE::AnimNext
 						if (Metadata.RigVMIndex != MAX_uint16)
 						{
 							const FTraitLatentPropertyMemoryLayout PropertyMemoryLayout = Trait->GetLatentPropertyMemoryLayout(*SharedData, Metadata.Name, LatentPropertyIndex);
+							check(PropertyMemoryLayout.Size != 0);
+							check(PropertyMemoryLayout.Alignment != 0 && FMath::IsPowerOfTwo(PropertyMemoryLayout.Alignment));
 
 							// Align our property
 							LatentPropertyOffset = Align(LatentPropertyOffset, PropertyMemoryLayout.Alignment);
diff --git a/Engine/Plugins/Experimental/AnimNextControlRig/Source/AnimNextControlRig/Private/ControlRigTrait.cpp b/Engine/Plugins/Experimental/AnimNextControlRig/Source/AnimNextControlRig/Private/ControlRigTrait.cpp
--- a/Engine/Plugins/Experimental/AnimNextControlRig/Source/AnimNextControlRig/Private/ControlRigTrait.cpp
+++ b/Engine/Plugins/Experimental/AnimNextControlRig/Source/AnimNextControlRig/Private/ControlRigTrait.cpp
@@ -408,7 +408,10 @@ void FControlRigTrait::SerializeTraitSharedData(FArchive& Ar, FAnimNextTraitShar
 						}
 					}
 				}
+			}
 
+			if (ControlRigSharedData.ExposedPropertyControlNames.Num() > 0)
+			{
 				// Here I can not get the controls list, so I just use the exposed control names and types
 				const int32 NumControls = ControlRigSharedData.ExposedPropertyControlNames.Num();
 				check(NumControls == ControlRigSharedData.ExposedPropertyControlTypes.Num());
