Sync load AnimNext graph factory assets in cooked builds

diff --git a/Engine/Plugins/Experimental/AnimNextAnimGraph/Source/AnimNextAnimGraph/Internal/AnimNextAnimGraphSettings.h b/Engine/Plugins/Experimental/AnimNextAnimGraph/Source/AnimNextAnimGraph/Internal/AnimNextAnimGraphSettings.h
--- a/Engine/Plugins/Experimental/AnimNextAnimGraph/Source/AnimNextAnimGraph/Internal/AnimNextAnimGraphSettings.h
+++ b/Engine/Plugins/Experimental/AnimNextAnimGraph/Source/AnimNextAnimGraph/Internal/AnimNextAnimGraphSettings.h
@@ -74,7 +74,7 @@ public:
 	void GetNonNativePayloadFromGraph(const UObject* InObject, const UAnimNextAnimationGraph* InAnimationGraph, FInstancedPropertyBag& InOutGraphPayload) const;
 
 	// Rebuild mappings (loading assets synchronously if required) for lookup
-	void LoadAndRebuildMappings();
+	void LoadAndRebuildMappings(bool bLoadAsync);
 
 	// Get the animation graph we run by default when hosting in a module
 	const UAnimNextAnimationGraph* GetDefaultRunGraphHost() const
@@ -94,7 +94,7 @@ private:
 
 	TArray<UClass*> GetAllowedAssetClassesImpl() const;
 
-	void OnDefaultRunGraphHostLoaded(const UAnimNextAnimationGraph* AnimationGraph);
+	void OnDefaultRunGraphHostLoaded(const UAnimNextAnimationGraph* AnimationGraph, bool bLoadAsync);
 	void OnMappingAnimationGraphLoaded(const UAnimNextAnimationGraph* AnimationGraph, FAnimNextAssetGraphMapping& Mapping, int32 MappingIndex);
 	void FinalizeAsyncLoad();
 	
diff --git a/Engine/Plugins/Experimental/AnimNextAnimGraph/Source/AnimNextAnimGraph/Private/AnimNextAnimGraphModule.cpp b/Engine/Plugins/Experimental/AnimNextAnimGraph/Source/AnimNextAnimGraph/Private/AnimNextAnimGraphModule.cpp
--- a/Engine/Plugins/Experimental/AnimNextAnimGraph/Source/AnimNextAnimGraph/Private/AnimNextAnimGraphModule.cpp
+++ b/Engine/Plugins/Experimental/AnimNextAnimGraph/Source/AnimNextAnimGraph/Private/AnimNextAnimGraphModule.cpp
@@ -34,7 +34,16 @@ namespace UE::AnimNext::AnimGraph
 static FDelayedAutoRegisterHelper AutoRegisterFactoryMappings(EDelayedRegisterRunPhase::EndOfEngineInit, []
 {
 	UAnimNextAnimGraphSettings* Settings = GetMutableDefault<UAnimNextAnimGraphSettings>();
-	Settings->LoadAndRebuildMappings();
+
+#if WITH_EDITOR
+	// Async load in the editor to avoid stalling startup
+	const bool bLoadAsync = true;
+#else
+	// Outside the editor, we assume we are behind a loading screen of sorts
+	const bool bLoadAsync = false;
+#endif
+
+	Settings->LoadAndRebuildMappings(bLoadAsync);
 });
 
 void FAnimNextAnimGraphModule::StartupModule()
diff --git a/Engine/Plugins/Experimental/AnimNextAnimGraph/Source/AnimNextAnimGraph/Private/AnimNextAnimGraphSettings.cpp b/Engine/Plugins/Experimental/AnimNextAnimGraph/Source/AnimNextAnimGraph/Private/AnimNextAnimGraphSettings.cpp
--- a/Engine/Plugins/Experimental/AnimNextAnimGraph/Source/AnimNextAnimGraph/Private/AnimNextAnimGraphSettings.cpp
+++ b/Engine/Plugins/Experimental/AnimNextAnimGraph/Source/AnimNextAnimGraph/Private/AnimNextAnimGraphSettings.cpp
@@ -25,7 +25,7 @@ void UAnimNextAnimGraphSettings::PostEditChangeProperty(FPropertyChangedEvent& P
 			PropertyChangedEvent.Property->GetFName() == GET_MEMBER_NAME_CHECKED(FAnimNextAssetGraphMapping, Variable) ||
 			PropertyChangedEvent.Property->GetFName() == GET_MEMBER_NAME_CHECKED(FAnimNextAssetGraphMapping, RequiredDataInterfaces))
 		{
-			LoadAndRebuildMappings();
+			LoadAndRebuildMappings(false);
 		}
 	}
 }
@@ -288,7 +288,7 @@ void UAnimNextAnimGraphSettings::GetNonNativePayloadFromGraph(const UObject* InO
 	ObjectProperty->SetObjectPropertyValue(ValuePtr, const_cast<UObject*>(InObject));
 }
 
-void UAnimNextAnimGraphSettings::OnDefaultRunGraphHostLoaded(const UAnimNextAnimationGraph* AnimationGraph)
+void UAnimNextAnimGraphSettings::OnDefaultRunGraphHostLoaded(const UAnimNextAnimationGraph* AnimationGraph, bool bLoadAsync)
 {
 	if (AnimationGraph == nullptr)
 	{
@@ -314,7 +314,7 @@ void UAnimNextAnimGraphSettings::OnDefaultRunGraphHostLoaded(const UAnimNextAnim
 		UAssetRegistryHelpers::FixupRedirectedAssetPath(IN OUT AnimationGraphSoftPath);
 #endif
 
-		AnimationGraphSoftPath.LoadAsync(FLoadSoftObjectPathAsyncDelegate::CreateLambda(
+		int32 PackageID = AnimationGraphSoftPath.LoadAsync(FLoadSoftObjectPathAsyncDelegate::CreateLambda(
 			[WeakThis = TWeakObjectPtr<UAnimNextAnimGraphSettings>(this),
 			&Mapping = *Mapping,
 			MappingIndex = Mapping.GetIndex()]
@@ -332,6 +332,11 @@ void UAnimNextAnimGraphSettings::OnDefaultRunGraphHostLoaded(const UAnimNextAnim
 					}
 				}
 			}));
+
+		if (!bLoadAsync)
+		{
+			FlushAsyncLoading(PackageID);
+		}
 	}
 }
 
@@ -414,7 +419,7 @@ void UAnimNextAnimGraphSettings::FinalizeAsyncLoad()
 	bMappingsBuiltAtLeastOnce.store(true, std::memory_order_release);
 }
 
-void UAnimNextAnimGraphSettings::LoadAndRebuildMappings()
+void UAnimNextAnimGraphSettings::LoadAndRebuildMappings(bool bLoadAsync)
 {
 	using namespace UE::AnimNext::AnimGraph;
 
@@ -427,13 +432,18 @@ void UAnimNextAnimGraphSettings::LoadAndRebuildMappings()
 	AnimGraphModule.LoadedGraphs.Reset();
 	AssetGraphMap.Reset();
 
-	DefaultRunGraphHost.ToSoftObjectPath().LoadAsync(FLoadSoftObjectPathAsyncDelegate::CreateLambda([WeakThis = TWeakObjectPtr<UAnimNextAnimGraphSettings>(this)](const FSoftObjectPath& InSoftObjectPath, UObject* InObject)
+	int32 PackageID = DefaultRunGraphHost.ToSoftObjectPath().LoadAsync(FLoadSoftObjectPathAsyncDelegate::CreateLambda([WeakThis = TWeakObjectPtr<UAnimNextAnimGraphSettings>(this), bLoadAsync](const FSoftObjectPath& InSoftObjectPath, UObject* InObject)
 		{
 			if (TStrongObjectPtr<UAnimNextAnimGraphSettings> ThisPin = WeakThis.Pin())
 			{
-				ThisPin->OnDefaultRunGraphHostLoaded(CastChecked<const UAnimNextAnimationGraph>(InObject));
+				ThisPin->OnDefaultRunGraphHostLoaded(CastChecked<const UAnimNextAnimationGraph>(InObject), bLoadAsync);
 			}
 		}));
+
+	if (!bLoadAsync)
+	{
+		FlushAsyncLoading(PackageID);
+	}
 }
 
 TArray<UClass*> UAnimNextAnimGraphSettings::GetAllowedAssetClasses()
