Fix sync marker advancement for leader

diff --git a/Engine/Plugins/Experimental/AnimNextAnimGraph/Source/AnimNextAnimGraph/Private/Graph/SyncGroup_GraphInstanceComponent.cpp b/Engine/Plugins/Experimental/AnimNextAnimGraph/Source/AnimNextAnimGraph/Private/Graph/SyncGroup_GraphInstanceComponent.cpp
--- a/Engine/Plugins/Experimental/AnimNextAnimGraph/Source/AnimNextAnimGraph/Private/Graph/SyncGroup_GraphInstanceComponent.cpp
+++ b/Engine/Plugins/Experimental/AnimNextAnimGraph/Source/AnimNextAnimGraph/Private/Graph/SyncGroup_GraphInstanceComponent.cpp
@@ -1239,6 +1239,34 @@ namespace UE::AnimNext
 			return bAppliedFixup;
 		}
 
+		static float CalculateElapsedTime(const FTimelineState& StartTimelineState, const FTimelineState& EndTimelineState, bool bIsPlayingForward)
+		{
+			if (bIsPlayingForward)
+			{
+				if (EndTimelineState.GetPosition() >= StartTimelineState.GetPosition())
+				{
+					return EndTimelineState.GetPosition() - StartTimelineState.GetPosition();
+				}
+				else
+				{
+					// Looped
+					return (StartTimelineState.GetDuration() - StartTimelineState.GetPosition()) + EndTimelineState.GetPosition();
+				}
+			}
+			else
+			{
+				if (EndTimelineState.GetPosition() <= StartTimelineState.GetPosition())
+				{
+					return EndTimelineState.GetPosition() - StartTimelineState.GetPosition();
+				}
+				else
+				{
+					// Looped
+					return (EndTimelineState.GetPosition() - EndTimelineState.GetDuration()) - StartTimelineState.GetPosition();
+				}
+			}
+		}
+
 		static void AdvanceLeaderMarkerBased(FExecutionContext& Context, FSyncGroupContext& GroupContext)
 		{
 			FSyncGroupMemberContext& GroupLeaderContext = GroupContext.Members[GroupContext.LeaderIndex];
@@ -1298,11 +1326,15 @@ namespace UE::AnimNext
 			GroupContext.LeaderTimelineState = LeaderEndTimelineState;
 			GroupContext.bIsLeaderPlayingForward = (DeltaTime * PlayRate) >= 0.0f;
 
+			// Compute actual elapsed time
+			// Can't use DeltaTime because it doesn't account for play rate and we might have floating point noise
+			const float ElapsedTime = CalculateElapsedTime(LeaderStartTimelineState, LeaderEndTimelineState, GroupContext.bIsLeaderPlayingForward);
+
 			// Advance and collect the markers we passed
 			FMarkerPair PrevMarkerEnd = PrevMarkerStart;
 			FMarkerPair NextMarkerEnd = NextMarkerStart;
 			AdvanceAndCollectSyncMarkersPassed(
-				LeaderStartTimelineState, DeltaTime, GroupContext.bIsLeaderPlayingForward,
+				LeaderStartTimelineState, ElapsedTime, GroupContext.bIsLeaderPlayingForward,
 				GroupLeaderContext.SyncMarkers, PrevMarkerEnd, NextMarkerEnd, GroupContext.MarkersPassed);
 
 			ValidateLeaderMarkers(LeaderEndTimelineState, GroupLeaderContext.SyncMarkers, PrevMarkerEnd, NextMarkerEnd);
