Fix BoolAlphaInputArgs not updating correctly

diff --git a/Engine/Plugins/Experimental/AnimNextAnimGraph/Source/AnimNextAnimGraph/Private/Traits/Args/BoolAlphaInputArgTrait.cpp b/Engine/Plugins/Experimental/AnimNextAnimGraph/Source/AnimNextAnimGraph/Private/Traits/Args/BoolAlphaInputArgTrait.cpp
--- a/Engine/Plugins/Experimental/AnimNextAnimGraph/Source/AnimNextAnimGraph/Private/Traits/Args/BoolAlphaInputArgTrait.cpp
+++ b/Engine/Plugins/Experimental/AnimNextAnimGraph/Source/AnimNextAnimGraph/Private/Traits/Args/BoolAlphaInputArgTrait.cpp
@@ -49,6 +49,8 @@ void FBoolAlphaInputArgTrait::PreUpdate(FUpdateTraversalContext& Context, const
 	FInstanceData* InstanceData = Binding.GetInstanceData<FInstanceData>();
 
 	InstanceData->ComputedAlphaValue = AlphaInput::ComputeAlphaValueForBool(IN OUT InstanceData->AlphaBoolBlend, SharedData->GetbAlphaBoolEnabled(Binding), TraitState.GetDeltaTime());
+
+	IUpdate::PreUpdate(Context, Binding, TraitState);
 }
 
 FAlphaInputTraitArgs FBoolAlphaInputArgTrait::Get(const FExecutionContext& Context, const TTraitBinding<IAlphaInputArgs>& Binding) const
