Make Anim Next rewind debugger details widgets read/write for now so that assets will be double clickable

diff --git a/Engine/Plugins/Experimental/AnimNext/Source/AnimNextEditor/Private/RewindDebugger/AnimNextModuleTrack.h b/Engine/Plugins/Experimental/AnimNext/Source/AnimNextEditor/Private/RewindDebugger/AnimNextModuleTrack.h
--- a/Engine/Plugins/Experimental/AnimNext/Source/AnimNextEditor/Private/RewindDebugger/AnimNextModuleTrack.h
+++ b/Engine/Plugins/Experimental/AnimNext/Source/AnimNextEditor/Private/RewindDebugger/AnimNextModuleTrack.h
@@ -17,7 +17,7 @@ class UPropertyBagDetailsObject : public UObject
 	GENERATED_BODY()
 
 public:
-	UPROPERTY(VisibleAnywhere, Category=Details, meta=(ShowOnlyInnerProperties))
+	UPROPERTY(EditAnywhere, Category=Details, meta=(ShowOnlyInnerProperties))
 	FInstancedPropertyBag Properties;
 };
 
diff --git a/Engine/Plugins/Experimental/AnimNextAnimGraph/Source/AnimNextAnimGraph/Public/EvaluationVM/SerializableEvaluationProgram.h b/Engine/Plugins/Experimental/AnimNextAnimGraph/Source/AnimNextAnimGraph/Public/EvaluationVM/SerializableEvaluationProgram.h
--- a/Engine/Plugins/Experimental/AnimNextAnimGraph/Source/AnimNextAnimGraph/Public/EvaluationVM/SerializableEvaluationProgram.h
+++ b/Engine/Plugins/Experimental/AnimNextAnimGraph/Source/AnimNextAnimGraph/Public/EvaluationVM/SerializableEvaluationProgram.h
@@ -23,6 +23,6 @@ struct FSerializableEvaluationProgram
 	
 	FSerializableEvaluationProgram(const UE::AnimNext::FEvaluationProgram& Other);
 	
-	UPROPERTY(VisibleAnywhere, Category = "Tasks", meta = (ExpandByDefault))
+	UPROPERTY(EditAnywhere, Category = "Tasks", meta = (ExpandByDefault))
 	TArray<FInstancedStruct> Tasks;	
 };
diff --git a/Engine/Plugins/Experimental/AnimNextAnimGraph/Source/AnimNextAnimGraphEditor/Private/RewindDebugger/EvaluationProgramTrack.h b/Engine/Plugins/Experimental/AnimNextAnimGraph/Source/AnimNextAnimGraphEditor/Private/RewindDebugger/EvaluationProgramTrack.h
--- a/Engine/Plugins/Experimental/AnimNextAnimGraph/Source/AnimNextAnimGraphEditor/Private/RewindDebugger/EvaluationProgramTrack.h
+++ b/Engine/Plugins/Experimental/AnimNextAnimGraph/Source/AnimNextAnimGraphEditor/Private/RewindDebugger/EvaluationProgramTrack.h
@@ -17,7 +17,7 @@ class UEvaluationProgramDetailsObject : public UObject
 	GENERATED_BODY()
 
 public:
-	UPROPERTY(VisibleAnywhere, Category=Details, meta=(ShowOnlyInnerProperties))
+	UPROPERTY(EditAnywhere, Category=Details, meta=(ShowOnlyInnerProperties))
 	FSerializableEvaluationProgram Program;
 };
 
