Reset unique groups in case they are re-used
Convert sync group logging to use the visual logger
Remove dead code

diff --git a/Engine/Plugins/Experimental/AnimNextAnimGraph/Source/AnimNextAnimGraph/Private/Graph/SyncGroup_GraphInstanceComponent.cpp b/Engine/Plugins/Experimental/AnimNextAnimGraph/Source/AnimNextAnimGraph/Private/Graph/SyncGroup_GraphInstanceComponent.cpp
--- a/Engine/Plugins/Experimental/AnimNextAnimGraph/Source/AnimNextAnimGraph/Private/Graph/SyncGroup_GraphInstanceComponent.cpp
+++ b/Engine/Plugins/Experimental/AnimNextAnimGraph/Source/AnimNextAnimGraph/Private/Graph/SyncGroup_GraphInstanceComponent.cpp
@@ -1,13 +1,20 @@
 // Copyright Epic Games, Inc. All Rights Reserved.
 
 #include "Graph/SyncGroup_GraphInstanceComponent.h"
+
+#include "Graph/AnimNextAnimationGraph.h"
+#include "TraitCore/NodeInstance.h"
 #include "TraitInterfaces/IGroupSynchronization.h"
 #include "TraitInterfaces/ITimeline.h"
 #include "Animation/AnimTypes.h"
 #include "Animation/AnimationAsset.h"
 #include "AnimationUtils.h"
+#include "Module/AnimNextModuleInstance.h"
+#include "VisualLogger/VisualLogger.h"
 
 // Enabled by default in development builds
+// We log debug information using the visual logger
+// Use Rewind Debugger to record and replay logs in editor
 #define UE_DEBUG_SYNC_GROUPS !(UE_BUILD_SHIPPING || UE_BUILD_TEST)
 
 namespace UE::AnimNext
@@ -15,8 +22,8 @@ namespace UE::AnimNext
 	namespace Private
 	{
 #if UE_DEBUG_SYNC_GROUPS
-		// Enables/disable debug logging in the console
-		static TAutoConsoleVariable<bool> CDebugVarEnableSyncLog(TEXT("a.AnimNext.EnableSyncLog"), false, TEXT("Toggles debug sync group logging"));
+		// Enables/disable debug logging
+		static TAutoConsoleVariable<bool> CDebugVarEnableSyncLog(TEXT("a.AnimNext.EnableSyncLog"), true, TEXT("Toggles debug sync group logging"));
 
 		// Enables sync mode debugging
 		// 0: Auto Sync (uses markers when present, time otherwise, default behavior)
@@ -62,7 +69,7 @@ namespace UE::AnimNext
 			bool bIsActive = false;
 
 			// Whether or not this sync group just formed during the current update
-			bool bJustFormed = false;
+			bool bJustFormed = true;
 		};
 
 		// Sync group member context when performing group synchronization
@@ -151,7 +158,7 @@ namespace UE::AnimNext
 			GroupContext.MemberSeenMarkers.Reset();
 		}
 
-		static void FindLeaderIndex(FSyncGroupContext& GroupContext)
+		static void FindLeaderIndex(FExecutionContext& Context, FSyncGroupContext& GroupContext)
 		{
 			// Find our leader by looking at the total weight and the group role
 			const int32 NumMembers = GroupContext.Members.Num();
@@ -200,7 +207,7 @@ namespace UE::AnimNext
 #if UE_DEBUG_SYNC_GROUPS
 			if (CDebugVarEnableSyncLog.GetValueOnAnyThread() && GroupContext.State->PreviousLeaderIndex != LeaderIndex)
 			{
-				UE_LOG(LogAnimMarkerSync, Warning,
+				UE_VLOG_UELOG(Context.GetHostObject(), LogAnimMarkerSync, Verbose,
 					TEXT("[%s] [%p] Is New Leader"),
 					*GroupContext.State->GroupName.ToString(),
 					GroupContext.Members[LeaderIndex].State->TraitPtr.GetNodeInstance());
@@ -499,9 +506,7 @@ namespace UE::AnimNext
 			}
 
 			OutPrevMarker.MarkerIndex = PrevMarkerIndex;
-			OutPrevMarker.TimeToMarker = PrevTimeToMarker;
 			OutNextMarker.MarkerIndex = NextMarkerIndex;
-			OutNextMarker.TimeToMarker = NextTimeToMarker;
 		}
 
 		static FName GetMarkerName(int32 MarkerIndex, const FTimelineSyncMarkerArray& SyncMarkers)
@@ -607,14 +612,14 @@ namespace UE::AnimNext
 					// This goes to minus one as the very fact we are here means
 					// that there is no next marker to use
 					OutNextMarker.MarkerIndex = -1;
-					
+
 					break; // no need to keep searching, we are done
 				}
 			}
 		}
 
 		// Returns the boundary sync markers matching the sync position pair nearest the specified timeline position and the adjusted current time based on the sync position
-		static float GetNearestMatchingBoundarySyncMarkers(const FMarkerSyncAnimPosition& SyncPosition, const FTimelineSyncMarkerArray& SyncMarkers, const FTimelineState& TimelineState, FMarkerPair& OutPrevMarker, FMarkerPair& OutNextMarker)
+		static void GetNearestMatchingBoundarySyncMarkers(const FMarkerSyncAnimPosition& SyncPosition, const FTimelineSyncMarkerArray& SyncMarkers, const FTimelineState& TimelineState, FMarkerPair& OutPrevMarker, FMarkerPair& OutNextMarker)
 		{
 			check(SyncPosition.PreviousMarkerName != NAME_None || SyncPosition.NextMarkerName != NAME_None);
 
@@ -625,27 +630,18 @@ namespace UE::AnimNext
 			//		- Done
 			if (SyncPosition.PreviousMarkerName == NAME_None)
 			{
-				const float PrevMarkerTime = 0.0f;	// Start of sequence
-				float CurrentTime = 0.0f;
-
 				OutPrevMarker.MarkerIndex = MarkerIndexSpecialValues::AnimationBoundary;
 
 				for (int32 MarkerIndex = 0; MarkerIndex < SyncMarkers.Num(); ++MarkerIndex)
 				{
 					if (SyncMarkers[MarkerIndex].GetName() == SyncPosition.NextMarkerName)
 					{
-						const float NextMarkerTime = SyncMarkers[MarkerIndex].GetPosition();
-						CurrentTime = FMath::LerpStable(PrevMarkerTime, NextMarkerTime, SyncPosition.PositionBetweenMarkers);
-
-						OutPrevMarker.TimeToMarker = PrevMarkerTime - CurrentTime;
-
 						OutNextMarker.MarkerIndex = MarkerIndex;
-						OutNextMarker.TimeToMarker = NextMarkerTime - CurrentTime;
 						break;
 					}
 				}
 
-				return CurrentTime;
+				return;
 			}
 
 			// If the sync position's next marker is the end boundary
@@ -655,27 +651,18 @@ namespace UE::AnimNext
 			//		- Done
 			if (SyncPosition.NextMarkerName == NAME_None)
 			{
-				const float NextMarkerTime = TimelineState.GetDuration();	// End of sequence
-				float CurrentTime = 0.0f;
-
 				OutNextMarker.MarkerIndex = MarkerIndexSpecialValues::AnimationBoundary;
 
 				for (int32 MarkerIndex = SyncMarkers.Num() - 1; MarkerIndex >= 0; --MarkerIndex)
 				{
 					if (SyncMarkers[MarkerIndex].GetName() == SyncPosition.PreviousMarkerName)
 					{
-						const float PrevMarkerTime = SyncMarkers[MarkerIndex].GetPosition();
-						CurrentTime = FMath::LerpStable(PrevMarkerTime, NextMarkerTime, SyncPosition.PositionBetweenMarkers);
-
-						OutNextMarker.TimeToMarker = PrevMarkerTime - CurrentTime;
-
 						OutPrevMarker.MarkerIndex = MarkerIndex;
-						OutPrevMarker.TimeToMarker = NextMarkerTime - CurrentTime;
 						break;
 					}
 				}
 
-				return CurrentTime;
+				return;
 			}
 
 			// Otherwise
@@ -685,26 +672,6 @@ namespace UE::AnimNext
 			//		- If no such pair is found, the follower must not contain the pair of markers present on the leader (not possible)
 			//		- Done
 			GetExactBoundarySyncMarkersAt(SyncPosition, SyncMarkers, TimelineState, OutPrevMarker, OutNextMarker);
-
-			float CurrentTime = TimelineState.GetPosition();
-			float PrevMarkerTime = GetMarkerPosition(OutPrevMarker.MarkerIndex, SyncMarkers, 0.0f);
-			float NextMarkerTime = GetMarkerPosition(OutNextMarker.MarkerIndex, SyncMarkers, TimelineState.GetDuration());
-
-			if (TimelineState.IsLooping())
-			{
-				if (PrevMarkerTime > CurrentTime)
-				{
-					PrevMarkerTime -= TimelineState.GetDuration();
-				}
-				else if (NextMarkerTime < CurrentTime)
-				{
-					NextMarkerTime += TimelineState.GetDuration();
-				}
-			}
-
-			// Map our current time in between the markers we found
-			CurrentTime = FMath::LerpStable(PrevMarkerTime, NextMarkerTime, SyncPosition.PositionBetweenMarkers);
-			return CurrentTime;
 		}
 
 		// Returns the relative sync position between the two specified markers
@@ -1121,7 +1088,7 @@ namespace UE::AnimNext
 					(LeaderEndTimelineState.GetPosition() < LeaderStartTimelineState.GetPosition()) :
 					(LeaderStartTimelineState.GetPosition() < LeaderEndTimelineState.GetPosition());
 
-				UE_LOG(LogAnimMarkerSync, Warning,
+				UE_VLOG_UELOG(Context.GetHostObject(), LogAnimMarkerSync, Verbose,
 					TEXT("[%s] [%p] Leader [Time] [%0.2f %s -> %0.2f / %0.2f (%3.2f%%)] Playing [%s]%s"),
 					*GroupContext.State->GroupName.ToString(),
 					GroupLeader.TraitPtr.GetNodeInstance(),
@@ -1391,7 +1358,7 @@ namespace UE::AnimNext
 					JoinedStatus = FString::Printf(TEXT(" (joined from %0.2f)"), LeaderStartTimelineState.GetPosition());
 				}
 
-				UE_LOG(LogAnimMarkerSync, Warning,
+				UE_VLOG_UELOG(Context.GetHostObject(), LogAnimMarkerSync, Verbose,
 					TEXT("[%s] [%p] Leader [Mark] [%0.2f %s -> %0.2f / %0.2f (%3.2f%%)] [%s@%0.2f | %s@%0.2f (%3.2f%%) -> %s@%0.2f | %s@%0.2f (%3.2f%%)] Playing [%s]%s%s%s%s"),
 					*GroupContext.State->GroupName.ToString(),
 					GroupLeader.TraitPtr.GetNodeInstance(),
@@ -1420,7 +1387,7 @@ namespace UE::AnimNext
 		{
 			bool bUseMarkerSyncing = GroupContext.bCanGroupUseMarkerSyncing;
 
-#if !(UE_BUILD_SHIPPING || UE_BUILD_TEST)
+#if UE_DEBUG_SYNC_GROUPS
 			const int32 DebugSyncMode = CDebugVarSyncMode.GetValueOnAnyThread();
 			if (DebugSyncMode == 1)
 			{
@@ -1517,7 +1484,7 @@ namespace UE::AnimNext
 
 				const FTimelineState FollowerEndTimelineState = GroupMemberContext.TimelineTrait.GetState(Context);
 
-				UE_LOG(LogAnimMarkerSync, Warning,
+				UE_VLOG_UELOG(Context.GetHostObject(), LogAnimMarkerSync, Verbose,
 					TEXT("[%s] [%p] Follow [Time] [%0.2f %s -> %0.2f / %0.2f (%3.2f%%)] Playing [%s]%s"),
 					*GroupContext.State->GroupName.ToString(),
 					GroupMember.TraitPtr.GetNodeInstance(),
@@ -1661,7 +1628,7 @@ namespace UE::AnimNext
 					JoinedStatus = FString::Printf(TEXT(" (joined from %0.2f)"), MarkerSearchTime);
 				}
 
-				UE_LOG(LogAnimMarkerSync, Warning,
+				UE_VLOG_UELOG(Context.GetHostObject(), LogAnimMarkerSync, Verbose,
 					TEXT("[%s] [%p] Follow [Mark] [%0.2f %s -> %0.2f / %0.2f (%3.2f%%)] [%s@%0.2f | %s@%0.2f (%3.2f%%) -> %s@%0.2f | %s@%0.2f (%3.2f%%)] Playing [%s]%s%s%s%s%s"),
 					*GroupContext.State->GroupName.ToString(),
 					GroupMember.TraitPtr.GetNodeInstance(),
@@ -1716,7 +1683,7 @@ namespace UE::AnimNext
 
 				const FTimelineState FollowerEndTimelineState = GroupMemberContext.TimelineTrait.GetState(Context);
 
-				UE_LOG(LogAnimMarkerSync, Warning,
+				UE_VLOG_UELOG(Context.GetHostObject(), LogAnimMarkerSync, Verbose,
 					TEXT("[%s] [%p] Follow [Solo] [%0.2f %s -> %0.2f / %0.2f (%3.2f%%)] Playing [%s]"),
 					*GroupContext.State->GroupName.ToString(),
 					GroupMember.TraitPtr.GetNodeInstance(),
@@ -1752,7 +1719,7 @@ namespace UE::AnimNext
 				bool bUseMarkerSyncing = GroupMemberContext.bUseMarkerSyncing;
 				EAnimGroupSynchronizationRole GroupRole = GroupMember.GroupParameters.GroupRole;
 
-#if !(UE_BUILD_SHIPPING || UE_BUILD_TEST)
+#if UE_DEBUG_SYNC_GROUPS
 				const int32 DebugSyncMode = CDebugVarSyncMode.GetValueOnAnyThread();
 				if (DebugSyncMode == 1)
 				{
@@ -1780,6 +1747,12 @@ namespace UE::AnimNext
 				}
 			}
 		}
+
+		static UObject* GetHostObject(const FWeakTraitPtr& TraitPtr)
+		{
+			const FAnimNextModuleInstance* ModuleInstance = TraitPtr.GetNodeInstance()->GetOwner().GetRootGraphInstance()->GetModuleInstance();
+			return ModuleInstance != nullptr ? ModuleInstance->GetObject() : nullptr;
+		}
 	}
 
 	FSyncGroupGraphInstanceComponent::FSyncGroupGraphInstanceComponent(FAnimNextGraphInstance& InOwnerInstance)
@@ -1809,13 +1782,15 @@ namespace UE::AnimNext
 	{
 		using namespace Private;
 
-		const int32 GroupIndex = SyncGroupMap.FindOrAdd(GroupParameters.GroupName, SyncGroups.Num());
-		const bool bJustFormed = GroupIndex == SyncGroups.Num();
+		if (!ensure(TraitPtr.IsValid()))
+		{
+			return;
+		}
 
-		FSyncGroupState& GroupState = bJustFormed ? SyncGroups.AddDefaulted_GetRef() : SyncGroups[GroupIndex];
+		const int32 GroupIndex = SyncGroupMap.FindOrAdd(GroupParameters.GroupName, SyncGroups.Num());
+		FSyncGroupState& GroupState = GroupIndex == SyncGroups.Num() ? SyncGroups.AddDefaulted_GetRef() : SyncGroups[GroupIndex];
 		GroupState.GroupName = GroupParameters.GroupName;
 		GroupState.bIsActive = true;
-		GroupState.bJustFormed = bJustFormed;
 
 		bool bIsMemberNew = true;
 		for (FSyncGroupMember& GroupMember : GroupState.Members)
@@ -1839,10 +1814,11 @@ namespace UE::AnimNext
 #if UE_DEBUG_SYNC_GROUPS
 			if (CDebugVarEnableSyncLog.GetValueOnAnyThread())
 			{
-				UE_LOG(LogAnimMarkerSync, Warning,
-					TEXT("[%s] [%p] Joined"),
+				UE_VLOG_UELOG(GetHostObject(TraitPtr), LogAnimMarkerSync, Verbose,
+					TEXT("[%s] [%p] %s"),
 					*GroupParameters.GroupName.ToString(),
-					TraitPtr.GetNodeInstance());
+					TraitPtr.GetNodeInstance(),
+					GroupState.Members.Num() == 1 ? TEXT("Created") : TEXT("Joined"));
 			}
 #endif
 		}
@@ -1884,12 +1860,38 @@ namespace UE::AnimNext
 
 	void FSyncGroupGraphInstanceComponent::ReleaseUniqueGroupName(FName GroupName)
 	{
-		Private::FSyncGroupUniqueName* Entry = nullptr;
+		using namespace Private;
+
+		FSyncGroupUniqueName* Entry = nullptr;
 		if (ensure(UsedUniqueGroupNames.RemoveAndCopyValue(GroupName, Entry)))
 		{
 			Entry->NextFreeEntry = FirstFreeUniqueGroupName;
 			FirstFreeUniqueGroupName = Entry;
 		}
+
+		if (const int32* GroupIndex = SyncGroupMap.Find(GroupName))
+		{
+			FSyncGroupState& GroupState = SyncGroups[*GroupIndex];
+
+#if UE_DEBUG_SYNC_GROUPS
+			if (CDebugVarEnableSyncLog.GetValueOnAnyThread())
+			{
+				for (int32 GroupMemberIndex = GroupState.Members.Num() - 1; GroupMemberIndex >= 0; GroupMemberIndex--)
+				{
+					const FSyncGroupMember& GroupMember = GroupState.Members[GroupMemberIndex];
+
+					UE_VLOG_UELOG(GetHostObject(GroupMember.TraitPtr), LogAnimMarkerSync, Verbose,
+						TEXT("[%s] [%p] Left"),
+						*GroupState.GroupName.ToString(),
+						GroupMember.TraitPtr.GetNodeInstance());
+				}
+			}
+#endif
+
+			// We are releasing this group, clear any lingering member entries
+			// as we might re-use the group name this update
+			GroupState.Members.Reset();
+		}
 	}
 
 	void FSyncGroupGraphInstanceComponent::PreUpdate(FExecutionContext& Context)
@@ -1905,7 +1907,7 @@ namespace UE::AnimNext
 			}
 
 			GroupState.bIsActive = false;
-			GroupState.bJustFormed = false;
+			GroupState.bJustFormed = GroupState.Members.IsEmpty();
 		}
 	}
 
@@ -1923,6 +1925,9 @@ namespace UE::AnimNext
 
 				if (GroupState.bIsActive)
 				{
+					// Whether or not we have an active member from a previous update
+					bool bHasOldActiveMember = false;
+
 					for (int32 GroupMemberIndex = GroupState.Members.Num() - 1; GroupMemberIndex >= 0; GroupMemberIndex--)
 					{
 						FSyncGroupMember& GroupMember = GroupState.Members[GroupMemberIndex];
@@ -1931,7 +1936,7 @@ namespace UE::AnimNext
 #if UE_DEBUG_SYNC_GROUPS
 							if (CDebugVarEnableSyncLog.GetValueOnAnyThread())
 							{
-								UE_LOG(LogAnimMarkerSync, Warning,
+								UE_VLOG_UELOG(Context.GetHostObject(), LogAnimMarkerSync, Verbose,
 									TEXT("[%s] [%p] Left"),
 									*GroupState.GroupName.ToString(),
 									GroupMember.TraitPtr.GetNodeInstance());
@@ -1941,12 +1946,51 @@ namespace UE::AnimNext
 							// This member is no longer active, remove it
 							GroupState.Members.RemoveAtSwap(GroupMemberIndex, EAllowShrinking::No);
 						}
+						else if (!GroupMember.bJustJoined)
+						{
+							bHasOldActiveMember = true;
+						}
+					}
+
+					// If we don't have any lingering members from a previous update, then we are considered
+					// as a fresh new group
+					if (!bHasOldActiveMember && !GroupState.Members.IsEmpty() && !GroupState.bJustFormed)
+					{
+#if UE_DEBUG_SYNC_GROUPS
+						if (CDebugVarEnableSyncLog.GetValueOnAnyThread())
+						{
+							UE_VLOG_UELOG(Context.GetHostObject(), LogAnimMarkerSync, Verbose,
+								TEXT("[%s] Refreshed"),
+								*GroupState.GroupName.ToString());
+						}
+#endif
+
+						GroupState.bJustFormed = true;
 					}
 
 					MaxNumMembers = FMath::Max<uint32>(MaxNumMembers, GroupState.Members.Num());
 				}
 				else
 				{
+#if UE_DEBUG_SYNC_GROUPS
+					if (CDebugVarEnableSyncLog.GetValueOnAnyThread())
+					{
+						for (int32 GroupMemberIndex = GroupState.Members.Num() - 1; GroupMemberIndex >= 0; GroupMemberIndex--)
+						{
+							const FSyncGroupMember& GroupMember = GroupState.Members[GroupMemberIndex];
+
+							UE_VLOG_UELOG(Context.GetHostObject(), LogAnimMarkerSync, Verbose,
+								TEXT("[%s] [%p] Left"),
+								*GroupState.GroupName.ToString(),
+								GroupMember.TraitPtr.GetNodeInstance());
+						}
+
+						UE_VLOG_UELOG(Context.GetHostObject(), LogAnimMarkerSync, Verbose,
+							TEXT("[%s] Released"),
+							*GroupState.GroupName.ToString());
+					}
+#endif
+
 					// This group is no longer active, remove it
 					SyncGroups.RemoveAtSwap(SyncGroupIndex, EAllowShrinking::No);
 					bRebuildGroupIndexMap = true;
@@ -1979,7 +2023,7 @@ namespace UE::AnimNext
 			InitGroup(GroupState, MemberContexts, GroupContext);
 
 			// Find our leader for this group
-			FindLeaderIndex(GroupContext);
+			FindLeaderIndex(Context, GroupContext);
 
 			// If the leader has sync markers, collect and filter them
 			// We'll retain only markers common to all members
