Fix marker pair detection when it overlaps the start/end of a sequence
Replace check with ensure to gracefully handle invalid sync markers

diff --git a/Engine/Plugins/Animation/PoseSearch/Source/Runtime/Public/PoseSearch/PoseSearchDefines.h b/Engine/Plugins/Animation/PoseSearch/Source/Runtime/Public/PoseSearch/PoseSearchDefines.h
--- a/Engine/Plugins/Animation/PoseSearch/Source/Runtime/Public/PoseSearch/PoseSearchDefines.h
+++ b/Engine/Plugins/Animation/PoseSearch/Source/Runtime/Public/PoseSearch/PoseSearchDefines.h
@@ -42,7 +42,7 @@ static constexpr int32 PreallocatedSearchesNum = 4;
 static constexpr float FiniteDelta = 1 / 60.0f; // Time delta used for computing pose derivatives
 static constexpr int32 MaxNumberOfCollectedPoseCandidatesPerDatabase = 256;
 static constexpr int32 InvalidSearchId = INDEX_NONE; // @todo: it should be 0, since the search id is an hash!
-static constexpr int32 DatabaseIndexDerivedDataCacheKeyVersion = 2687; // bump this value to force reindex ALL the UPoseSearchDatabases, without having to change the global FDevSystemGuids::POSESEARCHDB_DERIVEDDATA_VER
+static constexpr int32 DatabaseIndexDerivedDataCacheKeyVersion = 2688; // bump this value to force reindex ALL the UPoseSearchDatabases, without having to change the global FDevSystemGuids::POSESEARCHDB_DERIVEDDATA_VER
 static constexpr int32 DefaultTickPriority = 100;
 
 template<uint32 Alignment = DEFAULT_ALIGNMENT> class TMemStackSetAllocator : public TSetAllocator<TSparseArrayAllocator<TMemStackAllocator<Alignment>, TMemStackAllocator<Alignment>>, TMemStackAllocator<Alignment>> {};
diff --git a/Engine/Plugins/Experimental/AnimNextAnimGraph/Source/AnimNextAnimGraph/Private/Graph/SyncGroup_GraphInstanceComponent.cpp b/Engine/Plugins/Experimental/AnimNextAnimGraph/Source/AnimNextAnimGraph/Private/Graph/SyncGroup_GraphInstanceComponent.cpp
--- a/Engine/Plugins/Experimental/AnimNextAnimGraph/Source/AnimNextAnimGraph/Private/Graph/SyncGroup_GraphInstanceComponent.cpp
+++ b/Engine/Plugins/Experimental/AnimNextAnimGraph/Source/AnimNextAnimGraph/Private/Graph/SyncGroup_GraphInstanceComponent.cpp
@@ -217,20 +217,29 @@ namespace UE::AnimNext
 			GroupContext.LeaderIndex = LeaderIndex;
 		}
 
-		static void CheckSyncMarkersSorted(const FTimelineSyncMarkerArray& SyncMarkers)
+		// Returns true if sync markers are valid and sorted, false otherwise
+		// Only in non-test/shipping builds
+		static bool CheckSyncMarkersSorted(const FTimelineSyncMarkerArray& SyncMarkers)
 		{
-#if DO_CHECK
-			float PreviousMarkerPosition = -1.0f;
+			bool bIsValid = true;
+
+#if !(UE_BUILD_SHIPPING || UE_BUILD_TEST)
+			float PreviousMarkerPosition = 0.0f;
 			for (const FTimelineSyncMarker& SyncMarker : SyncMarkers)
 			{
 				const float MarkerPosition = SyncMarker.GetPosition();
 
-				checkf(MarkerPosition >= 0.0, TEXT("Sync Markers should have a positive finite position"));
-				checkf(MarkerPosition > PreviousMarkerPosition, TEXT("Sync Markers should be sorted ascending by their position"));
+				if (!ensureMsgf(MarkerPosition >= 0.0 && MarkerPosition >= PreviousMarkerPosition, TEXT("Sync Markers should have a positive finite position and should be sorted in ascending order")))
+				{
+					bIsValid = false;
+					break;
+				}
 
 				PreviousMarkerPosition = MarkerPosition;
 			}
 #endif
+
+			return bIsValid;
 		}
 
 		static void BuildGroupState(FExecutionContext& Context, FSyncGroupContext& GroupContext)
@@ -260,9 +269,8 @@ namespace UE::AnimNext
 
 				SyncMarkers.Reset();
 				GroupLeaderContext.TimelineTrait.GetSyncMarkers(Context, SyncMarkers);
-				CheckSyncMarkersSorted(SyncMarkers);
 
-				bCanGroupUseMarkerSyncing = !SyncMarkers.IsEmpty();
+				bCanGroupUseMarkerSyncing = !SyncMarkers.IsEmpty() && CheckSyncMarkersSorted(SyncMarkers);
 				GroupLeaderContext.bUseMarkerSyncing = bCanGroupUseMarkerSyncing;
 			}
 
@@ -316,10 +324,9 @@ namespace UE::AnimNext
 					FTimelineSyncMarkerArray& SyncMarkers = GroupMemberContext.SyncMarkers;
 					SyncMarkers.Reset();
 					GroupMemberContext.TimelineTrait.GetSyncMarkers(Context, SyncMarkers);
-					CheckSyncMarkersSorted(SyncMarkers);
 
 					// Members without markers can still use time based syncing
-					const bool bCanMemberUseMarkerSyncing = !SyncMarkers.IsEmpty();
+					const bool bCanMemberUseMarkerSyncing = !SyncMarkers.IsEmpty() && CheckSyncMarkersSorted(SyncMarkers);
 					GroupMemberContext.bUseMarkerSyncing = bCanMemberUseMarkerSyncing;
 
 					if (bCanMemberUseMarkerSyncing)
@@ -567,28 +574,42 @@ namespace UE::AnimNext
 					float NextMarkerTime = NextMarker.GetPosition();
 
 					// Handle case where we need to loop to get to be able to get to the next marker.
+					bool bLooped = false;
 					if (NextMarkerTime < PrevMarkerTime)
 					{
 						NextMarkerTime += Duration;
+						bLooped = true;
 					}
 
 					// Get current time based of sync position.
 					float DesiredTime = FMath::LerpStable(PrevMarkerTime, NextMarkerTime, SyncPosition.PositionBetweenMarkers);
 
-					// Adjust for case where we need to loop to be able to get to the next marker.
-					if (DesiredTime > Duration)
-					{
-						DesiredTime -= Duration;
-					}
-
 					// Find marker indices closest to input time position.
-					const float TimeDelta = FMath::Abs(DesiredTime - CurrentTime);
+					float TimeDelta = FMath::Abs(DesiredTime - CurrentTime);
 					if (TimeDelta < BestTimeDelta)
 					{
 						BestTimeDelta = TimeDelta;
 						OutPrevMarker.MarkerIndex = PrevMarkerIdx;
 						OutNextMarker.MarkerIndex = NextMarkerIdx;
 					}
+					else if (bLooped)
+					{
+						// If we looped, we extended our next marker past the end of the sequence
+						// This means that there are two points we need to test:
+						//   - The one that lands near the end of the sequence (possibly overshooting/looping around)
+						//   - The one that lands near the start of the sequence (possibly undershooting/looping around)
+						// 
+						// We tested the first one above, now test the second
+						DesiredTime -= Duration;
+
+						TimeDelta = FMath::Abs(DesiredTime - CurrentTime);
+						if (TimeDelta < BestTimeDelta)
+						{
+							BestTimeDelta = TimeDelta;
+							OutPrevMarker.MarkerIndex = PrevMarkerIdx;
+							OutNextMarker.MarkerIndex = NextMarkerIdx;
+						}
+					}
 
 					// This marker test is done, move onto next one.
 					break;
diff --git a/Engine/Source/Runtime/Engine/Private/Animation/AnimSequence.cpp b/Engine/Source/Runtime/Engine/Private/Animation/AnimSequence.cpp
--- a/Engine/Source/Runtime/Engine/Private/Animation/AnimSequence.cpp
+++ b/Engine/Source/Runtime/Engine/Private/Animation/AnimSequence.cpp
@@ -4312,22 +4312,18 @@ void UAnimSequence::GetMarkerIndicesForPosition(
 					float NextMarkerTime = AuthoredSyncMarkers[NextMarkerIdx].Time;
 
 					// Handle case where we need to loop to get to be able to get to the next marker.
+					bool bLooped = false;
 					if (NextMarkerTime < PrevMarker.Time)
 					{
 						NextMarkerTime += GetPlayLength();
+						bLooped = true;
 					}
 
 					// Get current time based of sync position.
 					float ThisCurrentTime = PrevMarker.Time + SyncPosition.PositionBetweenMarkers * (NextMarkerTime - PrevMarker.Time);
 
-					// Adjust for case where we need to loop to be able to get to the next marker.
-					if (ThisCurrentTime > GetPlayLength())
-					{
-						ThisCurrentTime -= GetPlayLength();
-					}
-
 					// Find marker indices closest to input time position.
-					const float ThisDiff = FMath::Abs(ThisCurrentTime - CurrentInputTime);
+					float ThisDiff = FMath::Abs(ThisCurrentTime - CurrentInputTime);
 					if (ThisDiff < DiffToCurrentTime)
 					{
 						DiffToCurrentTime = ThisDiff;
@@ -4335,6 +4331,25 @@ void UAnimSequence::GetMarkerIndicesForPosition(
 						OutNextMarker.MarkerIndex = NextMarkerIdx;
 						OutCurrentTime = GetCurrentTimeFromMarkers(OutPrevMarker, OutNextMarker, SyncPosition.PositionBetweenMarkers);
 					}
+					else if (bLooped)
+					{
+						// If we looped, we extended our next marker past the end of the sequence
+						// This means that there are two points we need to test:
+						//   - The one that lands near the end of the sequence (possibly overshooting/looping around)
+						//   - The one that lands near the start of the sequence (possibly undershooting/looping around)
+						// 
+						// We tested the first one above, now test the second
+						ThisCurrentTime -= GetPlayLength();
+
+						ThisDiff = FMath::Abs(ThisCurrentTime - CurrentInputTime);
+						if (ThisDiff < DiffToCurrentTime)
+						{
+							DiffToCurrentTime = ThisDiff;
+							OutPrevMarker.MarkerIndex = PrevMarkerIdx;
+							OutNextMarker.MarkerIndex = NextMarkerIdx;
+							OutCurrentTime = GetCurrentTimeFromMarkers(OutPrevMarker, OutNextMarker, SyncPosition.PositionBetweenMarkers);
+						}
+					}
 
 					// This marker test is done, move onto next one.
 					break;
