Match sequence marker phase when joining a sync group
Add support for blend space sample play rates
Add support for blend space axis scaling
Add support for specifying an explicit delta time for a child trait stack

diff --git a/Engine/Plugins/Experimental/AnimNextAnimGraph/Source/AnimNextAnimGraph/Private/Graph/SyncGroup_GraphInstanceComponent.cpp b/Engine/Plugins/Experimental/AnimNextAnimGraph/Source/AnimNextAnimGraph/Private/Graph/SyncGroup_GraphInstanceComponent.cpp
--- a/Engine/Plugins/Experimental/AnimNextAnimGraph/Source/AnimNextAnimGraph/Private/Graph/SyncGroup_GraphInstanceComponent.cpp
+++ b/Engine/Plugins/Experimental/AnimNextAnimGraph/Source/AnimNextAnimGraph/Private/Graph/SyncGroup_GraphInstanceComponent.cpp
@@ -17,6 +17,117 @@
 // Use Rewind Debugger to record and replay logs in editor
 #define UE_DEBUG_SYNC_GROUPS !(UE_BUILD_SHIPPING || UE_BUILD_TEST)
 
+// [Sync Group Details]
+// Sequence synchronization is a complex topic, hopefully explained here in sufficient
+// details to clarify what is going on.
+// 
+// UE synchronization is group based: players are assigned a group they belong to
+// and they synchronize within that group. In order to do so, one member of the group
+// is elected leader (@see FindLeaderIndex, @see EAnimGroupSynchronizationRole) while
+// the rest are followers.
+// 
+// [Sync Group Name]
+// Typically, when a graph is hand authored, a user specified group name will be used.
+// This would be common for something like locomotion. However, sometimes we wish
+// for things to synchronize together but not with re-entrant versions.
+// Consider a blend space. We wish for its sequences to synchronize together. However
+// if our parent node spawns a new instance of the blend space with different parameters
+// we would like for each instance to NOT synchronize together. This is called
+// self-synchronization: it is local only to the current instance. In a case like this,
+// if we select a user supplied group name, then each instance will share the same
+// group name and self-synchronization is not possible. To that end, unique group names
+// can be generated and used for this. The parent node would be responsible for generating
+// the unique group name (e.g. motion matching node) and to forward it to its children.
+// 
+// [Sync Group Mode]
+// There are three sync group modes (@see EAnimGroupSynchronizationMode):
+//    - No synchronization: disables synchronization
+//    - Using group name: uses the user supplied group name to synchronize with other members
+//    - Using unique group name: ignores any user supplied group name and generates a new
+//      unique name for members to synchronize with
+// 
+// [Sync Group Method]
+// There are two synchronization methods: time based and marker based. This is mostly
+// dictated by the leader. If the leader has markers, everyone in the group will
+// attempt to use them. Members without markers will revert to time based syncing.
+// If the leader does not have markers, all members will use time based syncing.
+// 
+// [Sync Point Matching]
+// When a new member joins a group, we need to determine where it starts playing.
+// If we opt to match the sync point, we will make a best effort attempt to find a
+// suitable start point as follows.
+// 
+// When a new leader joins an existing group, we first look at the old leader.
+// If the old leader was using marker based syncing, then we use its old marker pair.
+// If the old leader was using time based syncing (maybe it was alone in the group),
+// we attempt to find the marker pair where it currently resides. We can only do so
+// if the old leader is still a member of the current group. Once we have a valid
+// marker pair for the old leader, we have two choices: match the phase or match
+// the marker pair.
+// 
+// Matching the phase is most desirable but it has the strictest requirements. Phase
+// matching allows the synchronization behavior to be deterministic: it doesn't matter
+// when a member joins, it always ends up at the same phase even if it must loop to
+// match. It works as follows:
+//     - Find the marker pair of the leader (or old leader)
+//       This pair comprises of: prev/next marker indices
+//     - Collect all markers from the start of the sequence, up to our current pair
+//     - The joining member can then skip each collected marker, in order, starting
+//       from the start of the sequence.
+// 
+// For this to work, phase matching requires:
+//     - Both sequences must begin with the same marker (e.g. both start with LeftFoot)
+//     - Both sequences must have their markers in the same order (e.g. LRLR and LRLRLR)
+// 
+// When phase matching is not possible, we attempt to match the marker pair. We do so
+// as follows:
+//     - Find the marker pair of the leader (or old leader)
+//       This pair comprises of: prev/next marker indices
+//     - The joining member can then start searching for the closest marker pair
+//       that matches. This could find a pair that is before/after the current position.
+//       Non-matching markers will be skipped/ignored. As such, it is beneficial if
+//       the new member has an approximate start position already set. If the pair
+//       already matches then we can begin playing from nearby.
+// 
+// If the new leader joins an existing group and it does not use markers, then it
+// will attempt to use the old leader's normalized play position like a time based
+// follower would. Play then resumes at the same normalized position within the group.
+// 
+// If a new leader joins but it does NOT wish to match the sync point, it will snap
+// the group to its current position (e.g. teleport/force update). When this occurs,
+// all members of the group will be treated as if they have just joined and thus defer
+// to their own sync point matching behavior.
+// 
+// Similarly, when a new follower joins a group, it can also attempt to match the sync
+// point of the leader. To do so, it operates in the same way as a new leader would:
+// we look for the leader's marker pair and we attempt to match the phase or the marker
+// pair as described above. If we aren't using markers, then we match the normalized
+// play position of the leader.
+// 
+// [Time Advancement]
+// Group synchronization hijacks the normal graph time update in that players will not
+// advance on their own. Instead, we must wait for all group members to be collected to
+// begin synchronization.
+// 
+// Leaders will advance normally using their own delta time.
+// 
+// Followers will compute their own delta time based on how much the leader advanced.
+// When markers are used, followers will attempt to pass the same markers passed by
+// the leader to attempt to keep in phase (when possible). Unknown or mismatched markers
+// will be skipped. The delta time computed is then the difference between where the
+// follower was and where it should be within the new marker pair.
+// When time based synchronization is used, we similarly compute our delta time based
+// on where the follower was and where it should be based on the normalized play time
+// of the leader.
+// 
+// Groups with a single member or un-grouped members advance normally without synchronization.
+// 
+// At times, a player will advance in time without playing events (e.g. notifies). This
+// typically occurs when we need to snap to a position (e.g. when we join). When this occurs,
+// we do not wish to fire events between our current and desired position. We always
+// perform one advance that does trigger events for leaders. Followers never trigger events.
+//
+
 namespace UE::AnimNext
 {
 	namespace Private
@@ -32,6 +143,25 @@ namespace UE::AnimNext
 		static TAutoConsoleVariable<int32> CDebugVarSyncMode(TEXT("a.AnimNext.SyncGroupMode"), 0, TEXT("Debug sync group follower mode. 0: Auto Sync, 1: Time Sync, 2: No Sync"));
 #endif
 
+		// Encapsulates the group phase position
+		struct FSyncGroupPhasePosition
+		{
+			// The name of the previous marker in our phase
+			FName PrevMarkerName;
+
+			// The name of the next marker in our phase
+			FName NextMarkerName;
+
+			// The index of the previous marker in our phase
+			int32 PrevMarkerIndex = MarkerIndexSpecialValues::Uninitialized;
+
+			// The index of the next marker in our phase
+			int32 NextMarkerIndex = MarkerIndexSpecialValues::Uninitialized;
+
+			// The normalized relative position between our two markers
+			float PositionBetweenMarkers = 0.0f;
+		};
+
 		// Sync group member state as collected during graph traversal
 		struct FSyncGroupMember
 		{
@@ -61,6 +191,12 @@ namespace UE::AnimNext
 			// If this sync group was active during the previous update, this is the timeline progress of its leader
 			FTimelineState PreviousLeaderTimelineState;
 
+			// If this sync group was active during the previous update, this is the list of sync markers from the previous leader
+			TArray<FTimelineSyncMarker, TInlineAllocator<8>> PreviousLeaderSyncMarkers;
+
+			// If this sync group was active during the previous update, this is the phase position where the leader ended
+			FSyncGroupPhasePosition PreviousLeaderPhasePosition;
+
 #if UE_DEBUG_SYNC_GROUPS
 			int32 PreviousLeaderIndex = INDEX_NONE;
 #endif
@@ -112,8 +248,8 @@ namespace UE::AnimNext
 			TArray<FName, TMemStackAllocator<>> MarkersPassed;
 
 			bool bCanGroupUseMarkerSyncing = false;
-			FMarkerSyncAnimPosition SyncPositionStart;
-			FMarkerSyncAnimPosition SyncPositionEnd;
+			FSyncGroupPhasePosition LeaderPhaseStart;
+			FSyncGroupPhasePosition LeaderPhaseEnd;
 
 			TSet<FName> ValidMarkers;
 			TSet<FName> CandidateMarkers;
@@ -537,21 +673,24 @@ namespace UE::AnimNext
 		}
 
 		// Returns the closest pair of sync markers to the specified timeline position that matches exactly the supplied sync position
-		static void GetExactBoundarySyncMarkersAt(const FMarkerSyncAnimPosition& SyncPosition, const FTimelineSyncMarkerArray& SyncMarkers, const FTimelineState& TimelineState, FMarkerPair& OutPrevMarker, FMarkerPair& OutNextMarker)
+		static void GetExactBoundarySyncMarkersAt(
+			const FSyncGroupPhasePosition& LeaderPhasePosition,
+			const FTimelineSyncMarkerArray& FollowerSyncMarkers, const FTimelineState& FollowerTimelineState,
+			FMarkerPair& OutPrevMarker, FMarkerPair& OutNextMarker)
 		{
-			const float CurrentTime = TimelineState.GetPosition();
-			const float Duration = TimelineState.GetDuration();
-			const bool bIsLooping = TimelineState.IsLooping();
-			const int32 NumSyncMarkers = SyncMarkers.Num();
+			const float CurrentTime = FollowerTimelineState.GetPosition();
+			const float Duration = FollowerTimelineState.GetDuration();
+			const bool bIsLooping = FollowerTimelineState.IsLooping();
+			const int32 NumSyncMarkers = FollowerSyncMarkers.Num();
 
 			float BestTimeDelta = FLT_MAX;
 
 			// Handle case for looping and sync position not being on either boundary.
 			for (int32 PrevMarkerIdx = 0; PrevMarkerIdx < NumSyncMarkers; ++PrevMarkerIdx)
 			{
-				const FTimelineSyncMarker& PrevMarker = SyncMarkers[PrevMarkerIdx];
+				const FTimelineSyncMarker& PrevMarker = FollowerSyncMarkers[PrevMarkerIdx];
 
-				if (PrevMarker.GetName() != SyncPosition.PreviousMarkerName)
+				if (PrevMarker.GetName() != LeaderPhasePosition.PrevMarkerName)
 				{
 					continue;	// Not matching, look for the next pair candidate
 				}
@@ -563,10 +702,10 @@ namespace UE::AnimNext
 				for (int32 NextMarkerCount = EndMarkerSearchStart; NextMarkerCount < EndCount; ++NextMarkerCount)
 				{
 					const int32 NextMarkerIdx = NextMarkerCount % NumSyncMarkers;
-					const FTimelineSyncMarker& NextMarker = SyncMarkers[NextMarkerIdx];
+					const FTimelineSyncMarker& NextMarker = FollowerSyncMarkers[NextMarkerIdx];
 
 					// We have matched the position's next marker name.
-					if (NextMarker.GetName() != SyncPosition.NextMarkerName)
+					if (NextMarker.GetName() != LeaderPhasePosition.NextMarkerName)
 					{
 						continue;	// Not matching, look for the next pair candidate
 					}
@@ -582,7 +721,7 @@ namespace UE::AnimNext
 					}
 
 					// Get current time based of sync position.
-					float DesiredTime = FMath::LerpStable(PrevMarkerTime, NextMarkerTime, SyncPosition.PositionBetweenMarkers);
+					float DesiredTime = FMath::LerpStable(PrevMarkerTime, NextMarkerTime, LeaderPhasePosition.PositionBetweenMarkers);
 
 					// Find marker indices closest to input time position.
 					float TimeDelta = FMath::Abs(DesiredTime - CurrentTime);
@@ -624,7 +763,7 @@ namespace UE::AnimNext
 					int32 PrevMarkerToUse = PrevMarkerIdx + 1;
 					while (BestTimeDelta > 0.f && PrevMarkerToUse < NumSyncMarkers)
 					{
-						BestTimeDelta = CurrentTime - SyncMarkers[PrevMarkerToUse].GetPosition();
+						BestTimeDelta = CurrentTime - FollowerSyncMarkers[PrevMarkerToUse].GetPosition();
 						++PrevMarkerToUse;
 					}
 
@@ -640,22 +779,25 @@ namespace UE::AnimNext
 		}
 
 		// Returns the boundary sync markers matching the sync position pair nearest the specified timeline position and the adjusted current time based on the sync position
-		static void GetNearestMatchingBoundarySyncMarkers(const FMarkerSyncAnimPosition& SyncPosition, const FTimelineSyncMarkerArray& SyncMarkers, const FTimelineState& TimelineState, FMarkerPair& OutPrevMarker, FMarkerPair& OutNextMarker)
+		static void GetNearestMatchingBoundarySyncMarkers(
+			const FSyncGroupPhasePosition& LeaderPhasePosition,
+			const FTimelineSyncMarkerArray& FollowerSyncMarkers, const FTimelineState& FollowerTimelineState,
+			FMarkerPair& OutPrevMarker, FMarkerPair& OutNextMarker)
 		{
-			check(SyncPosition.PreviousMarkerName != NAME_None || SyncPosition.NextMarkerName != NAME_None);
+			check(LeaderPhasePosition.PrevMarkerName != NAME_None || LeaderPhasePosition.NextMarkerName != NAME_None);
 
 			// If the sync position's previous marker is the start boundary
 			//		- We keep the previous marker at the boundary
 			//		- We look for the next marker, the first instance of the next marker's name
 			//		- We return the current time based on the sync position (see GetCurrentTimeFromMarkers)
 			//		- Done
-			if (SyncPosition.PreviousMarkerName == NAME_None)
+			if (LeaderPhasePosition.PrevMarkerName == NAME_None)
 			{
 				OutPrevMarker.MarkerIndex = MarkerIndexSpecialValues::AnimationBoundary;
 
-				for (int32 MarkerIndex = 0; MarkerIndex < SyncMarkers.Num(); ++MarkerIndex)
+				for (int32 MarkerIndex = 0; MarkerIndex < FollowerSyncMarkers.Num(); ++MarkerIndex)
 				{
-					if (SyncMarkers[MarkerIndex].GetName() == SyncPosition.NextMarkerName)
+					if (FollowerSyncMarkers[MarkerIndex].GetName() == LeaderPhasePosition.NextMarkerName)
 					{
 						OutNextMarker.MarkerIndex = MarkerIndex;
 						break;
@@ -670,13 +812,13 @@ namespace UE::AnimNext
 			//		- We look for the previous marker, the last instance of the previous marker's name
 			//		- We return the current time based on the sync position (see GetCurrentTimeFromMarkers)
 			//		- Done
-			if (SyncPosition.NextMarkerName == NAME_None)
+			if (LeaderPhasePosition.NextMarkerName == NAME_None)
 			{
 				OutNextMarker.MarkerIndex = MarkerIndexSpecialValues::AnimationBoundary;
 
-				for (int32 MarkerIndex = SyncMarkers.Num() - 1; MarkerIndex >= 0; --MarkerIndex)
+				for (int32 MarkerIndex = FollowerSyncMarkers.Num() - 1; MarkerIndex >= 0; --MarkerIndex)
 				{
-					if (SyncMarkers[MarkerIndex].GetName() == SyncPosition.PreviousMarkerName)
+					if (FollowerSyncMarkers[MarkerIndex].GetName() == LeaderPhasePosition.PrevMarkerName)
 					{
 						OutPrevMarker.MarkerIndex = MarkerIndex;
 						break;
@@ -692,16 +834,16 @@ namespace UE::AnimNext
 			//		- Once we find that pair, we return the current time based on the sync position (see GetCurrentTimeFromMarkers)
 			//		- If no such pair is found, the follower must not contain the pair of markers present on the leader (not possible)
 			//		- Done
-			GetExactBoundarySyncMarkersAt(SyncPosition, SyncMarkers, TimelineState, OutPrevMarker, OutNextMarker);
+			GetExactBoundarySyncMarkersAt(LeaderPhasePosition, FollowerSyncMarkers, FollowerTimelineState, OutPrevMarker, OutNextMarker);
 		}
 
 		// Returns the relative sync position between the two specified markers
-		static FMarkerSyncAnimPosition CalculateSyncPosition(const FTimelineSyncMarkerArray& SyncMarkers, const FTimelineState& TimelineState, const FMarkerPair& PrevMarker, const FMarkerPair& NextMarker)
+		static FSyncGroupPhasePosition CalculateSyncPosition(const FTimelineSyncMarkerArray& SyncMarkers, const FTimelineState& TimelineState, const FMarkerPair& PrevMarker, const FMarkerPair& NextMarker)
 		{
 			const float Duration = TimelineState.GetDuration();
 			float CurrentTime = TimelineState.GetPosition();
 
-			FMarkerSyncAnimPosition SyncPosition;
+			FSyncGroupPhasePosition PhasePosition;
 			float PrevTime = 0.0f;
 			float NextTime = Duration;
 
@@ -710,7 +852,8 @@ namespace UE::AnimNext
 				TEXT("MarkerCount: %d, PrevMarker : %d, NextMarker: %d, CurrentTime : %0.2f"), SyncMarkers.Num(), PrevMarker.MarkerIndex, NextMarker.MarkerIndex, CurrentTime))
 			{
 				PrevTime = SyncMarkers[PrevMarker.MarkerIndex].GetPosition();
-				SyncPosition.PreviousMarkerName = SyncMarkers[PrevMarker.MarkerIndex].GetName();
+				PhasePosition.PrevMarkerIndex = PrevMarker.MarkerIndex;
+				PhasePosition.PrevMarkerName = SyncMarkers[PrevMarker.MarkerIndex].GetName();
 			}
 
 			// Get next marker's time and name.
@@ -718,7 +861,8 @@ namespace UE::AnimNext
 				TEXT("MarkerCount: %d, PrevMarker : %d, NextMarker: %d, CurrentTime : %0.2f"), SyncMarkers.Num(), PrevMarker.MarkerIndex, NextMarker.MarkerIndex, CurrentTime))
 			{
 				NextTime = SyncMarkers[NextMarker.MarkerIndex].GetPosition();
-				SyncPosition.NextMarkerName = SyncMarkers[NextMarker.MarkerIndex].GetName();
+				PhasePosition.NextMarkerIndex = NextMarker.MarkerIndex;
+				PhasePosition.NextMarkerName = SyncMarkers[NextMarker.MarkerIndex].GetName();
 			}
 
 			// Account for looping
@@ -744,8 +888,8 @@ namespace UE::AnimNext
 			const float PositionBetweenMarkers = (CurrentTime - PrevTime) / TimeBetweenMarkers;
 			ensure(PositionBetweenMarkers >= (0.0f - SMALL_NUMBER) && PositionBetweenMarkers <= (1.0f + SMALL_NUMBER));
 
-			SyncPosition.PositionBetweenMarkers = FMath::Clamp(PositionBetweenMarkers, 0.0f, 1.0f);
-			return SyncPosition;
+			PhasePosition.PositionBetweenMarkers = FMath::Clamp(PositionBetweenMarkers, 0.0f, 1.0f);
+			return PhasePosition;
 		}
 
 		// Returns the desired time based on the provided markers and normalized position between
@@ -956,8 +1100,9 @@ namespace UE::AnimNext
 			}
 		}
 
-		static float AdvanceSyncMarkersPassed(const FTimelineState& TimelineState, bool bIsPlayingForward,
-			const FMarkerSyncAnimPosition& LeaderEndPosition,
+		static float AdvanceSyncMarkersPassed(
+			const FTimelineState& TimelineState, bool bIsPlayingForward,
+			const FSyncGroupPhasePosition& LeaderEndPosition,
 			const FTimelineSyncMarkerArray& SyncMarkers, const TArray<FName, TMemStackAllocator<>>& MarkersPassed,
 			FMarkerPair& PrevMarker, FMarkerPair& NextMarker)
 		{
@@ -1027,7 +1172,7 @@ namespace UE::AnimNext
 						SeekMarkerBackward(PrevMarker.MarkerIndex, bIsLooping, PassedMarker, SyncMarkers);
 					}
 
-					if (LeaderEndPosition.PreviousMarkerName == NAME_None)
+					if (LeaderEndPosition.PrevMarkerName == NAME_None)
 					{
 						// If our leader has reached the end boundary, make sure we reach it as well
 						PrevMarker.MarkerIndex = MarkerIndexSpecialValues::AnimationBoundary;
@@ -1037,19 +1182,114 @@ namespace UE::AnimNext
 						NextMarker = PrevMarker;
 
 						// Find the previous marker to match our leader
-						SeekMarkerBackward(PrevMarker.MarkerIndex, bIsLooping, LeaderEndPosition.PreviousMarkerName, SyncMarkers);
+						SeekMarkerBackward(PrevMarker.MarkerIndex, bIsLooping, LeaderEndPosition.PrevMarkerName, SyncMarkers);
 					}
 				}
 
 				if (PrevMarker.MarkerIndex != MarkerIndexSpecialValues::AnimationBoundary)
 				{
-					check(SyncMarkers[PrevMarker.MarkerIndex].GetName() == LeaderEndPosition.PreviousMarkerName);
+					check(SyncMarkers[PrevMarker.MarkerIndex].GetName() == LeaderEndPosition.PrevMarkerName);
 				}
 
 				return CalculateTimeFromSyncPosition(SyncMarkers, TimelineState, PrevMarker, NextMarker, LeaderEndPosition.PositionBetweenMarkers);
 			}
 		}
 
+		// Returns whether or not two members have phases that can be matched
+		static bool CanMatchPhase(const TArrayView<const FTimelineSyncMarker>& SyncMarkersA, const TArrayView<const FTimelineSyncMarker>& SyncMarkersB)
+		{
+			// TODO: We can cache this value once we've computed this between a follower and leader
+			// We only need to recompute this if the leader changes or when joining
+
+			const TArrayView<const FTimelineSyncMarker>* LongestSyncMarkersList;
+			const TArrayView<const FTimelineSyncMarker>* ShortestSyncMarkersList;
+
+			if (SyncMarkersA.Num() >= SyncMarkersB.Num())
+			{
+				LongestSyncMarkersList = &SyncMarkersA;
+				ShortestSyncMarkersList = &SyncMarkersB;
+			}
+			else
+			{
+				LongestSyncMarkersList = &SyncMarkersB;
+				ShortestSyncMarkersList = &SyncMarkersA;
+			}
+
+			const int32 NumLongestSyncMarkers = LongestSyncMarkersList->Num();
+			const int32 NumShortestSyncMarkers = ShortestSyncMarkersList->Num();
+
+			for (int32 LongestIdx = 0; LongestIdx < NumLongestSyncMarkers; ++LongestIdx)
+			{
+				int32 ShortestIdx;
+				if (SyncMarkersA[0].GetName() == SyncMarkersA.Last().GetName())
+				{
+					// Looping markers
+					ShortestIdx = LongestIdx <= NumShortestSyncMarkers ? LongestIdx : ((LongestIdx % NumShortestSyncMarkers) + 1);
+				}
+				else
+				{
+					ShortestIdx = LongestIdx % NumShortestSyncMarkers;
+				}
+
+				if ((*LongestSyncMarkersList)[LongestIdx].GetName() != (*ShortestSyncMarkersList)[ShortestIdx].GetName())
+				{
+					// Markers do not match which means it is not possible for us to phase match
+					return false;
+				}
+			}
+
+			return true;
+		}
+
+		// Returns the phase position for B that matches the current phase on A
+		static FSyncGroupPhasePosition FindMatchingPhaseTime(const TArrayView<const FTimelineSyncMarker>& SyncMarkersA, const FSyncGroupPhasePosition& SyncPositionA, const FTimelineSyncMarkerArray& SyncMarkersB)
+		{
+			// We know our phases can match, so we can simply skip the necessary amount of markers in B to match
+			// Consider an example:
+			//   - Sequence A: [R, L, R, L]
+			//   - Sequence B: [R, L, R, L, R, L]
+			// Both A and B can remap with a simply module because the phases must match
+			// (B, 4) maps to (A, 0): 4 % 4 = 0
+			// 
+			// Similarly with more markers:
+			//   - Sequence C: [R, M, L, R, M, L]
+			//   - Sequence D: [R, M, L, R, M, L, R, M, L]
+			// 
+			// We also have the case where an animation loops. Looping animations have the first/last keyframes
+			// identical with the exception of root motion being different. And so if we have a sync marker on
+			// the first keyframe, we need a matching sync marker on the last keyframe (we never interpolate
+			// between the last/first keyframes).
+			//   - Sequence E: [R, L, R]
+			//   - Sequence F: [R, L, R, L, R]
+			// Here, we can't use the modulo as is, we have to add 1 to the marker index after the module operation
+			// (F, 3) maps to (E, 1): (3 % 3) + 1 = 1
+			// Note that this only works if we remap a marker that exceeds the number in the other. Otherwise
+			// we use the index as-is.
+
+			int32 PrevMarkerIndex;
+			int32 NextMarkerIndex;
+			if (SyncMarkersA[0].GetName() == SyncMarkersA.Last().GetName())
+			{
+				// Looping markers
+				PrevMarkerIndex = SyncPositionA.PrevMarkerIndex <= SyncMarkersB.Num() ? SyncPositionA.PrevMarkerIndex : ((SyncPositionA.PrevMarkerIndex % SyncMarkersB.Num()) + 1);
+				NextMarkerIndex = SyncPositionA.NextMarkerIndex <= SyncMarkersB.Num() ? SyncPositionA.NextMarkerIndex : ((SyncPositionA.NextMarkerIndex % SyncMarkersB.Num()) + 1);
+			}
+			else
+			{
+				PrevMarkerIndex = SyncPositionA.PrevMarkerIndex % SyncMarkersB.Num();
+				NextMarkerIndex = SyncPositionA.NextMarkerIndex % SyncMarkersB.Num();
+			}
+
+			// Markers should match
+			check(SyncMarkersA[SyncPositionA.PrevMarkerIndex].GetName() == SyncMarkersB[PrevMarkerIndex].GetName());
+			check(SyncMarkersA[SyncPositionA.NextMarkerIndex].GetName() == SyncMarkersB[NextMarkerIndex].GetName());
+
+			FSyncGroupPhasePosition PhasePosition = SyncPositionA;
+			PhasePosition.PrevMarkerIndex = PrevMarkerIndex;
+			PhasePosition.NextMarkerIndex = NextMarkerIndex;
+			return PhasePosition;
+		}
+
 		static void AdvanceLeaderTimeBased(FExecutionContext& Context, FSyncGroupContext& GroupContext)
 		{
 			FSyncGroupMemberContext& GroupLeaderContext = GroupContext.Members[GroupContext.LeaderIndex];
@@ -1156,7 +1396,10 @@ namespace UE::AnimNext
 		// Ensures that the follower markers we found match the ones from the leader
 		// They can mismatch if the leader/follower have mismatched looping state
 		// Returns whether or not the markers were modified
-		static bool SanitizeFollowerMarkers(const FMarkerSyncAnimPosition& SyncPosition, const FTimelineSyncMarkerArray& SyncMarkers, bool bIsPlayingForward, bool bIsLooping, FMarkerPair& PrevMarker, FMarkerPair& NextMarker)
+		static bool SanitizeFollowerMarkers(
+			const FSyncGroupPhasePosition& LeaderPhasePosition,
+			const FTimelineSyncMarkerArray& FollowerSyncMarkers, bool bIsPlayingForward, bool bIsLooping,
+			FMarkerPair& PrevMarker, FMarkerPair& NextMarker)
 		{
 			// Must have found some markers
 			check(PrevMarker.MarkerIndex != MarkerIndexSpecialValues::Uninitialized);
@@ -1166,14 +1409,14 @@ namespace UE::AnimNext
 
 			if (bIsPlayingForward)
 			{
-				if (GetMarkerName(PrevMarker, SyncMarkers) != SyncPosition.PreviousMarkerName)
+				if (GetMarkerName(PrevMarker, FollowerSyncMarkers) != LeaderPhasePosition.PrevMarkerName)
 				{
 					// Our previous marker doesn't match, look for the next one that matches
-					SeekMarkerForward(PrevMarker.MarkerIndex, bIsLooping, SyncPosition.PreviousMarkerName, SyncMarkers);
+					SeekMarkerForward(PrevMarker.MarkerIndex, bIsLooping, LeaderPhasePosition.PrevMarkerName, FollowerSyncMarkers);
 
 					// Fixup our next marker index
 					NextMarker.MarkerIndex = PrevMarker.MarkerIndex + 1;
-					if (NextMarker.MarkerIndex == SyncMarkers.Num())
+					if (NextMarker.MarkerIndex == FollowerSyncMarkers.Num())
 					{
 						NextMarker.MarkerIndex = bIsLooping ? 0 : MarkerIndexSpecialValues::AnimationBoundary;
 					}
@@ -1181,26 +1424,26 @@ namespace UE::AnimNext
 					bAppliedFixup = true;
 				}
 
-				if (GetMarkerName(NextMarker, SyncMarkers) != SyncPosition.NextMarkerName)
+				if (GetMarkerName(NextMarker, FollowerSyncMarkers) != LeaderPhasePosition.NextMarkerName)
 				{
 					// Our next marker doesn't match, look for the next one that matches
-					SeekMarkerForward(NextMarker.MarkerIndex, bIsLooping, SyncPosition.NextMarkerName, SyncMarkers);
+					SeekMarkerForward(NextMarker.MarkerIndex, bIsLooping, LeaderPhasePosition.NextMarkerName, FollowerSyncMarkers);
 
 					bAppliedFixup = true;
 				}
 			}
 			else
 			{
-				if (GetMarkerName(NextMarker, SyncMarkers) != SyncPosition.NextMarkerName)
+				if (GetMarkerName(NextMarker, FollowerSyncMarkers) != LeaderPhasePosition.NextMarkerName)
 				{
 					// Our next marker doesn't match, look for the next one that matches
-					SeekMarkerBackward(NextMarker.MarkerIndex, bIsLooping, SyncPosition.NextMarkerName, SyncMarkers);
+					SeekMarkerBackward(NextMarker.MarkerIndex, bIsLooping, LeaderPhasePosition.NextMarkerName, FollowerSyncMarkers);
 
 					// Fixup our previous marker index
 					if (NextMarker.MarkerIndex == MarkerIndexSpecialValues::AnimationBoundary ||
 						(bIsLooping && NextMarker.MarkerIndex == 0))
 					{
-						PrevMarker.MarkerIndex = SyncMarkers.Num() - 1;
+						PrevMarker.MarkerIndex = FollowerSyncMarkers.Num() - 1;
 					}
 					else
 					{
@@ -1210,10 +1453,10 @@ namespace UE::AnimNext
 					bAppliedFixup = true;
 				}
 
-				if (GetMarkerName(PrevMarker, SyncMarkers) != SyncPosition.PreviousMarkerName)
+				if (GetMarkerName(PrevMarker, FollowerSyncMarkers) != LeaderPhasePosition.PrevMarkerName)
 				{
 					// Our previous marker doesn't match, look for the next one that matches
-					SeekMarkerBackward(PrevMarker.MarkerIndex, bIsLooping, SyncPosition.PreviousMarkerName, SyncMarkers);
+					SeekMarkerBackward(PrevMarker.MarkerIndex, bIsLooping, LeaderPhasePosition.PrevMarkerName, FollowerSyncMarkers);
 
 					bAppliedFixup = true;
 				}
@@ -1221,8 +1464,8 @@ namespace UE::AnimNext
 
 			// Follower markers must only match in name, they might not be siblings if the pattern on the leader doesn't match the one on the follower
 			// e.g. leader with RLRL and follower with RLLRL
-			check(GetMarkerName(PrevMarker, SyncMarkers) == SyncPosition.PreviousMarkerName);
-			check(GetMarkerName(NextMarker, SyncMarkers) == SyncPosition.NextMarkerName);
+			check(GetMarkerName(PrevMarker, FollowerSyncMarkers) == LeaderPhasePosition.PrevMarkerName);
+			check(GetMarkerName(NextMarker, FollowerSyncMarkers) == LeaderPhasePosition.NextMarkerName);
 
 			return bAppliedFixup;
 		}
@@ -1273,21 +1516,68 @@ namespace UE::AnimNext
 
 			const float PlayRate = LeaderStartTimelineState.GetPlayRate() != 0.0f ? LeaderStartTimelineState.GetPlayRate() : 1.0f;
 
-			const bool bIsSmoothJoining = GroupLeaderContext.State->bJustJoined &&
-				GroupLeaderContext.State->GroupParameters.bMatchSyncPoint &&
-				!GroupContext.State->bJustFormed;
-			if (bIsSmoothJoining)
+#if UE_DEBUG_SYNC_GROUPS
+			FString JoiningMode;
+#endif
+
+			if (GroupLeaderContext.State->bJustJoined && !GroupContext.State->bJustFormed)
 			{
-				// This is a soft join where we wish to match the current sync group position
-				const float CurrentTime = LeaderStartTimelineState.GetPosition();
-				const float DesiredTime = GroupContext.PreviousLeaderTimelineState.GetPositionRatio() * LeaderStartTimelineState.GetDuration();
+				if (GroupLeaderContext.State->GroupParameters.bMatchSyncPoint)
+				{
+					const float CurrentTime = LeaderStartTimelineState.GetPosition();
+					float DesiredTime;
 
-				// Seek where we should be without dispatching events
-				const float DeltaTime = (DesiredTime - CurrentTime) / PlayRate;
-				GroupSyncTrait.AdvanceBy(Context, DeltaTime, false);
+					// Can we match the phase of the previous leader?
+					const bool bCanMatchPhase = CanMatchPhase(GroupContext.State->PreviousLeaderSyncMarkers, GroupLeaderContext.SyncMarkers);
+					if (bCanMatchPhase)
+					{
+						// When our phases can match, we simulate playback from the start of the timeline
+						// This ensures a deterministic outcome no matter when a leader joins the group
+						// We collect the sync markers of our previous leader between [StartPosition, CurrentPosition]
+						// and we pass the same markers on our follower
 
-				// Update our cached state
-				LeaderStartTimelineState = GroupLeaderContext.TimelineTrait.GetState(Context);
+						const FSyncGroupPhasePosition LeaderPhasePosition = FindMatchingPhaseTime(GroupContext.State->PreviousLeaderSyncMarkers, GroupContext.State->PreviousLeaderPhasePosition, GroupLeaderContext.SyncMarkers);
+
+						const FMarkerPair PrevMarker(LeaderPhasePosition.PrevMarkerIndex, 0.0f);
+						const FMarkerPair NextMarker(LeaderPhasePosition.NextMarkerIndex, 0.0f);
+
+						DesiredTime = CalculateTimeFromSyncPosition(GroupLeaderContext.SyncMarkers, LeaderStartTimelineState, PrevMarker, NextMarker, LeaderPhasePosition.PositionBetweenMarkers);
+
+#if UE_DEBUG_SYNC_GROUPS
+						JoiningMode = TEXT(" (phase matched join)");
+#endif
+					}
+					else
+					{
+						// If we can't match our phases then we attempt to find the closest matching marker
+						// pair from our normalized position as dictated by the previous leader (e.g. if the
+						// previous leader is at 80%, we look for the closest pair around the 80% mark on our follower)
+						// Our current time is thus not relevant and we ignore it
+
+						// We assume that we start at the same position in normalized time as our previous leader
+						DesiredTime = GroupContext.PreviousLeaderTimelineState.GetPositionRatio() * LeaderStartTimelineState.GetDuration();
+
+#if UE_DEBUG_SYNC_GROUPS
+						JoiningMode = TEXT(" (relative matched join)");
+#endif
+					}
+
+					// Seek where we should be without dispatching events
+					const float DeltaTime = (DesiredTime - CurrentTime) / PlayRate;
+					GroupSyncTrait.AdvanceBy(Context, DeltaTime, false);
+
+					// Update our cached state
+					LeaderStartTimelineState = GroupLeaderContext.TimelineTrait.GetState(Context);
+				}
+				else
+				{
+					// If we aren't requesting to match the sync point, then we'll use whatever marker pair
+					// we currently lie between and the group will snap to us
+
+#if UE_DEBUG_SYNC_GROUPS
+					JoiningMode = TEXT(" (unmatched join)");
+#endif
+				}
 			}
 
 			// Find the sync markers around our current position
@@ -1298,7 +1588,7 @@ namespace UE::AnimNext
 			ValidateLeaderMarkers(LeaderStartTimelineState, GroupLeaderContext.SyncMarkers, PrevMarkerStart, NextMarkerStart);
 
 			// Cache the start sync position
-			GroupContext.SyncPositionStart = CalculateSyncPosition(GroupLeaderContext.SyncMarkers, LeaderStartTimelineState, PrevMarkerStart, NextMarkerStart);
+			GroupContext.LeaderPhaseStart = CalculateSyncPosition(GroupLeaderContext.SyncMarkers, LeaderStartTimelineState, PrevMarkerStart, NextMarkerStart);
 
 			// Record where the leader started from, followers that join the group will start there
 			GroupContext.LeaderStartRatio = LeaderStartTimelineState.GetPositionRatio();
@@ -1328,7 +1618,7 @@ namespace UE::AnimNext
 			ValidateLeaderMarkers(LeaderEndTimelineState, GroupLeaderContext.SyncMarkers, PrevMarkerEnd, NextMarkerEnd);
 
 			// Cache the end sync position
-			GroupContext.SyncPositionEnd = CalculateSyncPosition(GroupLeaderContext.SyncMarkers, LeaderEndTimelineState, PrevMarkerEnd, NextMarkerEnd);
+			GroupContext.LeaderPhaseEnd = CalculateSyncPosition(GroupLeaderContext.SyncMarkers, LeaderEndTimelineState, PrevMarkerEnd, NextMarkerEnd);
 			GroupContext.LeaderEndRatio = LeaderEndTimelineState.GetPositionRatio();
 
 #if DO_CHECK
@@ -1337,7 +1627,7 @@ namespace UE::AnimNext
 				if (PrevMarkerStart.MarkerIndex == PrevMarkerEnd.MarkerIndex)
 				{
 					check(NextMarkerStart.MarkerIndex == NextMarkerEnd.MarkerIndex);
-					check(GroupContext.SyncPositionStart.PositionBetweenMarkers <= GroupContext.SyncPositionEnd.PositionBetweenMarkers);
+					check(GroupContext.LeaderPhaseStart.PositionBetweenMarkers <= GroupContext.LeaderPhaseEnd.PositionBetweenMarkers);
 				}
 			}
 			else
@@ -1345,7 +1635,7 @@ namespace UE::AnimNext
 				if (PrevMarkerStart.MarkerIndex == PrevMarkerEnd.MarkerIndex)
 				{
 					check(NextMarkerStart.MarkerIndex == NextMarkerEnd.MarkerIndex);
-					check(GroupContext.SyncPositionStart.PositionBetweenMarkers >= GroupContext.SyncPositionEnd.PositionBetweenMarkers);
+					check(GroupContext.LeaderPhaseStart.PositionBetweenMarkers >= GroupContext.LeaderPhaseEnd.PositionBetweenMarkers);
 				}
 			}
 #endif
@@ -1386,19 +1676,19 @@ namespace UE::AnimNext
 					LeaderStartTimelineState.GetPosition(), *DeltaTimeStr,
 					LeaderEndTimelineState.GetPosition(), LeaderEndTimelineState.GetDuration(),
 					LeaderEndTimelineState.GetPositionRatio() * 100.0f,
-					*GroupContext.SyncPositionStart.PreviousMarkerName.ToString(),
+					*GroupContext.LeaderPhaseStart.PrevMarkerName.ToString(),
 					GetMarkerPosition(PrevMarkerStart, GroupLeaderContext.SyncMarkers, 0.0f),
-					*GroupContext.SyncPositionStart.NextMarkerName.ToString(),
+					*GroupContext.LeaderPhaseStart.NextMarkerName.ToString(),
 					GetMarkerPosition(NextMarkerStart, GroupLeaderContext.SyncMarkers, LeaderEndTimelineState.GetDuration()),
-					GroupContext.SyncPositionStart.PositionBetweenMarkers * 100.0f,
-					*GroupContext.SyncPositionEnd.PreviousMarkerName.ToString(),
+					GroupContext.LeaderPhaseStart.PositionBetweenMarkers * 100.0f,
+					*GroupContext.LeaderPhaseEnd.PrevMarkerName.ToString(),
 					GetMarkerPosition(PrevMarkerEnd, GroupLeaderContext.SyncMarkers, 0.0f),
-					*GroupContext.SyncPositionEnd.NextMarkerName.ToString(),
+					*GroupContext.LeaderPhaseEnd.NextMarkerName.ToString(),
 					GetMarkerPosition(NextMarkerEnd, GroupLeaderContext.SyncMarkers, LeaderEndTimelineState.GetDuration()),
-					GroupContext.SyncPositionEnd.PositionBetweenMarkers * 100.0f,
+					GroupContext.LeaderPhaseEnd.PositionBetweenMarkers * 100.0f,
 					*LeaderEndTimelineState.GetDebugName().ToString(),
 					bLooped ? TEXT(" (looped)") : TEXT(""), *PassedMarkers,
-					bIsSmoothJoining ? TEXT(" (smooth join)") : TEXT(""),
+					*JoiningMode,
 					*JoinedStatus);
 			}
 #endif
@@ -1538,17 +1828,64 @@ namespace UE::AnimNext
 			const float PlayRate = FollowerStartTimelineState.GetPlayRate() != 0.0f ? FollowerStartTimelineState.GetPlayRate() : 1.0f;
 			const float StartTime = FollowerStartTimelineState.GetPosition();
 
-			const bool bIsSmoothJoining = GroupMemberContext.State->bJustJoined && GroupMemberContext.State->GroupParameters.bMatchSyncPoint;
-			if (bIsSmoothJoining)
+#if UE_DEBUG_SYNC_GROUPS
+			FString JoiningMode;
+#endif
+
+			const bool bCanMatchPhase = CanMatchPhase(GroupLeaderContext.SyncMarkers, GroupMemberContext.SyncMarkers);
+
+			if (GroupMemberContext.State->bJustJoined)
 			{
-				// If we just joined the group as a follower then our current time is not relevant
-				// We assume that we start at the same position in normalized time as our leader
-				// This is where we'll start searching for valid markers, we'll perform a final fix-up afterwards
-				const float CurrentTime = FollowerStartTimelineState.GetPosition();
-				const float DesiredTime = GroupContext.LeaderStartRatio * FollowerStartTimelineState.GetDuration();
+				if (GroupMemberContext.State->GroupParameters.bMatchSyncPoint)
+				{
+					float DesiredTime;
 
-				// Seek where we should be without dispatching events
-				FollowerStartTimelineState = FollowerStartTimelineState.WithPosition(DesiredTime);
+					if (bCanMatchPhase)
+					{
+						// When our phases can match, we simulate playback from the start of the timeline
+						// This ensures a deterministic outcome no matter when a member joins the group
+						// We collect the sync markers of our leader between [StartPosition, CurrentPosition]
+						// and we pass the same markers on our follower
+
+						const FSyncGroupPhasePosition FollowerPhasePosition = FindMatchingPhaseTime(GroupLeaderContext.SyncMarkers, GroupContext.LeaderPhaseStart, GroupMemberContext.SyncMarkers);
+
+						const FMarkerPair PrevMarker(FollowerPhasePosition.PrevMarkerIndex, 0.0f);
+						const FMarkerPair NextMarker(FollowerPhasePosition.NextMarkerIndex, 0.0f);
+
+						DesiredTime = CalculateTimeFromSyncPosition(GroupMemberContext.SyncMarkers, FollowerStartTimelineState, PrevMarker, NextMarker, FollowerPhasePosition.PositionBetweenMarkers);
+
+#if UE_DEBUG_SYNC_GROUPS
+						JoiningMode = TEXT(" (phase matched join)");
+#endif
+					}
+					else
+					{
+						// If we can't match our phases then we attempt to find the closest matching marker
+						// pair from our normalized position as dictated by the leader (e.g. if leader is
+						// at 80%, we look for the closest pair around the 80% mark on our follower)
+						// Our current time is thus not relevant and we ignore it
+
+						// We assume that we start at the same position in normalized time as our leader
+						const float CurrentTime = FollowerStartTimelineState.GetPosition();
+						DesiredTime = GroupContext.LeaderStartRatio * FollowerStartTimelineState.GetDuration();
+
+#if UE_DEBUG_SYNC_GROUPS
+						JoiningMode = TEXT(" (relative matched join)");
+#endif
+					}
+
+					// Seek where we should be without dispatching events
+					FollowerStartTimelineState = FollowerStartTimelineState.WithPosition(DesiredTime);
+				}
+				else
+				{
+					// If we aren't requesting to match the sync point, then we attempt to find the closest
+					// matching marker pair from our current position
+
+#if UE_DEBUG_SYNC_GROUPS
+					JoiningMode = TEXT(" (unmatched join)");
+#endif
+				}
 			}
 
 #if UE_DEBUG_SYNC_GROUPS
@@ -1558,11 +1895,17 @@ namespace UE::AnimNext
 			// Find the nearest markers matching our current time
 			FMarkerPair PrevMarkerStart;
 			FMarkerPair NextMarkerStart;
-			GetNearestMatchingBoundarySyncMarkers(GroupContext.SyncPositionStart, GroupMemberContext.SyncMarkers, FollowerStartTimelineState, PrevMarkerStart, NextMarkerStart);
+			GetNearestMatchingBoundarySyncMarkers(
+				GroupContext.LeaderPhaseStart,
+				GroupMemberContext.SyncMarkers, FollowerStartTimelineState,
+				PrevMarkerStart, NextMarkerStart);
 
 			FMarkerPair PrevMarkerSanitized = PrevMarkerStart;
 			FMarkerPair NextMarkerSanitized = NextMarkerStart;
-			bool bHadInvalidMarkers = SanitizeFollowerMarkers(GroupContext.SyncPositionStart, GroupMemberContext.SyncMarkers, bIsPlayingForward, FollowerStartTimelineState.IsLooping(), PrevMarkerSanitized, NextMarkerSanitized);
+			bool bHadInvalidMarkers = SanitizeFollowerMarkers(
+				GroupContext.LeaderPhaseStart,
+				GroupMemberContext.SyncMarkers, bIsPlayingForward, FollowerStartTimelineState.IsLooping(),
+				PrevMarkerSanitized, NextMarkerSanitized);
 
 			if (GroupMemberContext.State->bJustJoined)
 			{
@@ -1572,8 +1915,9 @@ namespace UE::AnimNext
 				// To avoid this, now that we've found good markers, we seek again to the position the leader
 				// started at between them
 				const float CurrentTime = StartTime;
-				const float DesiredTime = CalculateTimeFromSyncPosition(GroupMemberContext.SyncMarkers, FollowerStartTimelineState,
-					PrevMarkerSanitized, NextMarkerSanitized, GroupContext.SyncPositionStart.PositionBetweenMarkers);;
+				const float DesiredTime = CalculateTimeFromSyncPosition(
+					GroupMemberContext.SyncMarkers, FollowerStartTimelineState,
+					PrevMarkerSanitized, NextMarkerSanitized, GroupContext.LeaderPhaseStart.PositionBetweenMarkers);
 
 				// Seek where we should be without dispatching events
 				const float DeltaTime = (DesiredTime - CurrentTime) / PlayRate;
@@ -1588,7 +1932,7 @@ namespace UE::AnimNext
 			FMarkerPair PrevMarkerEnd = PrevMarkerSanitized;
 			FMarkerPair NextMarkerEnd = NextMarkerSanitized;
 			float DesiredTime = AdvanceSyncMarkersPassed(
-				FollowerStartTimelineState, bIsPlayingForward, GroupContext.SyncPositionEnd,
+				FollowerStartTimelineState, bIsPlayingForward, GroupContext.LeaderPhaseEnd,
 				GroupMemberContext.SyncMarkers, GroupContext.MarkersPassed, PrevMarkerEnd, NextMarkerEnd);
 
 			if (bIsPlayingForward)
@@ -1660,17 +2004,17 @@ namespace UE::AnimNext
 					GetMarkerPosition(PrevMarkerStart, GroupMemberContext.SyncMarkers, 0.0f),
 					*GetMarkerName(NextMarkerStart, GroupMemberContext.SyncMarkers).ToString(),
 					GetMarkerPosition(NextMarkerStart, GroupMemberContext.SyncMarkers, FollowerEndTimelineState.GetDuration()),
-					GroupContext.SyncPositionStart.PositionBetweenMarkers * 100.0f,
+					GroupContext.LeaderPhaseStart.PositionBetweenMarkers * 100.0f,
 					*GetMarkerName(PrevMarkerEnd, GroupMemberContext.SyncMarkers).ToString(),
 					GetMarkerPosition(PrevMarkerEnd, GroupMemberContext.SyncMarkers, 0.0f),
 					*GetMarkerName(NextMarkerEnd, GroupMemberContext.SyncMarkers).ToString(),
 					GetMarkerPosition(NextMarkerEnd, GroupMemberContext.SyncMarkers, FollowerEndTimelineState.GetDuration()),
-					GroupContext.SyncPositionEnd.PositionBetweenMarkers * 100.0f,
+					GroupContext.LeaderPhaseEnd.PositionBetweenMarkers * 100.0f,
 					*FollowerEndTimelineState.GetDebugName().ToString(),
 					bLooped ? TEXT(" (looped)") : TEXT(""),
 					bUnexpectedDeltaTime ? TEXT(" (unusual delta time)") : TEXT(""),
 					bHadInvalidMarkers ? *SanitizedResult : TEXT(""),
-					bIsSmoothJoining ? TEXT(" (smooth join)") : TEXT(""),
+					*JoiningMode,
 					*JoinedStatus);
 			}
 #endif
@@ -2058,6 +2402,8 @@ namespace UE::AnimNext
 
 			// Retain the leader progress as we might need it if a new leader joins the group during the next update
 			GroupState.PreviousLeaderTimelineState = GroupContext.LeaderTimelineState;
+			GroupState.PreviousLeaderSyncMarkers = GroupContext.Members[GroupContext.LeaderIndex].SyncMarkers;
+			GroupState.PreviousLeaderPhasePosition = GroupContext.LeaderPhaseEnd;
 
 #if UE_DEBUG_SYNC_GROUPS
 			GroupState.PreviousLeaderIndex = GroupContext.LeaderIndex;
diff --git a/Engine/Plugins/Experimental/AnimNextAnimGraph/Source/AnimNextAnimGraph/Private/Traits/BlendSpacePlayer.cpp b/Engine/Plugins/Experimental/AnimNextAnimGraph/Source/AnimNextAnimGraph/Private/Traits/BlendSpacePlayer.cpp
--- a/Engine/Plugins/Experimental/AnimNextAnimGraph/Source/AnimNextAnimGraph/Private/Traits/BlendSpacePlayer.cpp
+++ b/Engine/Plugins/Experimental/AnimNextAnimGraph/Source/AnimNextAnimGraph/Private/Traits/BlendSpacePlayer.cpp
@@ -115,7 +115,7 @@ namespace UE::AnimNext
 			const TArray<FBlendSample>& BlendSamples = InstanceData->BlendSpace->GetBlendSamples();
 			if (InstanceData->SampleGraphs.Num() != BlendSamples.Num())
 			{
-				InstanceData->SampleGraphs.SetNum(BlendSamples.Num(), EAllowShrinking::No); 
+				InstanceData->SampleGraphs.SetNum(BlendSamples.Num(), EAllowShrinking::No);
 				for (int32 CurSampleIndex = 0; CurSampleIndex < BlendSamples.Num(); CurSampleIndex++)
 				{
 					const UAnimSequence* CurSequence = BlendSamples[CurSampleIndex].Animation;
@@ -127,9 +127,6 @@ namespace UE::AnimNext
 					FAnimNextNativeDataInterface_AnimSequencePlayer& NewNativeInterface = NewPayload.InitializeAs<FAnimNextNativeDataInterface_AnimSequencePlayer>();
 					NewNativeInterface.AnimSequence = CurSequence;
 
-					// @TODO: Playrate should be normalized too.
-					NewNativeInterface.PlayRate = SharedData->GetPlayRate(Binding);
-
 					// Assume start position is normalized.
 					NewNativeInterface.StartPosition = SharedData->GetStartPosition(Binding) * CurSequence->GetPlayLength();
 					NewNativeInterface.Loop = SharedData->GetbLoop(Binding);
@@ -157,7 +154,7 @@ namespace UE::AnimNext
 	uint32 FBlendSpacePlayerTrait::GetNumChildren(const FExecutionContext& Context, const TTraitBinding<IHierarchy>& Binding) const
 	{
 		const FInstanceData* InstanceData = Binding.GetInstanceData<FInstanceData>();
-		return InstanceData->NumChildren; 
+		return InstanceData->NumChildren;
 
 	}
 
@@ -174,30 +171,35 @@ namespace UE::AnimNext
 				}
 			}
 		}
-		check(Children.Num() == InstanceData->NumChildren); 
+		check(Children.Num() == InstanceData->NumChildren);
 	}
 
 	void FBlendSpacePlayerTrait::QueueChildrenForTraversal(FUpdateTraversalContext& Context, const TTraitBinding<IUpdateTraversal>& Binding, const FTraitUpdateState& TraitState, FUpdateTraversalQueue& TraversalQueue) const
 	{
 		const FInstanceData* InstanceData = Binding.GetInstanceData<FInstanceData>();
 
-		TTraitBinding<IContinuousBlend> ContinuousBlendTrait;
-		Binding.GetStackInterface(ContinuousBlendTrait);
+		if (InstanceData->NumChildren == 0)
+		{
+			return;
+		}
 
-		if (InstanceData->NumChildren > 0)
+		for (const FSampleGraphState& SampleState : InstanceData->SampleGraphs)
 		{
-			for (const FSampleGraphState& SampleState : InstanceData->SampleGraphs)
+			if (SampleState.bSampledThisFrame)
 			{
-				if (SampleState.bSampledThisFrame)
-				{
-					const bool bGraphHasNeverUpdated = SampleState.Instance.IsValid() && !SampleState.Instance->HasUpdated();
-					bool bNewlyRelevant = (SampleState.bSampledThisFrame && !SampleState.bSampledLastFrame) || bGraphHasNeverUpdated;
-					TraversalQueue.Push(SampleState.ChildPtr, TraitState.WithWeight(SampleState.Weight).AsNewlyRelevant(bNewlyRelevant));
+				const bool bGraphHasNeverUpdated = SampleState.Instance.IsValid() && !SampleState.Instance->HasUpdated();
+				const bool bNewlyRelevant = (SampleState.bSampledThisFrame && !SampleState.bSampledLastFrame) || bGraphHasNeverUpdated;
 
-					if (SampleState.Instance.IsValid())
-					{
-						SampleState.Instance->MarkAsUpdated();
-					}
+				const FTraitUpdateState SampleTraitState = TraitState
+					.WithDeltaTime(SampleState.DeltaTime)
+					.WithWeight(SampleState.Weight)
+					.AsNewlyRelevant(bNewlyRelevant);
+
+				TraversalQueue.Push(SampleState.ChildPtr, SampleTraitState);
+
+				if (SampleState.Instance.IsValid())
+				{
+					SampleState.Instance->MarkAsUpdated();
 				}
 			}
 		}
@@ -208,7 +210,7 @@ namespace UE::AnimNext
 		const FInstanceData* InstanceData = Binding.GetInstanceData<FInstanceData>();
 		if (ChildIndex < InstanceData->NumChildren)
 		{
-			int CurChildIndex = 0; 
+			int CurChildIndex = 0;
 			for (const FSampleGraphState& SampleState : InstanceData->SampleGraphs)
 			{
 				if (SampleState.bSampledThisFrame)
@@ -234,14 +236,14 @@ namespace UE::AnimNext
 		FInstanceData* InstanceData = Binding.GetInstanceData<FInstanceData>();
 		check(InstanceData);
 
-		const float DeltaTime = TraitState.GetDeltaTime(); 
+		const float DeltaTime = TraitState.GetDeltaTime();
 
 		// input data 
 		FVector BlendParameters(SharedData->GetXAxisSamplePoint(Binding), SharedData->GetYAxisSamplePoint(Binding), 0.0f);
 		for (FSampleGraphState& SampleState : InstanceData->SampleGraphs)
 		{
 			SampleState.bSampledLastFrame = SampleState.bSampledThisFrame;
-			SampleState.bSampledThisFrame = false; 
+			SampleState.bSampledThisFrame = false;
 		}
 
 		InstanceData->NumChildren = 0;
@@ -253,14 +255,21 @@ namespace UE::AnimNext
 				const FVector FilteredBlendParams = BlendSpacePtr->FilterInput(&InstanceData->BlendFilter, BlendParameters, DeltaTime);
 				if (BlendSpacePtr->UpdateBlendSamples(FilteredBlendParams, DeltaTime, InstanceData->BlendSamplesData, InstanceData->CachedTriangulationIndex))
 				{
-					// using the blend samples set the weights and activity of the 
+					const float BlendSpacePlayRate = SharedData->GetPlayRate(Binding);
+					const float AxisScaleFactor = BlendSpacePtr->ComputeAxisScaleFactor(BlendParameters, FilteredBlendParams);
+
+					const float BlendSpaceDeltaTime = DeltaTime * BlendSpacePlayRate * AxisScaleFactor;
+
+					// Update our weights based on the new samples
 					for (const FBlendSampleData& CurSample : InstanceData->BlendSamplesData)
 					{
 						if (CurSample.SampleDataIndex < InstanceData->SampleGraphs.Num())
 						{
 							FSampleGraphState& SampleState = InstanceData->SampleGraphs[CurSample.SampleDataIndex];
 							SampleState.Weight = CurSample.TotalWeight;
+							SampleState.DeltaTime = BlendSpaceDeltaTime * CurSample.SamplePlayRate;
 							SampleState.bSampledThisFrame = true;
+
 							InstanceData->NumChildren++;
 						}
 					}
diff --git a/Engine/Plugins/Experimental/AnimNextAnimGraph/Source/AnimNextAnimGraph/Private/Traits/BlendSpacePlayer.h b/Engine/Plugins/Experimental/AnimNextAnimGraph/Source/AnimNextAnimGraph/Private/Traits/BlendSpacePlayer.h
--- a/Engine/Plugins/Experimental/AnimNextAnimGraph/Source/AnimNextAnimGraph/Private/Traits/BlendSpacePlayer.h
+++ b/Engine/Plugins/Experimental/AnimNextAnimGraph/Source/AnimNextAnimGraph/Private/Traits/BlendSpacePlayer.h
@@ -63,10 +63,10 @@ namespace UE::AnimNext
 {
 	/**
 	 * FBlendSpacePlayerTrait
-	 * 
+	 *
 	 * A trait that can play a blend space
 	 */
-	struct FBlendSpacePlayerTrait : FBaseTrait, IEvaluate,  IUpdate, IUpdateTraversal, IHierarchy, ITimeline, IContinuousBlend, IGarbageCollection
+	struct FBlendSpacePlayerTrait : FBaseTrait, IEvaluate, IUpdate, IUpdateTraversal, IHierarchy, ITimeline, IContinuousBlend, IGarbageCollection
 	{
 		DECLARE_ANIM_TRAIT(FBlendSpacePlayerTrait, FBaseTrait)
 
@@ -83,11 +83,15 @@ namespace UE::AnimNext
 			// Weak handle to the graph instance's root
 			FTraitPtr ChildPtr;
 
-			float Weight = 0.0f; 
+			// The weight of this sample
+			float Weight = 0.0f;
 
-			bool bSampledLastFrame = false; 
+			// The scaled delta time for this sample
+			float DeltaTime = 0.0f;
 
-			bool bSampledThisFrame = false; 
+			bool bSampledLastFrame = false;
+
+			bool bSampledThisFrame = false;
 
 		};
 
@@ -101,9 +105,9 @@ namespace UE::AnimNext
 
 			// Cached blend samples updated by task and used to interpolate between points over time
 			TArray<FBlendSampleData> BlendSamplesData;
-			
+
 			/** Previous position in the triangulation/segmentation */
-			int32 CachedTriangulationIndex = -1;
+			int32 CachedTriangulationIndex = INDEX_NONE;
 
 			// Filter used to dampen coordinate changes
 			FBlendFilter BlendFilter;
diff --git a/Engine/Plugins/Experimental/AnimNextAnimGraph/Source/AnimNextAnimGraph/Public/TraitInterfaces/IUpdate.h b/Engine/Plugins/Experimental/AnimNextAnimGraph/Source/AnimNextAnimGraph/Public/TraitInterfaces/IUpdate.h
--- a/Engine/Plugins/Experimental/AnimNextAnimGraph/Source/AnimNextAnimGraph/Public/TraitInterfaces/IUpdate.h
+++ b/Engine/Plugins/Experimental/AnimNextAnimGraph/Source/AnimNextAnimGraph/Public/TraitInterfaces/IUpdate.h
@@ -125,6 +125,14 @@ namespace UE::AnimNext
 			return Result;
 		}
 
+		// Creates a new instance of the update state with the specified delta time
+		[[nodiscard]] FTraitUpdateState WithDeltaTime(float InDeltaTime) const
+		{
+			FTraitUpdateState Result(*this);
+			Result.DeltaTime = InDeltaTime;
+			return Result;
+		}
+
 		// Creates a new instance of the update state marking it as blending out
 		[[nodiscard]] FTraitUpdateState AsBlendingOut(bool InIsBlendingOut = true) const
 		{
diff --git a/Engine/Source/Runtime/Engine/Classes/Animation/BlendSpace.h b/Engine/Source/Runtime/Engine/Classes/Animation/BlendSpace.h
--- a/Engine/Source/Runtime/Engine/Classes/Animation/BlendSpace.h
+++ b/Engine/Source/Runtime/Engine/Classes/Animation/BlendSpace.h
@@ -593,6 +593,9 @@ public:
 	/** Interpolate BlendInput based on Filter data **/
 	ENGINE_API FVector FilterInput(FBlendFilter* Filter, const FVector& BlendInput, float DeltaTime) const;
 
+	/** Computes the axis scale factor. See AxisToScaleAnimation. */
+	ENGINE_API float ComputeAxisScaleFactor(const FVector& BlendSpacePosition, const FVector& FilteredBlendInput) const;
+
 #if WITH_EDITOR	
 	/** Validates sample data for blendspaces using the given animation sequence */
 	ENGINE_API static void UpdateBlendSpacesUsingAnimSequence(UAnimSequenceBase* Sequence);
diff --git a/Engine/Source/Runtime/Engine/Private/Animation/BlendSpace.cpp b/Engine/Source/Runtime/Engine/Private/Animation/BlendSpace.cpp
--- a/Engine/Source/Runtime/Engine/Private/Animation/BlendSpace.cpp
+++ b/Engine/Source/Runtime/Engine/Private/Animation/BlendSpace.cpp
@@ -417,6 +417,64 @@ void UBlendSpace::ForEachImmutableSample(const TFunctionRef<void(const FBlendSam
 	}
 }
 
+float UBlendSpace::ComputeAxisScaleFactor(const FVector& BlendSpacePosition, const FVector& FilteredBlendInput) const
+{
+	const EBlendSpaceAxis AxisToScale = GetAxisToScale();
+	if (AxisToScale == BSA_None)
+	{
+		return 1.0f;
+	}
+
+	float FilterMultiplier = 1.f;
+
+	// first use multiplier using new blendinput
+	// new filtered input is going to be used for sampling animation
+	// so we'll need to change playrate if you'd like to not slide foot
+	if (!BlendSpacePosition.Equals(FilteredBlendInput))
+	{
+		// apply speed change if you want
+		if (AxisToScale == BSA_X)
+		{
+			if (FilteredBlendInput.X != 0.f)
+			{
+				FilterMultiplier = BlendSpacePosition.X / FilteredBlendInput.X;
+			}
+		}
+		else if (AxisToScale == BSA_Y)
+		{
+			if (FilteredBlendInput.Y != 0.f)
+			{
+				FilterMultiplier = BlendSpacePosition.Y / FilteredBlendInput.Y;
+			}
+		}
+	}
+
+	// Now find if clamped input is different. If different, then apply scale to fit in. This allows
+	// "extrapolation" of the blend space outside of the range by time scaling the animation, which is
+	// appropriate when the specified axis is speed (for example).
+	FVector ClampedInput = GetClampedBlendInput(FilteredBlendInput);
+	if (!ClampedInput.Equals(FilteredBlendInput))
+	{
+		// apply speed change if you want, 
+		if (AxisToScale == BSA_X && !BlendParameters[0].bWrapInput)
+		{
+			if (ClampedInput.X != 0.f)
+			{
+				FilterMultiplier *= FilteredBlendInput.X / ClampedInput.X;
+			}
+		}
+		else if (AxisToScale == BSA_Y)
+		{
+			if (ClampedInput.Y != 0.f && !BlendParameters[1].bWrapInput)
+			{
+				FilterMultiplier *= FilteredBlendInput.Y / ClampedInput.Y;
+			}
+		}
+	}
+
+	return FilterMultiplier;
+}
+
 void UBlendSpace::TickAssetPlayer(FAnimTickRecord& Instance, struct FAnimNotifyQueue& NotifyQueue, FAnimAssetTickContext& Context) const
 {
 	check(Instance.BlendSpace.BlendSampleDataCache);
@@ -464,57 +522,11 @@ void UBlendSpace::TickAssetPlayer(FAnimTickRecord& Instance, struct FAnimNotifyQ
 			UE_LOG(LogAnimation, Verbose, TEXT("BlendSpace(%s) - FilteredBlendInput(%s) : PreAnimLength(%0.5f) "), *GetName(), *FilteredBlendInput.ToString(), PreInterpAnimLength);
 		}
 
-		EBlendSpaceAxis AxisToScale = GetAxisToScale();
-		if (AxisToScale != BSA_None)
+		float FilterMultiplier = ComputeAxisScaleFactor(BlendSpacePosition, FilteredBlendInput);
+		if (FilterMultiplier != 1.0f)
 		{
-			float FilterMultiplier = 1.f;
-			// first use multiplier using new blendinput
-			// new filtered input is going to be used for sampling animation
-			// so we'll need to change playrate if you'd like to not slide foot
-			if (!BlendSpacePosition.Equals(FilteredBlendInput))
-			{
-				// apply speed change if you want, 
-				if (AxisToScale == BSA_X)
-				{
-					if (FilteredBlendInput.X != 0.f)
-					{
-						FilterMultiplier = BlendSpacePosition.X / FilteredBlendInput.X;
-					}
-				}
-				else if (AxisToScale == BSA_Y)
-				{
-					if (FilteredBlendInput.Y != 0.f)
-					{
-						FilterMultiplier = BlendSpacePosition.Y / FilteredBlendInput.Y;
-					}
-				}
-			}
-
-			// Now find if clamped input is different. If different, then apply scale to fit in. This allows
-			// "extrapolation" of the blend space outside of the range by time scaling the animation, which is
-			// appropriate when the specified axis is speed (for example).
-			FVector ClampedInput = GetClampedBlendInput(FilteredBlendInput);
-			if (!ClampedInput.Equals(FilteredBlendInput))
-			{
-				// apply speed change if you want, 
-				if (AxisToScale == BSA_X && !BlendParameters[0].bWrapInput)
-				{
-					if (ClampedInput.X != 0.f)
-					{
-						FilterMultiplier *= FilteredBlendInput.X / ClampedInput.X;
-					}
-				}
-				else if (AxisToScale == BSA_Y)
-				{
-					if (ClampedInput.Y != 0.f && !BlendParameters[1].bWrapInput)
-					{
-						FilterMultiplier *= FilteredBlendInput.Y / ClampedInput.Y;
-					}
-				}
-			}
-
 			Instance.DeltaTimeRecord->Delta *= FilterMultiplier;
-			UE_LOG(LogAnimation, Verbose, TEXT("BlendSpace(%s) - FilteredBlendInput(%s) : FilteredBlendInput(%s), FilterMultiplier(%0.2f)"), 
+			UE_LOG(LogAnimation, Verbose, TEXT("BlendSpace(%s) - FilteredBlendInput(%s) : FilteredBlendInput(%s), FilterMultiplier(%0.2f)"),
 				*GetName(), *BlendSpacePosition.ToString(), *FilteredBlendInput.ToString(), FilterMultiplier);
 		}
 
