Handle case where previous sync group leader did not have sync markers

diff --git a/Engine/Plugins/Experimental/AnimNextAnimGraph/Source/AnimNextAnimGraph/Private/Graph/SyncGroup_GraphInstanceComponent.cpp b/Engine/Plugins/Experimental/AnimNextAnimGraph/Source/AnimNextAnimGraph/Private/Graph/SyncGroup_GraphInstanceComponent.cpp
--- a/Engine/Plugins/Experimental/AnimNextAnimGraph/Source/AnimNextAnimGraph/Private/Graph/SyncGroup_GraphInstanceComponent.cpp
+++ b/Engine/Plugins/Experimental/AnimNextAnimGraph/Source/AnimNextAnimGraph/Private/Graph/SyncGroup_GraphInstanceComponent.cpp
@@ -1218,6 +1218,11 @@ namespace UE::AnimNext
 			const int32 NumLongestSyncMarkers = LongestSyncMarkersList->Num();
 			const int32 NumShortestSyncMarkers = ShortestSyncMarkersList->Num();
 
+			if (NumShortestSyncMarkers == 0)
+			{
+				return false;	// Can't match if we have no markers
+			}
+
 			for (int32 LongestIdx = 0; LongestIdx < NumLongestSyncMarkers; ++LongestIdx)
 			{
 				int32 ShortestIdx;
