AnimNext: Run graphs at 'default' mesh component LOD by default

Includes fixup (simple clamp) for rigLogic trait that could access OOB if ther rig logic data didnt contain data for a specific LOD.

diff --git a/Engine/Plugins/Experimental/AnimNext/Source/AnimNext/Public/ReferencePose.h b/Engine/Plugins/Experimental/AnimNext/Source/AnimNext/Public/ReferencePose.h
--- a/Engine/Plugins/Experimental/AnimNext/Source/AnimNext/Public/ReferencePose.h
+++ b/Engine/Plugins/Experimental/AnimNext/Source/AnimNext/Public/ReferencePose.h
@@ -306,6 +306,21 @@ struct TReferencePose
 
 		return ReferenceLocalTransforms[LODBoneIndex].GetScale3D();
 	}
+
+	// Get the LOD level of the 'source' of this pose.
+	// If this pose is generated from a dynamic source, such as a skeletal mesh component, returns the LOD of the component
+	// TODO: As the predicted LOD level can vary across the frame, the LOD should be cached at a module-component level so modules have a consistent view
+	// of the current LOD across the frame. For now we are OK as the LOD level is precalculated in the mesh component and we have manual tick dependencies
+	// set up to prevent races on the value
+	int32 GetSourceLODLevel() const
+	{
+		const USkeletalMeshComponent* MeshComponent = SkeletalMeshComponent.Get();
+		if (MeshComponent == nullptr)
+		{
+			return 0;
+		}
+		return MeshComponent->GetPredictedLODLevel();
+	}
 };
 
 /**
diff --git a/Engine/Plugins/Experimental/AnimNextAnimGraph/Source/AnimNextAnimGraph/Internal/Graph/RigUnit_AnimNextRunAnimationGraph_v1.h b/Engine/Plugins/Experimental/AnimNextAnimGraph/Source/AnimNextAnimGraph/Internal/Graph/RigUnit_AnimNextRunAnimationGraph_v1.h
--- a/Engine/Plugins/Experimental/AnimNextAnimGraph/Source/AnimNextAnimGraph/Internal/Graph/RigUnit_AnimNextRunAnimationGraph_v1.h
+++ b/Engine/Plugins/Experimental/AnimNextAnimGraph/Source/AnimNextAnimGraph/Internal/Graph/RigUnit_AnimNextRunAnimationGraph_v1.h
@@ -35,9 +35,9 @@ struct ANIMNEXTANIMGRAPH_API FRigUnit_AnimNextRunAnimationGraph_v1 : public FRig
 	UPROPERTY(EditAnywhere, Category = "Graph", meta = (Input, Output))
 	FAnimNextAnimGraph Instance;
 
-	// LOD to run the graph at
+	// LOD to run the graph at. If this is -1 then the reference pose's source LOD will be used
 	UPROPERTY(EditAnywhere, Category = "Graph", meta = (Input))
-	int32 LOD = 0;
+	int32 LOD = -1;
 
 	// Reference pose for the graph
 	UPROPERTY(EditAnywhere, Category = "Graph", meta = (Input))
diff --git a/Engine/Plugins/Experimental/AnimNextAnimGraph/Source/AnimNextAnimGraph/Internal/Graph/RigUnit_AnimNextRunAnimationGraph_v2.h b/Engine/Plugins/Experimental/AnimNextAnimGraph/Source/AnimNextAnimGraph/Internal/Graph/RigUnit_AnimNextRunAnimationGraph_v2.h
--- a/Engine/Plugins/Experimental/AnimNextAnimGraph/Source/AnimNextAnimGraph/Internal/Graph/RigUnit_AnimNextRunAnimationGraph_v2.h
+++ b/Engine/Plugins/Experimental/AnimNextAnimGraph/Source/AnimNextAnimGraph/Internal/Graph/RigUnit_AnimNextRunAnimationGraph_v2.h
@@ -37,9 +37,9 @@ struct ANIMNEXTANIMGRAPH_API FRigUnit_AnimNextRunAnimationGraph_v2 : public FRig
 	UPROPERTY(EditAnywhere, Category = "Run", meta = (Input, ShowOnlyInnerProperties, ExportAsReference="true"))
 	FAnimNextAnimGraph Graph;
 
-	// LOD to run the graph at
+	// LOD to run the graph at. If this is -1 then the reference pose's source LOD will be used
 	UPROPERTY(EditAnywhere, Category = "Run", meta = (Input))
-	int32 LOD = 0;
+	int32 LOD = -1;
 
 	// Reference pose for the graph
 	UPROPERTY(EditAnywhere, Category = "Run", meta = (Input))
diff --git a/Engine/Plugins/Experimental/AnimNextAnimGraph/Source/AnimNextAnimGraph/Private/Graph/RigUnit_AnimNextRunAnimationGraph_v1.cpp b/Engine/Plugins/Experimental/AnimNextAnimGraph/Source/AnimNextAnimGraph/Private/Graph/RigUnit_AnimNextRunAnimationGraph_v1.cpp
--- a/Engine/Plugins/Experimental/AnimNextAnimGraph/Source/AnimNextAnimGraph/Private/Graph/RigUnit_AnimNextRunAnimationGraph_v1.cpp
+++ b/Engine/Plugins/Experimental/AnimNextAnimGraph/Source/AnimNextAnimGraph/Private/Graph/RigUnit_AnimNextRunAnimationGraph_v1.cpp
@@ -85,13 +85,19 @@ FRigUnit_AnimNextRunAnimationGraph_v1_Execute()
 
 	const UE::AnimNext::FReferencePose& RefPose = ReferencePose.ReferencePose.GetRef<UE::AnimNext::FReferencePose>();
 
+	int32 DesiredLOD = LOD;
+	if (DesiredLOD == -1)
+	{
+		DesiredLOD = RefPose.GetSourceLODLevel();
+	}
+
 	// TODO: Currently forcing additive flag to false here
-	if (Result.LODPose.ShouldPrepareForLOD(RefPose, LOD, false))
+	if (Result.LODPose.ShouldPrepareForLOD(RefPose, DesiredLOD, false))
 	{
-		Result.LODPose.PrepareForLOD(RefPose, LOD, true, false);
+		Result.LODPose.PrepareForLOD(RefPose, DesiredLOD, true, false);
 	}
 
-	ensure(Result.LODPose.LODLevel == LOD);
+	ensure(Result.LODPose.LODLevel == DesiredLOD);
 
 	// Every graph in a schedule will see the same input events (if they were queued before the schedule started)
 	FUpdateGraphContext UpdateGraphContext(InstanceRef.Get(), ExecuteContext.GetDeltaTime());
@@ -117,7 +123,7 @@ FRigUnit_AnimNextRunAnimationGraph_v1_Execute()
 
 	UpdateGraph(UpdateGraphContext);
 	
-	FEvaluateGraphContext EvaluateGraphContext(InstanceRef.Get(), RefPose, LOD);
+	FEvaluateGraphContext EvaluateGraphContext(InstanceRef.Get(), RefPose, DesiredLOD);
 	EvaluateGraphContext.SetBindingObject(RefPose.SkeletalMeshComponent);
 	EvaluateGraph(EvaluateGraphContext, Result);
 
diff --git a/Engine/Plugins/Experimental/AnimNextAnimGraph/Source/AnimNextAnimGraph/Private/Graph/RigUnit_AnimNextRunAnimationGraph_v2.cpp b/Engine/Plugins/Experimental/AnimNextAnimGraph/Source/AnimNextAnimGraph/Private/Graph/RigUnit_AnimNextRunAnimationGraph_v2.cpp
--- a/Engine/Plugins/Experimental/AnimNextAnimGraph/Source/AnimNextAnimGraph/Private/Graph/RigUnit_AnimNextRunAnimationGraph_v2.cpp
+++ b/Engine/Plugins/Experimental/AnimNextAnimGraph/Source/AnimNextAnimGraph/Private/Graph/RigUnit_AnimNextRunAnimationGraph_v2.cpp
@@ -35,13 +35,19 @@ FRigUnit_AnimNextRunAnimationGraph_v2_Execute()
 
 	const UE::AnimNext::FReferencePose& RefPose = ReferencePose.ReferencePose.GetRef<UE::AnimNext::FReferencePose>();
 
+	int32 DesiredLOD = LOD;
+	if (DesiredLOD == -1)
+	{
+		DesiredLOD = RefPose.GetSourceLODLevel();
+	}
+
 	// TODO: Currently forcing additive flag to false here
-	if (Result.LODPose.ShouldPrepareForLOD(RefPose, LOD, false))
+	if (Result.LODPose.ShouldPrepareForLOD(RefPose, DesiredLOD, false))
 	{
-		Result.LODPose.PrepareForLOD(RefPose, LOD, true, false);
+		Result.LODPose.PrepareForLOD(RefPose, DesiredLOD, true, false);
 	}
 
-	ensure(Result.LODPose.LODLevel == LOD);
+	ensure(Result.LODPose.LODLevel == DesiredLOD);
 
 	// Get a host to run this graph
 	const UAnimNextAnimationGraph* HostGraph = Graph.HostGraph ? Graph.HostGraph.Get() : GetDefault<UAnimNextAnimGraphSettings>()->GetDefaultRunGraphHost();// GetHostGraphToRun(ExecuteContext);
@@ -145,7 +151,7 @@ FRigUnit_AnimNextRunAnimationGraph_v2_Execute()
 
 	UpdateGraph(UpdateGraphContext);
 
-	FEvaluateGraphContext EvaluateGraphContext(HostInstance, RefPose, LOD);
+	FEvaluateGraphContext EvaluateGraphContext(HostInstance, RefPose, DesiredLOD);
 	EvaluateGraphContext.SetBindingObject(RefPose.SkeletalMeshComponent);
 	EvaluateGraph(EvaluateGraphContext, Result);
 
diff --git a/Engine/Plugins/Experimental/RigLogicAnimNext/Source/RigLogicAnimNext/Private/RigLogicInstanceData.cpp b/Engine/Plugins/Experimental/RigLogicAnimNext/Source/RigLogicAnimNext/Private/RigLogicInstanceData.cpp
--- a/Engine/Plugins/Experimental/RigLogicAnimNext/Source/RigLogicAnimNext/Private/RigLogicInstanceData.cpp
+++ b/Engine/Plugins/Experimental/RigLogicAnimNext/Source/RigLogicAnimNext/Private/RigLogicInstanceData.cpp
@@ -55,13 +55,14 @@ namespace UE::AnimNext
 		CachedDNAIndexMapping = MakeShared<FDNAIndexMapping>();
 		CachedDNAIndexMapping->Init(DNAAsset->GetBehaviorReader().Get(), Skeleton, SkeletalMesh);
 
+		NumLODs = CachedRigRuntimeContext->VariableJointIndicesPerLOD.Num();
+
 		InitBoneIndexMapping(ReferencePose);
 		InitSparseAndDenseDriverJointMapping(ReferencePose);
 	}
 
 	void FRigLogicAnimNextInstanceData::InitBoneIndexMapping(const UE::AnimNext::FReferencePose* ReferencePose)
 	{
-		const uint32 NumLODs = CachedRigRuntimeContext->VariableJointIndicesPerLOD.Num();
 		RigLogicToSkeletonBoneIndexMappingPerLOD.Empty();
 		RigLogicToSkeletonBoneIndexMappingPerLOD.SetNum(NumLODs);
 
@@ -95,8 +96,6 @@ namespace UE::AnimNext
 
 	void FRigLogicAnimNextInstanceData::InitSparseAndDenseDriverJointMapping(const UE::AnimNext::FReferencePose* ReferencePose)
 	{
-		const uint32 NumLODs = CachedRigRuntimeContext->VariableJointIndicesPerLOD.Num();
-
 		// Populate driver joint to raw control attribute mapping (used to feed RigLogic with inputs from the joint hierarchy)
 		SparseDriverJointsToControlAttributesMapPerLOD.Empty();
 		DenseDriverJointsToControlAttributesMapPerLOD.Empty();
diff --git a/Engine/Plugins/Experimental/RigLogicAnimNext/Source/RigLogicAnimNext/Private/RigLogicTask.cpp b/Engine/Plugins/Experimental/RigLogicAnimNext/Source/RigLogicAnimNext/Private/RigLogicTask.cpp
--- a/Engine/Plugins/Experimental/RigLogicAnimNext/Source/RigLogicAnimNext/Private/RigLogicTask.cpp
+++ b/Engine/Plugins/Experimental/RigLogicAnimNext/Source/RigLogicAnimNext/Private/RigLogicTask.cpp
@@ -54,7 +54,7 @@ void FAnimNextRigLogicTask::Execute(UE::AnimNext::FEvaluationVM& VM) const
 		{
 			if (const FRigLogic* RigLogic = RigRuntimeContext->RigLogic.Get())
 			{
-				const int32 LODLevel = Keyframe->Pose.LODLevel;
+				const int32 LODLevel = FMath::Clamp(Keyframe->Pose.LODLevel, 0, (int32)PoolInstanceData->NumLODs - 1);
 
 				// 1. Feed input curves to RigLogic.
 				FBlendedCurve& Curves = Keyframe->Curves;
diff --git a/Engine/Plugins/Experimental/RigLogicAnimNext/Source/RigLogicAnimNext/Public/RigLogicInstanceData.h b/Engine/Plugins/Experimental/RigLogicAnimNext/Source/RigLogicAnimNext/Public/RigLogicInstanceData.h
--- a/Engine/Plugins/Experimental/RigLogicAnimNext/Source/RigLogicAnimNext/Public/RigLogicInstanceData.h
+++ b/Engine/Plugins/Experimental/RigLogicAnimNext/Source/RigLogicAnimNext/Public/RigLogicInstanceData.h
@@ -43,6 +43,8 @@ namespace UE::AnimNext
 		TArray<TArray<FPoseBoneControlAttributeMapping>> SparseDriverJointsToControlAttributesMapPerLOD;
 		TArray<TArray<FPoseBoneControlAttributeMapping>> DenseDriverJointsToControlAttributesMapPerLOD;
 
+		uint32 NumLODs = 0;
+
 		void Init(const UE::AnimNext::FReferencePose* ReferencePose);
 
 	private:
