Improve trait event list performance
Add flag to control whether or not to track timeline updates in injection requests

diff --git a/Engine/Plugins/Experimental/AnimNext/Source/AnimNext/Public/TraitCore/TraitEventList.h b/Engine/Plugins/Experimental/AnimNext/Source/AnimNext/Public/TraitCore/TraitEventList.h
--- a/Engine/Plugins/Experimental/AnimNext/Source/AnimNext/Public/TraitCore/TraitEventList.h
+++ b/Engine/Plugins/Experimental/AnimNext/Source/AnimNext/Public/TraitCore/TraitEventList.h
@@ -39,7 +39,7 @@ namespace UE::AnimNext
 
 	private:
 		// A list of events
-		TArray<FAnimNextTraitEventPtr> Events;
+		TArray<FAnimNextTraitEventPtr, TInlineAllocator<4>> Events;
 	};
 
 	// Decrements and purges expired entries from the specified event list
diff --git a/Engine/Plugins/Experimental/AnimNextAnimGraph/Source/AnimNextAnimGraph/Private/Injection/InjectionRequest.cpp b/Engine/Plugins/Experimental/AnimNextAnimGraph/Source/AnimNextAnimGraph/Private/Injection/InjectionRequest.cpp
--- a/Engine/Plugins/Experimental/AnimNextAnimGraph/Source/AnimNextAnimGraph/Private/Injection/InjectionRequest.cpp
+++ b/Engine/Plugins/Experimental/AnimNextAnimGraph/Source/AnimNextAnimGraph/Private/Injection/InjectionRequest.cpp
@@ -169,6 +169,7 @@ namespace UE::AnimNext
 	const FTimelineState& FInjectionRequest::GetTimelineState() const
 	{
 		check(IsInGameThread());
+		ensureMsgf(RequestArgs.bTrackTimelineProgress, TEXT("Attempting to query the timeline state of an injection request that isn't tracking the timeline progress"));
 		return TimelineState;
 	}
 
diff --git a/Engine/Plugins/Experimental/AnimNextAnimGraph/Source/AnimNextAnimGraph/Private/Injection/InjectionSiteTrait.cpp b/Engine/Plugins/Experimental/AnimNextAnimGraph/Source/AnimNextAnimGraph/Private/Injection/InjectionSiteTrait.cpp
--- a/Engine/Plugins/Experimental/AnimNextAnimGraph/Source/AnimNextAnimGraph/Private/Injection/InjectionSiteTrait.cpp
+++ b/Engine/Plugins/Experimental/AnimNextAnimGraph/Source/AnimNextAnimGraph/Private/Injection/InjectionSiteTrait.cpp
@@ -283,6 +283,8 @@ namespace UE::AnimNext
 					CurrentRequestTimeLeft = TimeToEnd;
 				}
 
+				// Only raise a timeline update event if we care about it
+				if (InjectionTracker.Request->GetArgs().bTrackTimelineProgress)
 				{
 					auto TimelineUpdateEvent = MakeTraitEvent<FInjection_TimelineUpdateEvent>();
 					TimelineUpdateEvent->Request = InjectionTracker.Request;
diff --git a/Engine/Plugins/Experimental/AnimNextAnimGraph/Source/AnimNextAnimGraph/Private/PlayAnim/PlayAnimSlotTrait.cpp b/Engine/Plugins/Experimental/AnimNextAnimGraph/Source/AnimNextAnimGraph/Private/PlayAnim/PlayAnimSlotTrait.cpp
--- a/Engine/Plugins/Experimental/AnimNextAnimGraph/Source/AnimNextAnimGraph/Private/PlayAnim/PlayAnimSlotTrait.cpp
+++ b/Engine/Plugins/Experimental/AnimNextAnimGraph/Source/AnimNextAnimGraph/Private/PlayAnim/PlayAnimSlotTrait.cpp
@@ -264,6 +264,8 @@ namespace UE::AnimNext
 				CurrentRequestTimeLeft = TimeToEnd;
 			}
 
+			// Only raise a timeline update event if we care about it
+			if (SlotRequest.Request->GetArgs().bTrackTimelineProgress)
 			{
 				auto TimelineUpdateEvent = MakeTraitEvent<FInjection_TimelineUpdateEvent>();
 				TimelineUpdateEvent->Request = SlotRequest.Request;
diff --git a/Engine/Plugins/Experimental/AnimNextAnimGraph/Source/AnimNextAnimGraph/Public/Injection/InjectionRequest.h b/Engine/Plugins/Experimental/AnimNextAnimGraph/Source/AnimNextAnimGraph/Public/Injection/InjectionRequest.h
--- a/Engine/Plugins/Experimental/AnimNextAnimGraph/Source/AnimNextAnimGraph/Public/Injection/InjectionRequest.h
+++ b/Engine/Plugins/Experimental/AnimNextAnimGraph/Source/AnimNextAnimGraph/Public/Injection/InjectionRequest.h
@@ -139,6 +139,10 @@ private:
 	UPROPERTY()
 	EAnimNextInjectionLifetimeType LifetimeType = EAnimNextInjectionLifetimeType::Auto;
 
+	// Whether or not the request should track the timeline progress
+	UPROPERTY()
+	bool bTrackTimelineProgress = false;
+
 	// Payload that will be applied to the animation graph's variables via its data interfaces.
 	UPROPERTY()
 	FAnimNextDataInterfacePayload Payload;
@@ -195,7 +199,7 @@ namespace UE::AnimNext
 		// Returns the request status
 		ANIMNEXTANIMGRAPH_API EInjectionStatus GetStatus() const;
 
-		// Returns the current timeline state
+		// Returns the current timeline state (make sure to enable FInjectionRequestArgs::bTrackTimelineProgress to use this)
 		ANIMNEXTANIMGRAPH_API const FTimelineState& GetTimelineState() const;
 
 		// Returns whether or not this request has expired
