Add new sequence info track to display currently playing animation clips current time and sync marker info

#tests animsandbox

diff --git a/Engine/Plugins/Experimental/AnimNextAnimGraph/Source/AnimNextAnimGraphEditor/Private/AnimNextAnimGraphEditorModule.cpp b/Engine/Plugins/Experimental/AnimNextAnimGraph/Source/AnimNextAnimGraphEditor/Private/AnimNextAnimGraphEditorModule.cpp
--- a/Engine/Plugins/Experimental/AnimNextAnimGraph/Source/AnimNextAnimGraphEditor/Private/AnimNextAnimGraphEditorModule.cpp
+++ b/Engine/Plugins/Experimental/AnimNextAnimGraph/Source/AnimNextAnimGraphEditor/Private/AnimNextAnimGraphEditorModule.cpp
@@ -38,6 +38,7 @@
 #include "Module/AnimNextEventGraphSchema.h"
 #include "RewindDebugger/AnimNextAnimGraphTraceModule.h"
 #include "RewindDebugger/EvaluationProgramTrack.h"
+#include "RewindDebugger/SequenceInfoTrack.h"
 #include "PersonaModule.h"
 #include "Developer/TraceServices/Public/TraceServices/ModuleService.h"
 
@@ -47,6 +48,7 @@ namespace
 {
 	FAnimNextAnimGraphTraceModule GAnimNextAnimGraphTraceModule;
 	UE::AnimNextEditor::FEvaluationProgramTrackCreator GAnimNextModulesTrackCreator;
+	UE::AnimNextEditor::FSequenceInfoTrackCreator GSequenceInfoTrackCreator;
 }
 
 namespace UE::AnimNext::Editor
@@ -71,8 +73,13 @@ class FAnimNextAnimGraphEditorModule : public IModuleInterface
 					FAnimNextAnimGraph::StaticStruct()->GetFName(),
 					FOnGetPropertyTypeCustomizationInstance::CreateLambda([] { return MakeShared<FAnimNextGraphDetails>(); }));
 					
+		PropertyModule.RegisterCustomPropertyTypeLayout(
+							FAnimNextSequenceTraceInfo::StaticStruct()->GetFName(),
+							FOnGetPropertyTypeCustomizationInstance::CreateLambda([] { return MakeShared<AnimNextEditor::FAnimNextSequenceTraceInfoCustomization>(); }));					
+					
 		IModularFeatures::Get().RegisterModularFeature(TraceServices::ModuleFeatureName, &GAnimNextAnimGraphTraceModule);
 		IModularFeatures::Get().RegisterModularFeature(RewindDebugger::IRewindDebuggerTrackCreator::ModularFeatureName, &GAnimNextModulesTrackCreator);
+		IModularFeatures::Get().RegisterModularFeature(RewindDebugger::IRewindDebuggerTrackCreator::ModularFeatureName, &GSequenceInfoTrackCreator);
 
 		Workspace::IWorkspaceEditorModule& WorkspaceEditorModule = FModuleManager::Get().LoadModuleChecked<Workspace::IWorkspaceEditorModule>("WorkspaceEditor");
 
diff --git a/Engine/Plugins/Experimental/AnimNextAnimGraph/Source/AnimNextAnimGraphEditor/Private/RewindDebugger/SequenceInfoTrack.cpp b/Engine/Plugins/Experimental/AnimNextAnimGraph/Source/AnimNextAnimGraphEditor/Private/RewindDebugger/SequenceInfoTrack.cpp
new file mode 100644
--- /dev/null
+++ b/Engine/Plugins/Experimental/AnimNextAnimGraph/Source/AnimNextAnimGraphEditor/Private/RewindDebugger/SequenceInfoTrack.cpp
@@ -0,0 +1,357 @@
+// Copyright Epic Games, Inc. All Rights Reserved.
+
+#include "SequenceInfoTrack.h"
+
+#include "AnimNextAnimGraphProvider.h"
+#include "DetailWidgetRow.h"
+#include "IDetailChildrenBuilder.h"
+#include "IGameplayProvider.h"
+#include "IRewindDebugger.h"
+#include "Modules/ModuleManager.h"
+#include "PropertyEditorModule.h"
+#include "Animation/AnimSequence.h"
+#include "EvaluationVM/Tasks/PushAnimSequenceKeyframe.h"
+#include "Serialization/ObjectAndNameAsStringProxyArchive.h"
+#include "Widgets/SCanvas.h"
+#include "Widgets/Text/STextBlock.h"
+#include "Widgets/Colors/SColorBlock.h"
+
+#define LOCTEXT_NAMESPACE "EvaluationProgramTrack"
+
+namespace UE::AnimNextEditor
+{
+static const FName AnimNextModulesName("AnimNextModules");
+
+FName FSequenceInfoTrackCreator::GetTargetTypeNameInternal() const
+{
+	static const FName ObjectName("AnimNextComponent");
+	return ObjectName;
+}
+
+
+FText FSequenceInfoTrack::GetDisplayNameInternal() const
+{
+	return NSLOCTEXT("RewindDebugger", "SequenceInfoTrackName", "SequenceInfo");
+}
+
+void FSequenceInfoTrackCreator::GetTrackTypesInternal(TArray<RewindDebugger::FRewindDebuggerTrackType>& Types) const
+{
+	Types.Add({FSequenceInfoTrack::TrackName, LOCTEXT("AnimNextSequenceInfo", "AnimNextSequenceInfo")});
+}
+
+TSharedPtr<RewindDebugger::FRewindDebuggerTrack> FSequenceInfoTrackCreator::CreateTrackInternal(uint64 ObjectId) const
+{
+	return MakeShared<FSequenceInfoTrack>(ObjectId);
+}
+
+
+void FAnimNextSequenceTraceInfoCustomization::CustomizeHeader(TSharedRef<IPropertyHandle> StructPropertyHandle, class FDetailWidgetRow& HeaderRow,
+                                                              IPropertyTypeCustomizationUtils& StructCustomizationUtils)
+{
+}
+
+// Copied from SAnimNotifyPanel
+// Todo: find common place for this to live
+FLinearColor GenerateColorFromName(FName Name)
+{
+	constexpr uint8 Saturation = 255;
+	constexpr uint8 Luminosity = 255;
+	const uint8 Hue = static_cast<uint8>(GetTypeHash(Name.ToString())) * 157;
+	return FLinearColor::MakeFromHSV8(Hue, Saturation, Luminosity);
+}
+
+// Helper for lambdas
+const FAnimNextSequenceTraceInfo* ExtractSequenceTraceInfo(TSharedRef<IPropertyHandle> StructPropertyHandle)
+{
+	void* StructData = nullptr;
+	FPropertyAccess::Result Result = StructPropertyHandle->GetValueData(StructData);
+	if (Result != FPropertyAccess::Success)
+		return nullptr;
+
+	const FAnimNextSequenceTraceInfo* SequenceInfo = static_cast<FAnimNextSequenceTraceInfo*>(StructData);
+	return SequenceInfo;
+}
+
+void FAnimNextSequenceTraceInfoCustomization::CustomizeChildren(TSharedRef<IPropertyHandle> StructPropertyHandle,
+                                                                class IDetailChildrenBuilder& StructBuilder,
+                                                                IPropertyTypeCustomizationUtils& StructCustomizationUtils)
+{
+	if (FStructProperty* StructProperty = CastField<FStructProperty>(StructPropertyHandle->GetProperty()))
+	{
+		if (StructProperty->Struct == FAnimNextSequenceTraceInfo::StaticStruct())
+		{
+			static const FColor TimelineBackground(0xFF575761);
+			static const FColor TimelineForeground(0xFFF59F00);
+			constexpr float TimelineWidth = 300.f;
+			constexpr float TimelineHeight = 20.f;
+
+			const FAnimNextSequenceTraceInfo* SequenceInfo = ExtractSequenceTraceInfo(StructPropertyHandle);
+			if (SequenceInfo == nullptr)
+				return;
+
+			FString SequenceName = SequenceInfo->AnimSequence.IsValid() ? SequenceInfo->AnimSequence->GetName() : TEXT("NULL");
+
+			TSharedRef<SCanvas> TimelineCanvas =
+				SNew(SCanvas)
+
+				// Before playhead block
+				+ SCanvas::Slot()
+				.HAlign(HAlign_Left)
+				.VAlign(VAlign_Center)
+				.Position(FVector2D::ZeroVector)
+				.Size_Lambda( [=]()
+				{
+					const FAnimNextSequenceTraceInfo* UpdatedSequenceInfo = ExtractSequenceTraceInfo(StructPropertyHandle);
+					if (UpdatedSequenceInfo == nullptr)
+						return FVector2D::ZeroVector;
+					return FVector2D(TimelineWidth * UpdatedSequenceInfo->CalcAnimTimeRatio(), TimelineHeight);
+				} )
+				[
+					SNew(SColorBlock)
+					.Color(TimelineForeground)
+				]
+
+				// After playhead block
+				+ SCanvas::Slot()
+				.HAlign(HAlign_Left)
+				.VAlign(VAlign_Center)
+				.Position_Lambda([=]()
+				{
+					const FAnimNextSequenceTraceInfo* UpdatedSequenceInfo = ExtractSequenceTraceInfo(StructPropertyHandle);
+					if (UpdatedSequenceInfo == nullptr)
+						return FVector2D::ZeroVector;
+					
+					return FVector2D(UpdatedSequenceInfo->CalcAnimTimeRatio() * TimelineWidth, 0.f);
+				})
+				.Size_Lambda([=]()
+				{
+					const FAnimNextSequenceTraceInfo* UpdatedSequenceInfo = ExtractSequenceTraceInfo(StructPropertyHandle);
+					if (UpdatedSequenceInfo == nullptr)
+						return FVector2D::ZeroVector;
+					
+					return FVector2D(TimelineWidth * (1.0f - UpdatedSequenceInfo->CalcAnimTimeRatio()), TimelineHeight);
+				})
+				[
+					SNew(SColorBlock)
+					.Color(TimelineBackground)
+				]
+
+				// Time info
+				+ SCanvas::Slot()
+				.Position(FVector2D(TimelineWidth, 0.f))
+				.Size(FVector2D(TimelineWidth, TimelineHeight))
+				.HAlign(HAlign_Left)
+				.VAlign(VAlign_Center)
+				[
+					SNew(STextBlock)
+					                .Text_Lambda([=]()
+					                {
+					                	const FAnimNextSequenceTraceInfo* UpdatedSequenceInfo = ExtractSequenceTraceInfo(StructPropertyHandle);
+					                	if (UpdatedSequenceInfo == nullptr)
+					                		return FText::FromString("NULL");
+						                return FText::FromString(
+					   FString::Printf(TEXT("%.2fs (%.0f%%)"), UpdatedSequenceInfo->CurrentTimeSeconds, UpdatedSequenceInfo->CalcAnimTimeRatio() * 100.0f));
+					                })
+					                .Margin(FMargin(2.0f))
+					                .Justification(ETextJustify::Left)
+				];
+
+			// Todo: how to do this with lambda?
+			// Add all sync markers
+			if (SequenceInfo->DurationSeconds > 0.0f)
+			{
+				for (auto Marker : SequenceInfo->SyncMarkers)
+				{
+					constexpr float MarkerHeight = 8.f;
+					constexpr float MarkerWidth = 4.0f;
+					float MarkerTimeRatio = Marker.Time / SequenceInfo->DurationSeconds;
+
+					TimelineCanvas->AddSlot()
+					              .HAlign(HAlign_Left)
+					              .VAlign(VAlign_Center)
+					              .Position(FVector2D((MarkerTimeRatio * TimelineWidth) - 0.5f * MarkerWidth,
+					                                  (0.5f * (TimelineHeight) + MarkerHeight)))
+					              .Size(FVector2D(MarkerWidth, MarkerHeight))
+					[
+						SNew(SColorBlock)
+						.Color(GenerateColorFromName(Marker.Name))
+						.ToolTipText(FText::FromString(Marker.Name.ToString()))
+					];
+				}
+			}
+
+			TSharedPtr<IPropertyHandle> SequenceProperty =
+				StructPropertyHandle->GetChildHandle(GET_MEMBER_NAME_CHECKED(FAnimNextSequenceTraceInfo, AnimSequence));
+
+			StructBuilder.AddCustomRow(LOCTEXT("FAnimNextSequenceTraceInfoRow", "FAnimNextSequenceTraceInfo"))
+			             .NameContent()
+				[
+					SequenceProperty->CreatePropertyValueWidget()
+				]
+				.ValueContent()
+				[
+					TimelineCanvas
+				];
+		}
+	}
+
+	StructPropertyHandle->SetOnPropertyValueChanged(
+		FSimpleDelegate::CreateLambda([]
+		{
+			UE_LOG(LogTemp, Display, TEXT("PropertyChange"));
+		}));
+}
+
+FSequenceInfoTrack::FSequenceInfoTrack(uint64 InObjectId)
+	: ObjectId(InObjectId)
+{
+	Initialize();
+}
+
+FSequenceInfoTrack::FSequenceInfoTrack(uint64 InObjectId, uint64 InInstanceId)
+	: ObjectId(InObjectId)
+	, InstanceId(InInstanceId)
+{
+	Initialize();
+}
+
+
+FSequenceInfoTrack::~FSequenceInfoTrack()
+{
+	if (USequenceInfoDetailsObject* DetailsObject = DetailsObjectWeakPtr.Get())
+	{
+		DetailsObject->ClearFlags(RF_Standalone);
+	}
+}
+
+void FSequenceInfoTrack::Initialize()
+{
+	ExistenceRange = MakeShared<SEventTimelineView::FTimelineEventData>();
+	ExistenceRange->Windows.Add({0, 0, GetDisplayNameInternal(), GetDisplayNameInternal(), FLinearColor(0.1f, 0.15f, 0.11f)});
+
+	FPropertyEditorModule& PropertyEditorModule = FModuleManager::GetModuleChecked<FPropertyEditorModule>("PropertyEditor");
+	FDetailsViewArgs DetailsViewArgs;
+	DetailsViewArgs.NameAreaSettings = FDetailsViewArgs::HideNameArea;
+	DetailsView = PropertyEditorModule.CreateDetailView(DetailsViewArgs);
+
+	InitializeDetailsObject();
+}
+
+USequenceInfoDetailsObject* FSequenceInfoTrack::InitializeDetailsObject()
+{
+	USequenceInfoDetailsObject* DetailsObject = NewObject<USequenceInfoDetailsObject>();
+	DetailsObject->SetFlags(RF_Standalone);
+	DetailsObjectWeakPtr = MakeWeakObjectPtr(DetailsObject);
+	DetailsView->SetObject(DetailsObject);
+	return DetailsObject;
+}
+
+void FSequenceInfoTrack::RefreshSequenceInfoFromEvaluationProgram(TArray<FAnimNextSequenceTraceInfo>& OutSequenceInfo,
+                                                                  const FSerializableEvaluationProgram& Program)
+{
+	OutSequenceInfo.Reset();
+	for (const FInstancedStruct& Task : Program.Tasks)
+	{
+		if (Task.GetScriptStruct() == FAnimNextAnimSequenceKeyframeTask::StaticStruct())
+		{
+			const FAnimNextAnimSequenceKeyframeTask* SequenceTask = Task.GetPtr<FAnimNextAnimSequenceKeyframeTask>();
+			FAnimNextSequenceTraceInfo& TraceInfo = OutSequenceInfo.Emplace_GetRef();
+			TraceInfo.AnimSequence = SequenceTask->AnimSequence;
+			TraceInfo.DurationSeconds = SequenceTask->AnimSequence.IsValid() ? TraceInfo.AnimSequence->GetPlayLength() : 0.0f;
+			TraceInfo.CurrentTimeSeconds = SequenceTask->SampleTime;
+
+			// Fill out sync markers
+			TraceInfo.SyncMarkers.Reserve(SequenceTask->AnimSequence->AuthoredSyncMarkers.Num());
+			for (auto Marker : SequenceTask->AnimSequence->AuthoredSyncMarkers)
+			{
+				TraceInfo.SyncMarkers.Emplace(Marker.Time, Marker.MarkerName);
+			}
+		}
+	}
+}
+
+bool FSequenceInfoTrack::UpdateInternal()
+{
+	TRACE_CPUPROFILER_EVENT_SCOPE(FSequenceInfoTrack::UpdateInternal);
+
+	IRewindDebugger* RewindDebugger = IRewindDebugger::Instance();
+
+	const TraceServices::IAnalysisSession* AnalysisSession = RewindDebugger->GetAnalysisSession();
+
+	bool bChanged = false;
+
+	if (const FAnimNextAnimGraphProvider* AnimNextAnimGraphProvider = IRewindDebugger::Instance()->GetAnalysisSession()->ReadProvider<
+		FAnimNextAnimGraphProvider>("AnimNextAnimGraphProvider"))
+	{
+		double CurrentScrubTime = IRewindDebugger::Instance()->CurrentTraceTime();
+
+		USequenceInfoDetailsObject* DetailsObject = DetailsObjectWeakPtr.Get();
+		if (DetailsObject == nullptr)
+		{
+			// this should not happen unless the object was garbage collected (which should not happen since it's marked as Standalone)
+			Initialize();
+		}
+
+		if (InstanceId == 0)
+		{
+			AnimNextAnimGraphProvider->EnumerateEvaluationGraphs(ObjectId, [this](uint64 GraphId)
+			{
+				InstanceId = GraphId;
+			});
+		}
+
+		if (InstanceId != 0)
+		{
+			if (const FEvaluationProgramData* Data = AnimNextAnimGraphProvider->GetEvaluationProgramData(InstanceId))
+			{
+				if (PreviousScrubTime != CurrentScrubTime)
+				{
+					PreviousScrubTime = CurrentScrubTime;
+
+					const TraceServices::IFrameProvider& FramesProvider = TraceServices::ReadFrameProvider(*AnalysisSession);
+					TraceServices::FAnalysisSessionReadScope SessionReadScope(*AnalysisSession);
+
+					TraceServices::FFrame MarkerFrame;
+					if (FramesProvider.GetFrameFromTime(ETraceFrameType::TraceFrameType_Game, CurrentScrubTime, MarkerFrame))
+					{
+						Data->EvaluationProgramTimeline.EnumerateEvents(MarkerFrame.StartTime, MarkerFrame.EndTime,
+						[DetailsObject, this] ( double InStartTime, double InEndTime, uint32 InDepth, const TArray<uint8>& VariableData)
+						{
+							// Why do this if panel is not selected?
+							FMemoryReader Reader(VariableData);
+							FObjectAndNameAsStringProxyArchive Archive(
+							Reader, /*bInLoadIfFindFails*/true);
+
+							static const FSerializableEvaluationProgram Defaults;
+							FSerializableEvaluationProgram Program;
+							FSerializableEvaluationProgram::StaticStruct()->SerializeItem(
+							Archive, &Program, &Defaults);
+							RefreshSequenceInfoFromEvaluationProgram(
+							DetailsObject->SequenceTraceInfo, Program);
+							DetailsView->ForceRefresh(); // We need to force refresh to get the sync markers to refresh. todo: how to handle without recreating UI?
+							return TraceServices::EEventEnumerate::Stop;
+							});
+						}
+				}
+			}
+		}
+	}
+
+	return bChanged;
+}
+
+bool FSequenceInfoTrackCreator::HasDebugInfoInternal(uint64 ObjectId) const
+{
+	const IRewindDebugger* RewindDebugger = IRewindDebugger::Instance();
+	if (const TraceServices::IAnalysisSession* AnalysisSession = RewindDebugger->GetAnalysisSession())
+	{
+		TraceServices::FAnalysisSessionReadScope SessionReadScope(*AnalysisSession);
+
+		return true;
+	}
+
+	return false;
+}
+}
+
+#undef LOCTEXT_NAMESPACE
diff --git a/Engine/Plugins/Experimental/AnimNextAnimGraph/Source/AnimNextAnimGraphEditor/Private/RewindDebugger/SequenceInfoTrack.h b/Engine/Plugins/Experimental/AnimNextAnimGraph/Source/AnimNextAnimGraphEditor/Private/RewindDebugger/SequenceInfoTrack.h
new file mode 100644
--- /dev/null
+++ b/Engine/Plugins/Experimental/AnimNextAnimGraph/Source/AnimNextAnimGraphEditor/Private/RewindDebugger/SequenceInfoTrack.h
@@ -0,0 +1,121 @@
+// Copyright Epic Games, Inc. All Rights Reserved.
+#pragma once
+
+#include "IRewindDebuggerTrackCreator.h"
+#include "RewindDebuggerTrack.h"
+#include "EvaluationVM/SerializableEvaluationProgram.h"
+#include "Textures/SlateIcon.h"
+#include "SEventTimelineView.h"
+#include "IDetailsView.h"
+#include "IPropertyTypeCustomization.h"
+
+#include "SequenceInfoTrack.generated.h"
+
+class UAnimSequence;
+
+USTRUCT()
+struct FAnimNextSyncMarkerTraceInfo
+{
+	GENERATED_BODY()
+	
+	UPROPERTY()
+	float Time;
+
+	UPROPERTY()
+	FName Name;
+};
+
+USTRUCT()
+struct FAnimNextSequenceTraceInfo
+{
+	GENERATED_BODY()
+
+	FORCEINLINE float CalcAnimTimeRatio() const
+	{
+		return DurationSeconds > 0.0f ? CurrentTimeSeconds / DurationSeconds : 0.0f;
+	}
+
+	UPROPERTY(VisibleAnywhere, Category = Properties)
+	TWeakObjectPtr<UAnimSequence> AnimSequence;
+
+	UPROPERTY(VisibleAnywhere, Category = Properties)
+	float DurationSeconds;
+
+	UPROPERTY(VisibleAnywhere, Category = Properties)
+	float CurrentTimeSeconds;
+
+	UPROPERTY(VisibleAnywhere, Category = Properties)
+	TArray<FAnimNextSyncMarkerTraceInfo> SyncMarkers;
+};
+
+UCLASS()
+class USequenceInfoDetailsObject : public UObject
+{
+	GENERATED_BODY()
+
+public:
+	UPROPERTY(EditAnywhere, Category=Details, meta=(ShowOnlyInnerProperties, FullyExpand=true))
+	TArray<FAnimNextSequenceTraceInfo> SequenceTraceInfo;
+};
+
+namespace UE::AnimNextEditor
+{
+class FAnimNextSequenceTraceInfoCustomization : public IPropertyTypeCustomization
+{
+public:
+	virtual void CustomizeHeader(TSharedRef<IPropertyHandle> StructPropertyHandle,
+		class FDetailWidgetRow& HeaderRow,
+		IPropertyTypeCustomizationUtils& StructCustomizationUtils) override;
+	virtual void CustomizeChildren(TSharedRef<IPropertyHandle> StructPropertyHandle,
+		class IDetailChildrenBuilder& StructBuilder,
+		IPropertyTypeCustomizationUtils& StructCustomizationUtils) override;
+};
+
+class FSequenceInfoTrack : public RewindDebugger::FRewindDebuggerTrack
+{
+public:
+	FSequenceInfoTrack(uint64 InObjectId);
+	FSequenceInfoTrack(uint64 InObjectId, uint64 InstanceId);
+	virtual ~FSequenceInfoTrack();
+	
+	TSharedPtr<SEventTimelineView::FTimelineEventData> GetExistenceRange() const { return ExistenceRange; }
+
+	static constexpr const TCHAR* TrackName = TEXT("SequenceInfoTrack");
+	
+private:
+	virtual bool UpdateInternal() override;
+	virtual TSharedPtr<SWidget> GetDetailsViewInternal() override { return DetailsView; }
+
+	virtual FSlateIcon GetIconInternal() override { return Icon; }
+	virtual FName GetNameInternal() const override { return TrackName; }
+	virtual FText GetDisplayNameInternal() const override;
+	virtual uint64 GetObjectIdInternal() const override { return ObjectId; }
+
+	void Initialize();
+	USequenceInfoDetailsObject* InitializeDetailsObject();
+
+	static void RefreshSequenceInfoFromEvaluationProgram(TArray<FAnimNextSequenceTraceInfo>& OutSequenceInfo, const FSerializableEvaluationProgram& Program);
+
+	TSharedPtr<IDetailsView> DetailsView;
+	FSlateIcon Icon;
+	uint64 ObjectId;
+	uint64 InstanceId = 0; 
+	double PreviousScrubTime = -1.0;
+	TWeakObjectPtr<USequenceInfoDetailsObject> DetailsObjectWeakPtr;
+	TSharedPtr<SEventTimelineView::FTimelineEventData> ExistenceRange;
+
+	// TArray<TSharedPtr<FEvaluationProgramTrack>> Children;
+};
+
+
+class FSequenceInfoTrackCreator : public RewindDebugger::IRewindDebuggerTrackCreator
+{
+private:
+	virtual FName GetTargetTypeNameInternal() const;
+	virtual FName GetNameInternal() const override { return FSequenceInfoTrack::TrackName; }
+	virtual void GetTrackTypesInternal(TArray<RewindDebugger::FRewindDebuggerTrackType>& Types) const override;
+	virtual TSharedPtr<RewindDebugger::FRewindDebuggerTrack> CreateTrackInternal(uint64 ObjectId) const override;
+	virtual bool HasDebugInfoInternal(uint64 ObjectId) const override;
+};
+
+}
\ No newline at end of file
