Linux buildfix in blob build

diff --git a/Engine/Plugins/Experimental/AnimNextAnimGraph/Source/AnimNextAnimGraphEditor/Private/RewindDebugger/SequenceInfoTrack.cpp b/Engine/Plugins/Experimental/AnimNextAnimGraph/Source/AnimNextAnimGraphEditor/Private/RewindDebugger/SequenceInfoTrack.cpp
--- a/Engine/Plugins/Experimental/AnimNextAnimGraph/Source/AnimNextAnimGraphEditor/Private/RewindDebugger/SequenceInfoTrack.cpp
+++ b/Engine/Plugins/Experimental/AnimNextAnimGraph/Source/AnimNextAnimGraphEditor/Private/RewindDebugger/SequenceInfoTrack.cpp
@@ -20,8 +20,6 @@
 
 namespace UE::AnimNextEditor
 {
-static const FName AnimNextModulesName("AnimNextModules");
-
 FName FSequenceInfoTrackCreator::GetTargetTypeNameInternal() const
 {
 	static const FName ObjectName("AnimNextComponent");
