Add Modify curve trait
Use modify curve trait on locomotion data to enable slope warping

diff --git a/Engine/Plugins/Experimental/AnimNextAnimGraph/Source/AnimNextAnimGraph/Private/AnimNextAnimGraphModule.cpp b/Engine/Plugins/Experimental/AnimNextAnimGraph/Source/AnimNextAnimGraph/Private/AnimNextAnimGraphModule.cpp
--- a/Engine/Plugins/Experimental/AnimNextAnimGraph/Source/AnimNextAnimGraph/Private/AnimNextAnimGraphModule.cpp
+++ b/Engine/Plugins/Experimental/AnimNextAnimGraph/Source/AnimNextAnimGraph/Private/AnimNextAnimGraphModule.cpp
@@ -17,6 +17,7 @@
 #include "TraitCore/NodeTemplateRegistry.h"
 #include "TraitCore/TraitInterfaceRegistry.h"
 #include "TraitCore/TraitRegistry.h"
+#include "Traits/ModifyCurveTrait.h"
 
 #if WITH_ANIMNEXT_CONSOLE_COMMANDS
 #include "HAL/IConsoleManager.h"
@@ -73,7 +74,8 @@ void FAnimNextAnimGraphModule::StartupModule()
 	{
 		FAnimNextGraphReferencePose::StaticStruct(),
 		FRigVMTrait_AnimNextPublicVariables::StaticStruct(),
-		FAnimNextInjectionBlendSettings::StaticStruct()
+		FAnimNextInjectionBlendSettings::StaticStruct(),
+		FModifyCurveParameters::StaticStruct()
 	};
 
 	RigVMRegistry.RegisterStructTypes(AllowedStructTypes);
diff --git a/Engine/Plugins/Experimental/AnimNextAnimGraph/Source/AnimNextAnimGraph/Private/Traits/ModifyCurveTrait.cpp b/Engine/Plugins/Experimental/AnimNextAnimGraph/Source/AnimNextAnimGraph/Private/Traits/ModifyCurveTrait.cpp
new file mode 100644
--- /dev/null
+++ b/Engine/Plugins/Experimental/AnimNextAnimGraph/Source/AnimNextAnimGraph/Private/Traits/ModifyCurveTrait.cpp
@@ -0,0 +1,82 @@
+#include "ModifyCurveTrait.h"
+
+#include "EvaluationVM/EvaluationVM.h"
+
+namespace UE::AnimNext
+{
+AUTO_REGISTER_ANIM_TRAIT(FModifyCurveTrait)
+
+#define TRAIT_INTERFACE_ENUMERATOR(GeneratorMacro) \
+GeneratorMacro(IEvaluate) \
+
+// Trait implementation boilerplate
+GENERATE_ANIM_TRAIT_IMPLEMENTATION(FModifyCurveTrait, TRAIT_INTERFACE_ENUMERATOR, NULL_ANIM_TRAIT_INTERFACE_ENUMERATOR, NULL_ANIM_TRAIT_EVENT_ENUMERATOR)
+#undef TRAIT_INTERFACE_ENUMERATOR
+}
+
+void UE::AnimNext::FModifyCurveTrait::PostEvaluate(FEvaluateTraversalContext& Context, const TTraitBinding<IEvaluate>& Binding) const
+{
+	IEvaluate::PostEvaluate(Context, Binding);
+
+	const FSharedData* SharedData = Binding.GetSharedData<FSharedData>();
+	check(SharedData);
+
+	FInstanceData* InstanceData = Binding.GetInstanceData<FInstanceData>();
+	check(InstanceData);
+
+	// Copy Alpha value since it is 'latent'
+	InstanceData->Alpha = SharedData->GetAlpha(Binding);
+
+#if ENABLE_ANIM_DEBUG 
+	InstanceData->HostObject = Context.GetHostObject();
+#endif // ENABLE_ANIM_DEBUG 
+
+	Context.AppendTask(FModifyCurveTask::Make(InstanceData, SharedData));
+}
+
+FModifyCurveTask FModifyCurveTask::Make(UE::AnimNext::FModifyCurveTrait::FInstanceData* InstanceData,
+	const UE::AnimNext::FModifyCurveTrait::FSharedData* SharedData)
+{
+	FModifyCurveTask Task;
+	Task.InstanceData = InstanceData;
+	Task.SharedData = SharedData;
+	return Task;
+}
+
+void FModifyCurveTask::Execute(UE::AnimNext::FEvaluationVM& VM) const
+{
+	using namespace UE::AnimNext;
+	
+	if (const TUniquePtr<FKeyframeState>* Keyframe = VM.PeekValue<TUniquePtr<FKeyframeState>>(KEYFRAME_STACK_NAME, 0))
+	{
+		for (const FModifyCurveParameters& CurveParameters : SharedData->ModifyCurveParameters)
+		{
+			float CurrentValue = Keyframe->Get()->Curves.Get(CurveParameters.CurveName);
+			float UpdatedValue = ProcessCurveOperation(CurrentValue, CurveParameters.CurveValue, InstanceData->Alpha, SharedData->ApplyMode);
+			Keyframe->Get()->Curves.Set(CurveParameters.CurveName, UpdatedValue);
+		}
+		
+	}
+}
+
+float FModifyCurveTask::ProcessCurveOperation(float CurrentValue, float NewValue, float Alpha, EAnimNext_ModifyCurveApplyMode ApplyMode)
+{
+	float UseNewValue = CurrentValue;
+
+	// Use ApplyMode enum to decide how to apply
+	if (ApplyMode == EAnimNext_ModifyCurveApplyMode::Add)
+	{
+		UseNewValue = CurrentValue + NewValue;
+	}
+	else if (ApplyMode == EAnimNext_ModifyCurveApplyMode::Scale)
+	{
+		UseNewValue = CurrentValue * NewValue;
+	}
+	else if (ApplyMode == EAnimNext_ModifyCurveApplyMode::Blend)
+	{
+		UseNewValue = NewValue;
+	}
+
+	const float UseAlpha = FMath::Clamp(Alpha, 0.f, 1.f);
+	return FMath::Lerp(CurrentValue, UseNewValue, UseAlpha);
+}
diff --git a/Engine/Plugins/Experimental/AnimNextAnimGraph/Source/AnimNextAnimGraph/Private/Traits/ModifyCurveTrait.h b/Engine/Plugins/Experimental/AnimNextAnimGraph/Source/AnimNextAnimGraph/Private/Traits/ModifyCurveTrait.h
new file mode 100644
--- /dev/null
+++ b/Engine/Plugins/Experimental/AnimNextAnimGraph/Source/AnimNextAnimGraph/Private/Traits/ModifyCurveTrait.h
@@ -0,0 +1,107 @@
+#pragma once
+
+#include "CoreMinimal.h"
+
+#include "TraitCore/Trait.h"
+#include "TraitInterfaces/IContinuousBlend.h"
+#include "TraitInterfaces/IEvaluate.h"
+#include "TraitInterfaces/IHierarchy.h"
+#include "TraitInterfaces/IUpdate.h"
+#include "TraitInterfaces/Args/IAlphaInputArgs.h"
+
+#include "ModifyCurveTrait.generated.h"
+
+UENUM()
+enum class EAnimNext_ModifyCurveApplyMode : uint8
+{
+	/** Add new value to input curve value */
+	Add,
+
+	/** Scale input value by new value */
+	Scale,
+
+	/** Blend input with new curve value, using Alpha setting on the node */
+	Blend,
+};
+
+USTRUCT()
+struct FModifyCurveParameters
+{
+	GENERATED_BODY()
+
+	UPROPERTY(Category=CurveData, EditAnywhere)
+	FName CurveName;
+
+	UPROPERTY(Category=CurveData, EditAnywhere)
+	float CurveValue;
+	
+};
+
+USTRUCT(meta = (DisplayName = "Modify Curve"))
+struct FAnimNextModifyCurveTraitSharedData : public FAnimNextTraitSharedData
+{
+	GENERATED_BODY()
+
+	UPROPERTY(Category=Settings, EditAnywhere)
+	EAnimNext_ModifyCurveApplyMode ApplyMode;
+
+	UPROPERTY(Category=Settings, EditAnywhere, meta=(Inline, Hidden))
+	TArray<FModifyCurveParameters> ModifyCurveParameters;
+
+	UPROPERTY(Category=Settings, EditAnywhere, meta=(ClampMin=0, ClampMax=1))
+	float Alpha = 1.0f;
+
+#define TRAIT_LATENT_PROPERTIES_ENUMERATOR(GeneratorMacro) \
+	GeneratorMacro(Alpha) \
+
+	GENERATE_TRAIT_LATENT_PROPERTIES(FAnimNextModifyCurveTraitSharedData, TRAIT_LATENT_PROPERTIES_ENUMERATOR)
+#undef TRAIT_LATENT_PROPERTIES_ENUMERATOR
+};
+
+namespace UE::AnimNext
+{
+
+	/**
+	 * FModifyCurveTrait
+	 * 
+	 * A trait that can write to a specific float curve value
+	 */
+	struct FModifyCurveTrait : FAdditiveTrait, IEvaluate
+	{
+		DECLARE_ANIM_TRAIT(FModifyCurveTrait, FAdditiveTrait)
+
+		using FSharedData = FAnimNextModifyCurveTraitSharedData;
+
+		struct FInstanceData : FTrait::FInstanceData
+		{
+			float Alpha = 1.0f;
+			
+#if ENABLE_ANIM_DEBUG 
+			/** Debug Object for VisualLogger */
+			TObjectPtr<const UObject> HostObject = nullptr;
+#endif // ENABLE_ANIM_DEBUG 
+		};
+
+		// IEvaluate impl
+		virtual void PostEvaluate(FEvaluateTraversalContext& Context, const TTraitBinding<IEvaluate>& Binding) const override;
+	};
+}
+
+USTRUCT()
+struct FModifyCurveTask : public FAnimNextEvaluationTask
+{
+	GENERATED_BODY()
+
+	DECLARE_ANIM_EVALUATION_TASK(FModifyCurveTask)
+
+	static FModifyCurveTask Make(UE::AnimNext::FModifyCurveTrait::FInstanceData* InstanceData, const UE::AnimNext::FModifyCurveTrait::FSharedData* SharedData);
+
+	// Task entry point
+	virtual void Execute(UE::AnimNext::FEvaluationVM& VM) const override;
+
+private:
+	static float ProcessCurveOperation(float CurrentValue, float NewValue, float Alpha, EAnimNext_ModifyCurveApplyMode ApplyMode);
+
+	UE::AnimNext::FModifyCurveTrait::FInstanceData* InstanceData = nullptr;
+	const UE::AnimNext::FModifyCurveTrait::FSharedData* SharedData = nullptr;
+};
