UAF : Register RigVMClasses as a dependency during serialization

#bs-asset-validation-passed: 1654583 v2_AAAAAGgF9IU=
#rnx

diff --git a/Engine/Plugins/Experimental/AnimNextAnimGraph/Source/AnimNextAnimGraph/Private/Graph/AnimNextAnimationGraph.cpp b/Engine/Plugins/Experimental/AnimNextAnimGraph/Source/AnimNextAnimGraph/Private/Graph/AnimNextAnimationGraph.cpp
--- a/Engine/Plugins/Experimental/AnimNextAnimGraph/Source/AnimNextAnimGraph/Private/Graph/AnimNextAnimationGraph.cpp
+++ b/Engine/Plugins/Experimental/AnimNextAnimGraph/Source/AnimNextAnimGraph/Private/Graph/AnimNextAnimationGraph.cpp
@@ -18,6 +18,7 @@
 #include "Module/AnimNextModuleInstance.h"
 #include "DataRegistryTypes.h"
 #include "Module/AnimNextSkeletalMeshComponentReferenceComponent.h"
+#include "RigVMBlueprintGeneratedClass.h"
 
 #if WITH_EDITOR	
 #include "Engine/ExternalAssetDependencyGatherer.h"
@@ -177,6 +178,23 @@ void UAnimNextAnimationGraph::Serialize(FArchive& Ar)
 	}
 }
 
+#if WITH_EDITORONLY_DATA
+void UAnimNextAnimationGraph::GetPreloadDependencies(TArray<UObject*>& OutDeps)
+{
+	Super::GetPreloadDependencies(OutDeps);
+
+	// Temp fix for Control Rig Trait getting called serialization before the class it depends has been loaded
+	// Ideally, we need a mechanism to allow traits to provide dependencies
+	for (TObjectPtr<UObject>& ReferencedObject : GraphReferencedObjects)
+	{
+		if (URigVMBlueprintGeneratedClass* ReferencedObjectPtr = Cast<URigVMBlueprintGeneratedClass>(ReferencedObject.Get()))
+		{
+			OutDeps.Add(ReferencedObjectPtr);
+		}
+	}
+}
+#endif
+
 bool UAnimNextAnimationGraph::LoadFromArchiveBuffer(const TArray<uint8>& InSharedDataArchiveBuffer)
 {
 	using namespace UE::AnimNext;
diff --git a/Engine/Plugins/Experimental/AnimNextAnimGraph/Source/AnimNextAnimGraph/Public/Graph/AnimNextAnimationGraph.h b/Engine/Plugins/Experimental/AnimNextAnimGraph/Source/AnimNextAnimGraph/Public/Graph/AnimNextAnimationGraph.h
--- a/Engine/Plugins/Experimental/AnimNextAnimGraph/Source/AnimNextAnimGraph/Public/Graph/AnimNextAnimationGraph.h
+++ b/Engine/Plugins/Experimental/AnimNextAnimGraph/Source/AnimNextAnimGraph/Public/Graph/AnimNextAnimationGraph.h
@@ -70,6 +70,9 @@ public:
 
 	// UObject interface
 	virtual void Serialize(FArchive& Ar) override;
+#if WITH_EDITORONLY_DATA
+	virtual void GetPreloadDependencies(TArray<UObject*>& OutDeps) override;
+#endif
 
 	// Allocates an instance of the graph with an optional module and/or parent graph instance
 	// @param	InModuleInstance			The module instance to use
