Avoid sampling root motion outside of the sequence duration

diff --git a/Engine/Plugins/Experimental/AnimNextAnimGraph/Source/AnimNextAnimGraph/Private/Traits/SequencePlayer.cpp b/Engine/Plugins/Experimental/AnimNextAnimGraph/Source/AnimNextAnimGraph/Private/Traits/SequencePlayer.cpp
--- a/Engine/Plugins/Experimental/AnimNextAnimGraph/Source/AnimNextAnimGraph/Private/Traits/SequencePlayer.cpp
+++ b/Engine/Plugins/Experimental/AnimNextAnimGraph/Source/AnimNextAnimGraph/Private/Traits/SequencePlayer.cpp
@@ -62,11 +62,25 @@ namespace UE::AnimNext
 
 		if (InstanceData->AnimSequence)
 		{
-			auto ExtractRootMotionAttribute = [AnimSequence = InstanceData->AnimSequence](float StartTime, float DeltaTime, bool bAllowLooping)
+			TTraitBinding<ITimeline> TimelineTrait;
+			Binding.GetStackInterface(TimelineTrait);
+
+			const FTimelineState State = TimelineTrait.GetState(Context);
+			const float PlayRate = State.GetPlayRate();
+			const bool bIsLooping = State.IsLooping();
+
+			auto ExtractRootMotionAttribute = [AnimSequence = InstanceData->AnimSequence, PlayRate, bIsLooping](float StartTime, float DeltaTime, bool bAllowLooping)
 			{
-				// We do not check for lifetimes, assume the sequence is alive during pose list exectuion.
+				// We do not check for lifetimes, assume the sequence is alive during pose list execution.
 				check(AnimSequence->IsValidLowLevel());
-				return AnimSequence->ExtractRootMotion(FAnimExtractContext(static_cast<double>(StartTime), true, FDeltaTimeRecord(DeltaTime), bAllowLooping && AnimSequence->bLoop));
+
+				// Make sure we don't extract outside our valid range
+				StartTime = FMath::Clamp(StartTime, 0.0f, AnimSequence->GetPlayLength());
+
+				// Account for play rate
+				DeltaTime *= PlayRate;
+
+				return AnimSequence->ExtractRootMotion(FAnimExtractContext(static_cast<double>(StartTime), true, FDeltaTimeRecord(DeltaTime), bAllowLooping && bIsLooping && AnimSequence->bLoop));
 			};
 
 			return FOnExtractRootMotionAttribute::CreateLambda(ExtractRootMotionAttribute);
