Account for AnimSequence RateScale in anim next sequence playback

diff --git a/Engine/Plugins/Experimental/AnimNextAnimGraph/Source/AnimNextAnimGraph/Private/Traits/SequencePlayer.cpp b/Engine/Plugins/Experimental/AnimNextAnimGraph/Source/AnimNextAnimGraph/Private/Traits/SequencePlayer.cpp
--- a/Engine/Plugins/Experimental/AnimNextAnimGraph/Source/AnimNextAnimGraph/Private/Traits/SequencePlayer.cpp
+++ b/Engine/Plugins/Experimental/AnimNextAnimGraph/Source/AnimNextAnimGraph/Private/Traits/SequencePlayer.cpp
@@ -113,7 +113,7 @@ namespace UE::AnimNext
 			return FTimelineState(
 				InstanceData->InternalTimeAccumulator,
 				AnimSeq->GetPlayLength(),
-				SharedData->GetPlayRate(Binding),
+				SharedData->GetPlayRate(Binding) * AnimSeq->RateScale,
 				SharedData->GetbLoop(Binding)).WithDebugName(AnimSeq->GetFName());
 		}
 
