Avoid purging children with near zero weight to handle zero delta time updates

diff --git a/Engine/Plugins/Experimental/AnimNextAnimGraph/Source/AnimNextAnimGraph/Private/Traits/BlendSmoother.cpp b/Engine/Plugins/Experimental/AnimNextAnimGraph/Source/AnimNextAnimGraph/Private/Traits/BlendSmoother.cpp
--- a/Engine/Plugins/Experimental/AnimNextAnimGraph/Source/AnimNextAnimGraph/Private/Traits/BlendSmoother.cpp
+++ b/Engine/Plugins/Experimental/AnimNextAnimGraph/Source/AnimNextAnimGraph/Private/Traits/BlendSmoother.cpp
@@ -120,14 +120,7 @@ namespace UE::AnimNext
 
 			ChildBlendData.Blend.Update(DeltaTime);
 
-			float NewBlendWeight = ChildBlendData.Blend.GetBlendedValue();
-
-			if (!FAnimWeight::IsRelevant(NewBlendWeight))
-			{
-				// Our new weight is no longer relevant, snap it to zero and normalization below will fix-up the other weights
-				// We'll then terminate the blend below
-				NewBlendWeight = 0.0f;
-			}
+			const float NewBlendWeight = ChildBlendData.Blend.GetBlendedValue();
 
 			ChildBlendData.Weight = NewBlendWeight;
 			SumWeight += NewBlendWeight;
diff --git a/Engine/Plugins/Experimental/AnimNextAnimGraph/Source/AnimNextAnimGraph/Private/Traits/BlendStackTrait.cpp b/Engine/Plugins/Experimental/AnimNextAnimGraph/Source/AnimNextAnimGraph/Private/Traits/BlendStackTrait.cpp
--- a/Engine/Plugins/Experimental/AnimNextAnimGraph/Source/AnimNextAnimGraph/Private/Traits/BlendStackTrait.cpp
+++ b/Engine/Plugins/Experimental/AnimNextAnimGraph/Source/AnimNextAnimGraph/Private/Traits/BlendStackTrait.cpp
@@ -239,6 +239,10 @@ namespace UE::AnimNext
 		FInstanceData* InstanceData = Binding.GetInstanceData<FInstanceData>();
 		if (InstanceData->ChildGraphs.IsValidIndex(ChildIndex))
 		{
+			// We should never terminate the blend of the currently active graph index as it would mean we no
+			// longer have an active pose to produce
+			check(InstanceData->CurrentlyActiveGraphIndex != ChildIndex);
+
 			FGraphState& Graph = InstanceData->ChildGraphs[ChildIndex];
 
 			// Deallocate our graph
