Add toggle to blend layer trait to blend in mesh space

diff --git a/Engine/Plugins/Experimental/AnimNextAnimGraph/Source/AnimNextAnimGraph/Private/EvaluationVM/Tasks/ConvertRotationSpace.ispc b/Engine/Plugins/Experimental/AnimNextAnimGraph/Source/AnimNextAnimGraph/Private/EvaluationVM/Tasks/ConvertRotationSpace.ispc
new file mode 100644
--- /dev/null
+++ b/Engine/Plugins/Experimental/AnimNextAnimGraph/Source/AnimNextAnimGraph/Private/EvaluationVM/Tasks/ConvertRotationSpace.ispc
@@ -0,0 +1,29 @@
+// Copyright Epic Games, Inc. All Rights Reserved.
+
+#include "Math/Quat.isph"
+
+export void ConvertRotationsToMeshSpace(uniform FVector4 LocalPose[], const uniform unsigned int16 ParentBones[], const uniform int NumBones)
+{
+	for (uniform int BoneIndex = 1; BoneIndex < NumBones; ++BoneIndex)
+	{
+		const uniform unsigned int16 ParentIndex = ParentBones[BoneIndex];
+		const uniform FVector4 LocalPoseRotation = LocalPose[BoneIndex];
+		const uniform FVector4 LocalParentRotation = LocalPose[ParentIndex];
+
+		const uniform FVector4 MeshSpaceRotation = VectorQuaternionMultiply2(LocalParentRotation, LocalPoseRotation);
+		LocalPose[BoneIndex] = MeshSpaceRotation;
+	}
+}
+
+export void ConvertRotationsToLocalSpace(uniform FVector4 MeshPose[], const uniform unsigned int16 ParentBones[], const uniform int NumBones)
+{
+	for (uniform int BoneIndex = NumBones - 1; BoneIndex > 0; --BoneIndex)
+	{
+		const uniform unsigned int16 ParentIndex = ParentBones[BoneIndex];
+		const uniform FVector4 PoseRotation = MeshPose[BoneIndex];
+		const uniform FVector4 ParentRotation = MeshPose[ParentIndex];
+
+		const uniform FVector4 LocalSpaceRotation = VectorQuaternionMultiply2(QuatInverse(ParentRotation), PoseRotation);
+		MeshPose[BoneIndex] = LocalSpaceRotation;
+	}
+}
\ No newline at end of file
diff --git a/Engine/Plugins/Experimental/AnimNextAnimGraph/Source/AnimNextAnimGraph/Private/EvaluationVM/Tasks/ConvertRotationSpaceDefs.h b/Engine/Plugins/Experimental/AnimNextAnimGraph/Source/AnimNextAnimGraph/Private/EvaluationVM/Tasks/ConvertRotationSpaceDefs.h
new file mode 100644
--- /dev/null
+++ b/Engine/Plugins/Experimental/AnimNextAnimGraph/Source/AnimNextAnimGraph/Private/EvaluationVM/Tasks/ConvertRotationSpaceDefs.h
@@ -0,0 +1,7 @@
+// Copyright Epic Games, Inc. All Rights Reserved.
+
+#pragma once
+
+#ifndef UE_ANIM_CONVERT_ROTATION_SPACE_TASKS_ISPC
+#define UE_ANIM_CONVERT_ROTATION_SPACE_TASKS_ISPC INTEL_ISPC
+#endif
\ No newline at end of file
diff --git a/Engine/Plugins/Experimental/AnimNextAnimGraph/Source/AnimNextAnimGraph/Private/EvaluationVM/Tasks/ConvertRotationsToLocalSpace.cpp b/Engine/Plugins/Experimental/AnimNextAnimGraph/Source/AnimNextAnimGraph/Private/EvaluationVM/Tasks/ConvertRotationsToLocalSpace.cpp
new file mode 100644
--- /dev/null
+++ b/Engine/Plugins/Experimental/AnimNextAnimGraph/Source/AnimNextAnimGraph/Private/EvaluationVM/Tasks/ConvertRotationsToLocalSpace.cpp
@@ -0,0 +1,58 @@
+// Copyright Epic Games, Inc. All Rights Reserved.
+
+#include "EvaluationVM/Tasks/ConvertRotationsToLocalSpace.h"
+
+#include "EvaluationVM/EvaluationVM.h"
+#include "EvaluationVM/Tasks/ConvertRotationSpaceDefs.h"
+
+#if UE_ANIM_CONVERT_ROTATION_SPACE_TASKS_ISPC
+#include "ConvertRotationSpace.ispc.generated.h"
+#endif
+
+FAnimNextConvertRotationsToLocalSpaceTask FAnimNextConvertRotationsToLocalSpaceTask::Make(const int32 NumPoses)
+{
+	FAnimNextConvertRotationsToLocalSpaceTask Task;
+	Task.NumPoses = NumPoses;
+
+	return Task;
+}
+
+void FAnimNextConvertRotationsToLocalSpaceTask::Execute(UE::AnimNext::FEvaluationVM& VM) const
+{
+	using namespace UE::AnimNext;
+
+	if (!EnumHasAnyFlags(VM.GetFlags(), EEvaluationFlags::Bones))
+	{
+		return;
+	}
+
+	for (int32 PoseIndex = 0; PoseIndex < NumPoses; ++PoseIndex)
+	{
+		TUniquePtr<FKeyframeState>* Keyframe = VM.PeekValueMutable<TUniquePtr<FKeyframeState>>(KEYFRAME_STACK_NAME, PoseIndex);
+		if (!Keyframe || !Keyframe->IsValid())
+		{
+			// No pose to convert
+			return;
+		}
+
+		FLODPoseStack& Pose = (*Keyframe)->Pose;
+		TArrayView<FQuat> PoseRotations = Pose.LocalTransforms.Rotations;
+
+#if UE_ANIM_CONVERT_ROTATION_SPACE_TASKS_ISPC
+		static_assert(sizeof(FBoneIndexType) == sizeof(uint16)); // ISPC needs updating if assert fails
+
+		const TArrayView<const FBoneIndexType> ParentIndices = Pose.GetRefPose().GetLODBoneIndexToParentLODBoneIndexMap(Pose.LODLevel);
+
+		ispc::ConvertRotationsToLocalSpace(
+			reinterpret_cast<ispc::FVector4*>(PoseRotations.GetData()),
+			reinterpret_cast<const uint16*>(ParentIndices.GetData()),
+			Pose.GetNumBones());
+#else
+		for (FBoneIndexType Index = PoseRotations.Num() - 1; Index > 0; --Index)
+		{
+			const FBoneIndexType ParentIndex = Pose.GetLODBoneParentIndex(Index);
+			PoseRotations[Index] = PoseRotations[ParentIndex].Inverse() * PoseRotations[Index];
+		}
+#endif // UE_ANIM_CONVERT_ROTATION_SPACE_TASKS_ISPC
+	}
+}
\ No newline at end of file
diff --git a/Engine/Plugins/Experimental/AnimNextAnimGraph/Source/AnimNextAnimGraph/Private/EvaluationVM/Tasks/ConvertRotationsToMeshSpace.cpp b/Engine/Plugins/Experimental/AnimNextAnimGraph/Source/AnimNextAnimGraph/Private/EvaluationVM/Tasks/ConvertRotationsToMeshSpace.cpp
new file mode 100644
--- /dev/null
+++ b/Engine/Plugins/Experimental/AnimNextAnimGraph/Source/AnimNextAnimGraph/Private/EvaluationVM/Tasks/ConvertRotationsToMeshSpace.cpp
@@ -0,0 +1,58 @@
+// Copyright Epic Games, Inc. All Rights Reserved.
+
+#include "EvaluationVM/Tasks/ConvertRotationsToMeshSpace.h"
+
+#include "EvaluationVM/EvaluationVM.h"
+#include "EvaluationVM/Tasks/ConvertRotationSpaceDefs.h"
+
+#if UE_ANIM_CONVERT_ROTATION_SPACE_TASKS_ISPC
+#include "ConvertRotationSpace.ispc.generated.h"
+#endif
+
+FAnimNextConvertRotationsToMeshSpaceTask FAnimNextConvertRotationsToMeshSpaceTask::Make(const int32 NumPoses)
+{
+	FAnimNextConvertRotationsToMeshSpaceTask Task;
+	Task.NumPoses = NumPoses;
+
+	return Task;
+}
+
+void FAnimNextConvertRotationsToMeshSpaceTask::Execute(UE::AnimNext::FEvaluationVM& VM) const
+{
+	using namespace UE::AnimNext;
+
+	if (!EnumHasAnyFlags(VM.GetFlags(), EEvaluationFlags::Bones))
+	{
+		return;
+	}
+
+	for (int32 PoseIndex = 0; PoseIndex < NumPoses; ++PoseIndex)
+	{
+		TUniquePtr<FKeyframeState>* Keyframe = VM.PeekValueMutable<TUniquePtr<FKeyframeState>>(KEYFRAME_STACK_NAME, PoseIndex);
+		if (!Keyframe || !Keyframe->IsValid())
+		{
+			// No pose to convert
+			return;
+		}
+
+		FLODPoseStack& Pose = (*Keyframe)->Pose;
+		TArrayView<FQuat> PoseRotations = Pose.LocalTransforms.Rotations;
+
+#if UE_ANIM_CONVERT_ROTATION_SPACE_TASKS_ISPC
+		static_assert(sizeof(FBoneIndexType) == sizeof(uint16)); // ISPC needs updating if assert fails
+
+		const TArrayView<const FBoneIndexType> ParentIndices = Pose.GetRefPose().GetLODBoneIndexToParentLODBoneIndexMap(Pose.LODLevel);
+
+		ispc::ConvertRotationsToMeshSpace(
+			reinterpret_cast<ispc::FVector4*>(PoseRotations.GetData()),
+			reinterpret_cast<const uint16*>(ParentIndices.GetData()),
+			Pose.GetNumBones());
+#else
+		for (FBoneIndexType Index = 1; Index < PoseRotations.Num(); ++Index)
+		{
+			const FBoneIndexType ParentIndex = Pose.GetLODBoneParentIndex(Index);
+			PoseRotations[Index] = PoseRotations[ParentIndex] * PoseRotations[Index];
+		}
+#endif // UE_ANIM_CONVERT_ROTATION_SPACE_TASKS_ISPC
+	}
+}
\ No newline at end of file
diff --git a/Engine/Plugins/Experimental/AnimNextAnimGraph/Source/AnimNextAnimGraph/Private/Traits/BlendLayer.cpp b/Engine/Plugins/Experimental/AnimNextAnimGraph/Source/AnimNextAnimGraph/Private/Traits/BlendLayer.cpp
--- a/Engine/Plugins/Experimental/AnimNextAnimGraph/Source/AnimNextAnimGraph/Private/Traits/BlendLayer.cpp
+++ b/Engine/Plugins/Experimental/AnimNextAnimGraph/Source/AnimNextAnimGraph/Private/Traits/BlendLayer.cpp
@@ -7,6 +7,8 @@
 #include "TraitCore/ExecutionContext.h"
 #include "EvaluationVM/Tasks/BlendKeyframes.h"
 #include "EvaluationVM/Tasks/NormalizeRotations.h"
+#include "EvaluationVM/Tasks/ConvertRotationsToLocalSpace.h"
+#include "EvaluationVM/Tasks/ConvertRotationsToMeshSpace.h"
 #include "HierarchyTable.h"
 #include "MaskProfile/HierarchyTableTypeMask.h"
 #include "SkeletonHierarchyTableType.h"
@@ -77,9 +79,20 @@ namespace UE::AnimNext
 			const TObjectPtr<USkeleton> Skeleton = TableMetadata.Skeleton;
 
 			const float BlendWeight = ContinuousBlendTrait.GetBlendWeight(Context, 1);
+			const bool bBlendInMeshSpace = SharedData->GetbBlendInMeshSpace(Binding);
+
+			if (bBlendInMeshSpace)
+			{
+				Context.AppendTask(FAnimNextConvertRotationsToMeshSpaceTask::Make(2));
+			}
 
 			Context.AppendTask(FAnimNextBlendKeyframePerBoneWithScaleTask::Make(InterpolationIndexProvider, Skeleton, InstanceData->BoneMaskWeights, InstanceData->CurveMaskWeights, InstanceData->AttributeMaskWeights, BlendWeight));
 			Context.AppendTask(FAnimNextNormalizeKeyframeRotationsTask());
+
+			if (bBlendInMeshSpace)
+			{
+				Context.AppendTask(FAnimNextConvertRotationsToLocalSpaceTask::Make(1));
+			}
 		}
 		else
 		{
diff --git a/Engine/Plugins/Experimental/AnimNextAnimGraph/Source/AnimNextAnimGraph/Private/Traits/BlendLayer.h b/Engine/Plugins/Experimental/AnimNextAnimGraph/Source/AnimNextAnimGraph/Private/Traits/BlendLayer.h
--- a/Engine/Plugins/Experimental/AnimNextAnimGraph/Source/AnimNextAnimGraph/Private/Traits/BlendLayer.h
+++ b/Engine/Plugins/Experimental/AnimNextAnimGraph/Source/AnimNextAnimGraph/Private/Traits/BlendLayer.h
@@ -25,6 +25,7 @@ struct FAnimNextBlendLayerTraitSharedData : public FAnimNextTraitSharedData
 	UPROPERTY()
 	FAnimNextTraitHandle ChildBlend;
 
+	// The strength with which to apply the blend pose
 	UPROPERTY(EditAnywhere, Category = "Default")
 	float BlendWeight = 1.0f;
 
@@ -33,8 +34,13 @@ struct FAnimNextBlendLayerTraitSharedData : public FAnimNextTraitSharedData
 	UPROPERTY(EditAnywhere, Category = "Default", meta = (Inline))
 	TObjectPtr<UHierarchyTable> BlendProfile = nullptr;
 
+	// Whether or not to blend in mesh space or local space
+	UPROPERTY(EditAnywhere, Category = "Default")
+	bool bBlendInMeshSpace = false;
+
 	#define TRAIT_LATENT_PROPERTIES_ENUMERATOR(GeneratorMacro) \
 		GeneratorMacro(BlendWeight) \
+		GeneratorMacro(bBlendInMeshSpace) \
 
 	GENERATE_TRAIT_LATENT_PROPERTIES(FAnimNextBlendLayerTraitSharedData, TRAIT_LATENT_PROPERTIES_ENUMERATOR)
 	#undef TRAIT_LATENT_PROPERTIES_ENUMERATOR
diff --git a/Engine/Plugins/Experimental/AnimNextAnimGraph/Source/AnimNextAnimGraph/Public/EvaluationVM/Tasks/ConvertRotationsToLocalSpace.h b/Engine/Plugins/Experimental/AnimNextAnimGraph/Source/AnimNextAnimGraph/Public/EvaluationVM/Tasks/ConvertRotationsToLocalSpace.h
new file mode 100644
--- /dev/null
+++ b/Engine/Plugins/Experimental/AnimNextAnimGraph/Source/AnimNextAnimGraph/Public/EvaluationVM/Tasks/ConvertRotationsToLocalSpace.h
@@ -0,0 +1,25 @@
+// Copyright Epic Games, Inc. All Rights Reserved.
+
+#pragma once
+
+#include "CoreMinimal.h"
+#include "EvaluationVM/EvaluationTask.h"
+
+#include "ConvertRotationsToLocalSpace.generated.h"
+
+USTRUCT()
+struct ANIMNEXTANIMGRAPH_API FAnimNextConvertRotationsToLocalSpaceTask : public FAnimNextEvaluationTask
+{
+	GENERATED_BODY()
+
+	DECLARE_ANIM_EVALUATION_TASK(FAnimNextConvertRotationsToLocalSpaceTask)
+
+	static FAnimNextConvertRotationsToLocalSpaceTask Make(const int32 NumPoses);
+
+	virtual void Execute(UE::AnimNext::FEvaluationVM& VM) const override;
+
+private:
+	// The number of poses on the stack to convert
+	UPROPERTY(VisibleAnywhere, Category = Properties)
+	int32 NumPoses = 0;
+};
\ No newline at end of file
diff --git a/Engine/Plugins/Experimental/AnimNextAnimGraph/Source/AnimNextAnimGraph/Public/EvaluationVM/Tasks/ConvertRotationsToMeshSpace.h b/Engine/Plugins/Experimental/AnimNextAnimGraph/Source/AnimNextAnimGraph/Public/EvaluationVM/Tasks/ConvertRotationsToMeshSpace.h
new file mode 100644
--- /dev/null
+++ b/Engine/Plugins/Experimental/AnimNextAnimGraph/Source/AnimNextAnimGraph/Public/EvaluationVM/Tasks/ConvertRotationsToMeshSpace.h
@@ -0,0 +1,25 @@
+// Copyright Epic Games, Inc. All Rights Reserved.
+
+#pragma once
+
+#include "CoreMinimal.h"
+#include "EvaluationVM/EvaluationTask.h"
+
+#include "ConvertRotationsToMeshSpace.generated.h"
+
+USTRUCT()
+struct ANIMNEXTANIMGRAPH_API FAnimNextConvertRotationsToMeshSpaceTask : public FAnimNextEvaluationTask
+{
+	GENERATED_BODY()
+
+	DECLARE_ANIM_EVALUATION_TASK(FAnimNextConvertRotationsToMeshSpaceTask)
+
+	static FAnimNextConvertRotationsToMeshSpaceTask Make(const int32 NumPoses);
+
+	virtual void Execute(UE::AnimNext::FEvaluationVM& VM) const override;
+
+private:
+	// The number of poses on the stack to convert
+	UPROPERTY(VisibleAnywhere, Category = Properties)
+	int32 NumPoses = 0;
+};
\ No newline at end of file
