AnimNext: Fix user-after-free in node template registry

Change template buffer to use a TPagedArray & added an FRWLock to ensure we dont race with the ALT when async loading assets.

diff --git a/Engine/Plugins/Experimental/AnimNextAnimGraph/Source/AnimNextAnimGraph/Internal/TraitCore/NodeTemplateRegistry.h b/Engine/Plugins/Experimental/AnimNextAnimGraph/Source/AnimNextAnimGraph/Internal/TraitCore/NodeTemplateRegistry.h
--- a/Engine/Plugins/Experimental/AnimNextAnimGraph/Source/AnimNextAnimGraph/Internal/TraitCore/NodeTemplateRegistry.h
+++ b/Engine/Plugins/Experimental/AnimNextAnimGraph/Source/AnimNextAnimGraph/Internal/TraitCore/NodeTemplateRegistry.h
@@ -4,6 +4,7 @@
 
 #include "CoreMinimal.h"
 #include "TraitCore/NodeTemplateRegistryHandle.h"
+#include "Containers/PagedArray.h"
 
 #define UE_API ANIMNEXTANIMGRAPH_API
 
@@ -46,9 +47,9 @@ namespace UE::AnimNext
 	private:
 		FNodeTemplateRegistry() = default;
 		FNodeTemplateRegistry(const FNodeTemplateRegistry&) = delete;
-		FNodeTemplateRegistry(FNodeTemplateRegistry&&) = default;
+		FNodeTemplateRegistry(FNodeTemplateRegistry&&) = delete;
 		FNodeTemplateRegistry& operator=(const FNodeTemplateRegistry&) = delete;
-		FNodeTemplateRegistry& operator=(FNodeTemplateRegistry&&) = default;
+		FNodeTemplateRegistry& operator=(FNodeTemplateRegistry&&) = delete;
 
 		// Finds and returns a node template based on its handle or nullptr if the handle is invalid
 		UE_API FNodeTemplate* FindMutable(FNodeTemplateRegistryHandle TemplateHandle);
@@ -64,8 +65,9 @@ namespace UE::AnimNext
 		// from this buffer when there are holes and coalesce everything.
 		// To do so, we would have to fix-up any outstanding handles within
 		// the shared data of loaded anim graphs.
-		TArray<uint8>								TemplateBuffer;
+		TPagedArray<uint8>							TemplateBuffer;
 		TMap<uint32, FNodeTemplateRegistryHandle>	TemplateUIDToHandleMap;
+		mutable FRWLock								TemplateLock;
 
 		friend class UE::AnimNext::AnimGraph::FAnimNextAnimGraphModule;
 		friend class FTraitWriter;
diff --git a/Engine/Plugins/Experimental/AnimNextAnimGraph/Source/AnimNextAnimGraph/Private/TraitCore/NodeTemplateRegistry.cpp b/Engine/Plugins/Experimental/AnimNextAnimGraph/Source/AnimNextAnimGraph/Private/TraitCore/NodeTemplateRegistry.cpp
--- a/Engine/Plugins/Experimental/AnimNextAnimGraph/Source/AnimNextAnimGraph/Private/TraitCore/NodeTemplateRegistry.cpp
+++ b/Engine/Plugins/Experimental/AnimNextAnimGraph/Source/AnimNextAnimGraph/Private/TraitCore/NodeTemplateRegistry.cpp
@@ -4,6 +4,7 @@
 
 #include "TraitCore/TraitTemplate.h"
 #include "TraitCore/NodeTemplate.h"
+#include "Misc/ScopeRWLock.h"
 
 namespace UE::AnimNext
 {
@@ -37,6 +38,7 @@ namespace UE::AnimNext
 
 	FNodeTemplateRegistryHandle FNodeTemplateRegistry::Find(uint32 NodeTemplateUID) const
 	{
+		FReadScopeLock Lock(TemplateLock);
 		const FNodeTemplateRegistryHandle* TemplateHandle = TemplateUIDToHandleMap.Find(NodeTemplateUID);
 		return TemplateHandle != nullptr ? *TemplateHandle : FNodeTemplateRegistryHandle();
 	}
@@ -50,24 +52,38 @@ namespace UE::AnimNext
 
 		const uint32 TemplateUID = NodeTemplate->GetUID();
 
-		if (FNodeTemplateRegistryHandle* TemplateHandle = TemplateUIDToHandleMap.Find(TemplateUID))
 		{
-			// We have already seen this template before, return its handle
-			return *TemplateHandle;
+			FReadScopeLock Lock(TemplateLock);
+			if (FNodeTemplateRegistryHandle* TemplateHandle = TemplateUIDToHandleMap.Find(TemplateUID))
+			{
+				// We have already seen this template before, return its handle
+				return *TemplateHandle;
+			}
 		}
 
-		// This is a new template
-		const uint32 TemplateSize = NodeTemplate->GetNodeTemplateSize();
+		{
+			FWriteScopeLock Lock(TemplateLock);
+
+			// Re-check the map in case we added during lock upgrade
+			if (FNodeTemplateRegistryHandle* TemplateHandle = TemplateUIDToHandleMap.Find(TemplateUID))
+			{
+				// We have already seen this template before, return its handle
+				return *TemplateHandle;
+			}
+
+			// This is a new template
+			const uint32 TemplateSize = NodeTemplate->GetNodeTemplateSize();
 
-		// Allocate and copy the template in our global buffer
-		const uint32 TemplateOffset = TemplateBuffer.AddUninitialized(TemplateSize);
-		FMemory::Memcpy(&TemplateBuffer[TemplateOffset], NodeTemplate, TemplateSize);
+			// Copy the template into our global buffer
+			TemplateBuffer.Append((uint8*)NodeTemplate, TemplateSize);
+			const uint32 TemplateOffset = TemplateBuffer.Num() - TemplateSize;
 
-		// Update our mapping
-		FNodeTemplateRegistryHandle TemplateHandle = FNodeTemplateRegistryHandle::MakeHandle(TemplateOffset);
-		TemplateUIDToHandleMap.Add(TemplateUID, TemplateHandle);
+			// Update our mapping
+			FNodeTemplateRegistryHandle TemplateHandle = FNodeTemplateRegistryHandle::MakeHandle(TemplateOffset);
+			TemplateUIDToHandleMap.Add(TemplateUID, TemplateHandle);
 
-		return TemplateHandle;
+			return TemplateHandle;
+		}
 	}
 
 	void FNodeTemplateRegistry::Unregister(const FNodeTemplate* NodeTemplate)
@@ -78,6 +94,8 @@ namespace UE::AnimNext
 		}
 
 		const uint32 TemplateUID = NodeTemplate->GetUID();
+
+		FWriteScopeLock Lock(TemplateLock);
 		TemplateUIDToHandleMap.Remove(TemplateUID);
 	}
 
@@ -88,6 +106,7 @@ namespace UE::AnimNext
 			return nullptr;
 		}
 
+		FReadScopeLock Lock(TemplateLock);
 		return reinterpret_cast<FNodeTemplate*>(&TemplateBuffer[TemplateHandle.GetTemplateOffset()]);
 	}
 
@@ -98,11 +117,13 @@ namespace UE::AnimNext
 			return nullptr;
 		}
 
+		FReadScopeLock Lock(TemplateLock);
 		return reinterpret_cast<const FNodeTemplate*>(&TemplateBuffer[TemplateHandle.GetTemplateOffset()]);
 	}
 
 	uint32 FNodeTemplateRegistry::GetNum() const
 	{
+		FReadScopeLock Lock(TemplateLock);
 		return TemplateUIDToHandleMap.Num();
 	}
 }
