Notify filtering trait:
- Reusing ScopedTag trait from AnimNext testing code (moved to TraitInterfaces)
- Simple IUpdate/IScopedTagTrait to insert DisableNotifies FName tag into scope
- Early out of NotifyDispatcher if tag is in scope
- RigUnit to check whether or not FAnimNextGraphInjectionData currently has an EvaluationModifier set/injected

diff --git a/Engine/Plugins/Experimental/AnimNextAnimGraph/Source/AnimNextAnimGraph/Private/Graph/RigUnit_HasValidEvaluationModifier.cpp b/Engine/Plugins/Experimental/AnimNextAnimGraph/Source/AnimNextAnimGraph/Private/Graph/RigUnit_HasValidEvaluationModifier.cpp
new file mode 100644
--- /dev/null
+++ b/Engine/Plugins/Experimental/AnimNextAnimGraph/Source/AnimNextAnimGraph/Private/Graph/RigUnit_HasValidEvaluationModifier.cpp
@@ -0,0 +1,8 @@
+// Copyright Epic Games, Inc. All Rights Reserved.
+
+#include "RigUnit_HasValidEvaluationModifier.h"
+
+FRigUnit_HasValidEvaluationModifier_Execute()
+{
+	bValue = InjectionData.HasValidEvaluationModifier();
+}
diff --git a/Engine/Plugins/Experimental/AnimNextAnimGraph/Source/AnimNextAnimGraph/Private/Graph/RigUnit_HasValidEvaluationModifier.h b/Engine/Plugins/Experimental/AnimNextAnimGraph/Source/AnimNextAnimGraph/Private/Graph/RigUnit_HasValidEvaluationModifier.h
new file mode 100644
--- /dev/null
+++ b/Engine/Plugins/Experimental/AnimNextAnimGraph/Source/AnimNextAnimGraph/Private/Graph/RigUnit_HasValidEvaluationModifier.h
@@ -0,0 +1,26 @@
+// Copyright Epic Games, Inc. All Rights Reserved.
+
+#pragma once
+
+#include "AnimNextExecuteContext.h"
+#include "Graph/AnimNextAnimGraph.h"
+#include "Graph/RigUnit_AnimNextBase.h"
+
+#include "RigUnit_HasValidEvaluationModifier.generated.h"
+
+USTRUCT(meta=(DisplayName="Has Valid Evaluation Modifier", Category="Animation Next Injection", NodeColor="0, 1, 1"))
+struct FRigUnit_HasValidEvaluationModifier : public FRigUnit_AnimNextBase
+{
+	GENERATED_BODY()
+
+	RIGVM_METHOD()
+	void Execute();
+
+	// The Injection Data instance to check
+	UPROPERTY(EditAnywhere, Category = "Struct", meta = (Input))
+	FAnimNextGraphInjectionData InjectionData;
+
+	// Whether Evaluation Modifier is valid (indicates active injection)
+	UPROPERTY(EditAnywhere, Category = "Struct", meta = (Output))
+	bool bValue;
+};
diff --git a/Engine/Plugins/Experimental/AnimNextAnimGraph/Source/AnimNextAnimGraph/Private/Traits/NotifyDispatcher.cpp b/Engine/Plugins/Experimental/AnimNextAnimGraph/Source/AnimNextAnimGraph/Private/Traits/NotifyDispatcher.cpp
--- a/Engine/Plugins/Experimental/AnimNextAnimGraph/Source/AnimNextAnimGraph/Private/Traits/NotifyDispatcher.cpp
+++ b/Engine/Plugins/Experimental/AnimNextAnimGraph/Source/AnimNextAnimGraph/Private/Traits/NotifyDispatcher.cpp
@@ -4,6 +4,7 @@
 #include "TraitInterfaces/INotifySource.h"
 #include "TraitInterfaces/ITimeline.h"
 #include "AnimationRuntime.h"
+#include "NotifyFilterTrait.h"
 #include "Animation/AnimTrace.h"
 #include "Animation/AnimNotifies/AnimNotify.h"
 #include "Animation/AnimNotifies/AnimNotifyState.h"
@@ -47,6 +48,11 @@ namespace UE::AnimNext
 			return;
 		}
 
+		if (!FNotifyFilterTrait::AreNotifiesEnabledInScope(Context))
+		{
+			return;
+		}
+
 		TTraitBinding<INotifySource> NotifySourceTrait;
 		ensure(Binding.GetStackInterfaceSuper(NotifySourceTrait));
 
@@ -309,4 +315,4 @@ void FAnimNextNotifyDispatcherComponent::TriggerSingleAnimNotify(float InDeltaTi
 			});
 		}
 	}
-}
\ No newline at end of file
+}
diff --git a/Engine/Plugins/Experimental/AnimNextAnimGraph/Source/AnimNextAnimGraph/Private/Traits/NotifyFilterTrait.cpp b/Engine/Plugins/Experimental/AnimNextAnimGraph/Source/AnimNextAnimGraph/Private/Traits/NotifyFilterTrait.cpp
new file mode 100644
--- /dev/null
+++ b/Engine/Plugins/Experimental/AnimNextAnimGraph/Source/AnimNextAnimGraph/Private/Traits/NotifyFilterTrait.cpp
@@ -0,0 +1,25 @@
+// Copyright Epic Games, Inc. All Rights Reserved.
+
+#include "NotifyFilterTrait.h"
+#include "TraitCore/ExecutionContext.h"
+
+namespace UE::AnimNext
+{
+	AUTO_REGISTER_ANIM_TRAIT(FNotifyFilterTrait)
+
+	FLazyName FNotifyFilterTrait::DisableNotifiesTag("DisableNotifies"); 
+	
+	// Trait implementation boilerplate
+	#define TRAIT_INTERFACE_ENUMERATOR(GeneratorMacro) \
+		GeneratorMacro(IScopedTagInterface) \
+		GeneratorMacro(IUpdate) \
+
+	GENERATE_ANIM_TRAIT_IMPLEMENTATION(FNotifyFilterTrait, TRAIT_INTERFACE_ENUMERATOR, NULL_ANIM_TRAIT_INTERFACE_ENUMERATOR, NULL_ANIM_TRAIT_EVENT_ENUMERATOR)
+	#undef TRAIT_INTERFACE_ENUMERATOR
+
+	bool FNotifyFilterTrait::AreNotifiesEnabledInScope(const FExecutionContext& Context)
+	{
+		return !IScopedTagInterface::IsTagInScope(Context, DisableNotifiesTag);		
+	}
+}
+
diff --git a/Engine/Plugins/Experimental/AnimNextAnimGraph/Source/AnimNextAnimGraph/Private/Traits/NotifyFilterTrait.h b/Engine/Plugins/Experimental/AnimNextAnimGraph/Source/AnimNextAnimGraph/Private/Traits/NotifyFilterTrait.h
new file mode 100644
--- /dev/null
+++ b/Engine/Plugins/Experimental/AnimNextAnimGraph/Source/AnimNextAnimGraph/Private/Traits/NotifyFilterTrait.h
@@ -0,0 +1,70 @@
+// Copyright Epic Games, Inc. All Rights Reserved.
+
+#pragma once
+
+#include "CoreMinimal.h"
+#include "TraitCore/IScopedTraitInterface.h"
+#include "TraitCore/TraitBinding.h"
+#include "TraitInterfaces/IScopedTag.h"
+#include "TraitInterfaces/IUpdate.h"
+
+#include "NotifyFilterTrait.generated.h"
+
+USTRUCT(meta = (DisplayName = "Notify Filter"))
+struct FAnimNextNotifyFilterTraitSharedData : public FAnimNextTraitSharedData
+{
+	GENERATED_BODY()
+
+	UPROPERTY(EditAnywhere, Category = "Filtering")
+	bool bShouldDisableNotifies = false;
+
+	// Latent pin support boilerplate
+	#define TRAIT_LATENT_PROPERTIES_ENUMERATOR(GeneratorMacro) \
+		GeneratorMacro(bShouldDisableNotifies) \
+
+	GENERATE_TRAIT_LATENT_PROPERTIES(FAnimNextNotifyFilterTraitSharedData, TRAIT_LATENT_PROPERTIES_ENUMERATOR)
+	#undef TRAIT_LATENT_PROPERTIES_ENUMERATOR
+};
+
+namespace UE::AnimNext
+{
+	struct FNotifyFilterTrait : FAdditiveTrait, IScopedTagInterface, IUpdate
+	{
+		DECLARE_ANIM_TRAIT(FNotifyFilterTrait, FAdditiveTrait)
+
+		using FSharedData = FAnimNextNotifyFilterTraitSharedData;
+		
+		static bool AreNotifiesEnabledInScope(const FExecutionContext& Context);
+
+		// IScopedTagInterface impl
+		virtual FName GetTag(const FExecutionContext& Context, const TTraitBinding<IScopedTagInterface>& Binding) const override
+		{
+			return DisableNotifiesTag;
+		}		
+
+		// IUpdate impl
+		virtual void PreUpdate(FUpdateTraversalContext& Context, const TTraitBinding<IUpdate>& Binding, const FTraitUpdateState& TraitState) const override
+		{
+			const FSharedData* SharedData = Binding.GetSharedData<FSharedData>();
+			if (SharedData->GetbShouldDisableNotifies(Binding))
+			{
+				Context.PushScopedInterface<IScopedTagInterface>(Binding);
+			}
+			
+			IUpdate::PreUpdate(Context, Binding, TraitState);
+		}
+
+		virtual void PostUpdate(FUpdateTraversalContext& Context, const TTraitBinding<IUpdate>& Binding, const FTraitUpdateState& TraitState) const override
+		{			
+			const FSharedData* SharedData = Binding.GetSharedData<FSharedData>();
+			if (SharedData->GetbShouldDisableNotifies(Binding))
+			{
+				ensure(Context.PopScopedInterface<IScopedTagInterface>(Binding));
+			}
+
+			IUpdate::PostUpdate(Context, Binding, TraitState);
+		}
+	private:
+		static FLazyName DisableNotifiesTag; 
+	};
+}
diff --git a/Engine/Plugins/Experimental/AnimNextAnimGraph/Source/AnimNextAnimGraph/Public/Graph/AnimNextAnimGraph.h b/Engine/Plugins/Experimental/AnimNextAnimGraph/Source/AnimNextAnimGraph/Public/Graph/AnimNextAnimGraph.h
--- a/Engine/Plugins/Experimental/AnimNextAnimGraph/Source/AnimNextAnimGraph/Public/Graph/AnimNextAnimGraph.h
+++ b/Engine/Plugins/Experimental/AnimNextAnimGraph/Source/AnimNextAnimGraph/Public/Graph/AnimNextAnimGraph.h
@@ -48,6 +48,7 @@ struct FAnimNextGraphInjectionData
 {
 	GENERATED_BODY()
 
+	bool HasValidEvaluationModifier() const { return EvaluationModifier.IsValid(); }
 private:
 	// Evaluation modifier to allow IEvaluate pass to be controlled
 	TWeakPtr<UE::AnimNext::IEvaluationModifier> EvaluationModifier;
diff --git a/Engine/Plugins/Experimental/AnimNextAnimGraph/Source/AnimNextAnimGraph/Public/TraitInterfaces/IScopedTag.h b/Engine/Plugins/Experimental/AnimNextAnimGraph/Source/AnimNextAnimGraph/Public/TraitInterfaces/IScopedTag.h
new file mode 100644
--- /dev/null
+++ b/Engine/Plugins/Experimental/AnimNextAnimGraph/Source/AnimNextAnimGraph/Public/TraitInterfaces/IScopedTag.h
@@ -0,0 +1,63 @@
+// Copyright Epic Games, Inc. All Rights Reserved.
+
+#pragma once
+
+#include "CoreMinimal.h"
+#include "TraitCore/ExecutionContext.h"
+#include "TraitCore/IScopedTraitInterface.h"
+#include "TraitCore/TraitBinding.h"
+
+namespace UE::AnimNext
+{
+struct IScopedTagInterface : IScopedTraitInterface
+{
+	DECLARE_ANIM_TRAIT_INTERFACE(IScopedTagInterface)
+
+	virtual FName GetTag(const FExecutionContext& Context, const TTraitBinding<IScopedTagInterface>& Binding) const;
+
+	static bool IsTagInScope(const FExecutionContext& Context, FName Tag);
+};
+
+template<>
+struct TTraitBinding<IScopedTagInterface> : FTraitBinding
+{
+	FName GetTag(const FExecutionContext& Context) const
+	{
+		return GetInterface()->GetTag(Context, *this);
+	}
+
+protected:
+	const IScopedTagInterface* GetInterface() const { return GetInterfaceTyped<IScopedTagInterface>(); }
+};
+
+inline FName IScopedTagInterface::GetTag(const FExecutionContext& Context, const TTraitBinding<IScopedTagInterface>& Binding) const
+{
+	TTraitBinding<IScopedTagInterface> SuperBinding;
+	if (Binding.GetStackInterfaceSuper(SuperBinding))
+	{
+		return SuperBinding.GetTag(Context);
+	}
+
+	return NAME_None;
+}
+
+inline bool IScopedTagInterface::IsTagInScope(const FExecutionContext& Context, FName Tag)
+{
+	bool bResult = false;
+
+	Context.ForEachScopedInterface<IScopedTagInterface>([&Context, Tag, &bResult](const TTraitBinding<IScopedTagInterface>& InterfaceBinding)
+		{
+			if (Tag == InterfaceBinding.GetTag(Context))
+			{
+				// We found our tag, stop iterating
+				bResult = true;
+				return false;
+			}
+
+			// Keep searching
+			return true;
+		});
+
+	return bResult;
+}
+}
diff --git a/Engine/Plugins/Experimental/AnimNextAnimGraph/Source/AnimNextAnimGraphTestSuite/Private/AnimNextAnimGraphTraitInterfacesTest.cpp b/Engine/Plugins/Experimental/AnimNextAnimGraph/Source/AnimNextAnimGraphTestSuite/Private/AnimNextAnimGraphTraitInterfacesTest.cpp
--- a/Engine/Plugins/Experimental/AnimNextAnimGraph/Source/AnimNextAnimGraphTestSuite/Private/AnimNextAnimGraphTraitInterfacesTest.cpp
+++ b/Engine/Plugins/Experimental/AnimNextAnimGraph/Source/AnimNextAnimGraphTestSuite/Private/AnimNextAnimGraphTraitInterfacesTest.cpp
@@ -5,6 +5,7 @@
 #include "AnimNextTest.h"
 
 #include "Misc/AutomationTest.h"
+#include "TraitInterfaces/IScopedTag.h"
 
 #if WITH_DEV_AUTOMATION_TESTS
 
@@ -229,58 +230,6 @@ namespace UE::AnimNext
 	GENERATE_ANIM_TRAIT_IMPLEMENTATION(FTraitWithChildren, TRAIT_INTERFACE_ENUMERATOR, NULL_ANIM_TRAIT_INTERFACE_ENUMERATOR, NULL_ANIM_TRAIT_EVENT_ENUMERATOR)
 	#undef TRAIT_INTERFACE_ENUMERATOR
 
-	struct IScopedTagInterface : IScopedTraitInterface
-	{
-		DECLARE_ANIM_TRAIT_INTERFACE(IScopedTagInterface)
-
-		virtual FName GetTag(const FExecutionContext& Context, const TTraitBinding<IScopedTagInterface>& Binding) const;
-
-		static bool IsTagInScope(const FExecutionContext& Context, FName Tag);
-	};
-
-	template<>
-	struct TTraitBinding<IScopedTagInterface> : FTraitBinding
-	{
-		FName GetTag(const FExecutionContext& Context) const
-		{
-			return GetInterface()->GetTag(Context, *this);
-		}
-
-	protected:
-		const IScopedTagInterface* GetInterface() const { return GetInterfaceTyped<IScopedTagInterface>(); }
-	};
-
-	FName IScopedTagInterface::GetTag(const FExecutionContext& Context, const TTraitBinding<IScopedTagInterface>& Binding) const
-	{
-		TTraitBinding<IScopedTagInterface> SuperBinding;
-		if (Binding.GetStackInterfaceSuper(SuperBinding))
-		{
-			return SuperBinding.GetTag(Context);
-		}
-
-		return NAME_None;
-	}
-
-	bool IScopedTagInterface::IsTagInScope(const FExecutionContext& Context, FName Tag)
-	{
-		bool bResult = false;
-
-		Context.ForEachScopedInterface<IScopedTagInterface>([&Context, Tag, &bResult](TTraitBinding<IScopedTagInterface>& InterfaceBinding)
-			{
-				if (Tag == InterfaceBinding.GetTag(Context))
-				{
-					// We found our tag, stop iterating
-					bResult = true;
-					return false;
-				}
-
-				// Keep searching
-				return true;
-			});
-
-		return bResult;
-	}
-
 	// Adds a scoped tag
 	struct FScopedTagTrait : FAdditiveTrait, IScopedTagInterface, IUpdate
 	{
