Fix uninitialized struct members
#rnx

diff --git a/Engine/Plugins/Experimental/AnimNextAnimGraph/Source/AnimNextAnimGraph/Private/Graph/RigUnit_HasValidEvaluationModifier.h b/Engine/Plugins/Experimental/AnimNextAnimGraph/Source/AnimNextAnimGraph/Private/Graph/RigUnit_HasValidEvaluationModifier.h
--- a/Engine/Plugins/Experimental/AnimNextAnimGraph/Source/AnimNextAnimGraph/Private/Graph/RigUnit_HasValidEvaluationModifier.h
+++ b/Engine/Plugins/Experimental/AnimNextAnimGraph/Source/AnimNextAnimGraph/Private/Graph/RigUnit_HasValidEvaluationModifier.h
@@ -22,5 +22,5 @@ struct FRigUnit_HasValidEvaluationModifier : public FRigUnit_AnimNextBase
 
 	// Whether Evaluation Modifier is valid (indicates active injection)
 	UPROPERTY(EditAnywhere, Category = "Struct", meta = (Output))
-	bool bValue;
+	bool bValue = false;
 };
