AnimNext: Pass-by-ref runtime refactor

diff --git a/Engine/Plugins/Experimental/AnimNextAnimGraph/Source/AnimNextAnimGraph/Private/Graph/AnimNextGraphInstance.cpp b/Engine/Plugins/Experimental/AnimNextAnimGraph/Source/AnimNextAnimGraph/Private/Graph/AnimNextGraphInstance.cpp
--- a/Engine/Plugins/Experimental/AnimNextAnimGraph/Source/AnimNextAnimGraph/Private/Graph/AnimNextGraphInstance.cpp
+++ b/Engine/Plugins/Experimental/AnimNextAnimGraph/Source/AnimNextAnimGraph/Private/Graph/AnimNextGraphInstance.cpp
@@ -3,8 +3,7 @@
 #include "Graph/AnimNextGraphInstance.h"
 
 #include "AnimNextAnimGraphStats.h"
-#include "DataInterface/AnimNextDataInterfaceHost.h"
-#include "DataInterface/DataInterfaceStructAdapter.h"
+#include "UAFRigVMComponent.h"
 #include "Graph/AnimNextAnimationGraph.h"
 #include "Graph/AnimNextGraphContextData.h"
 #include "Graph/AnimNextGraphLatentPropertiesContextData.h"
@@ -15,7 +14,6 @@
 #include "Logging/StructuredLog.h"
 #include "Module/AnimNextModule.h"
 #include "Module/AnimNextModuleInstance.h"
-#include "Module/ModuleGuard.h"
 
 DEFINE_STAT(STAT_AnimNext_Graph_RigVM);
 
@@ -51,9 +49,8 @@ void FAnimNextGraphInstance::Release()
 	GraphInstancePtr.Reset();
 	ModuleInstance = nullptr;
 	RootGraphInstance = nullptr;
-	ExtendedExecuteContext.Reset();
-	Components.Empty();
-	DataInterface = nullptr;
+	ReleaseComponents();
+	Asset = nullptr;
 }
 
 bool FAnimNextGraphInstance::IsValid() const
@@ -63,7 +60,7 @@ bool FAnimNextGraphInstance::IsValid() const
 
 const UAnimNextAnimationGraph* FAnimNextGraphInstance::GetAnimationGraph() const
 {
-	return CastChecked<UAnimNextAnimationGraph>(DataInterface);
+	return CastChecked<UAnimNextAnimationGraph>(Asset);
 }
 
 FName FAnimNextGraphInstance::GetEntryPoint() const
@@ -122,39 +119,6 @@ bool FAnimNextGraphInstance::IsRoot() const
 	return this == RootGraphInstance;
 }
 
-void FAnimNextGraphInstance::AddStructReferencedObjects(FReferenceCollector& Collector)
-{
-	if (!IsRoot())
-	{
-		return;	// If we aren't the root graph instance, we don't own the components
-	}
-
-	if (const UE::AnimNext::FGCGraphInstanceComponent* Component = TryGetComponent<UE::AnimNext::FGCGraphInstanceComponent>())
-	{
-		Component->AddReferencedObjects(Collector);
-	}
-}
-
-UE::AnimNext::FGraphInstanceComponent* FAnimNextGraphInstance::TryGetComponent(int32 ComponentNameHash, FName ComponentName) const
-{
-	if (const TSharedPtr<UE::AnimNext::FGraphInstanceComponent>* Component = RootGraphInstance->Components.FindByHash(ComponentNameHash, ComponentName))
-	{
-		return Component->Get();
-	}
-
-	return nullptr;
-}
-
-UE::AnimNext::FGraphInstanceComponent& FAnimNextGraphInstance::AddComponent(int32 ComponentNameHash, FName ComponentName, TSharedPtr<UE::AnimNext::FGraphInstanceComponent>&& Component)
-{
-	return *RootGraphInstance->Components.AddByHash(ComponentNameHash, ComponentName, MoveTemp(Component)).Get();
-}
-
-GraphInstanceComponentMapType::TConstIterator FAnimNextGraphInstance::GetComponentIterator() const
-{
-	return RootGraphInstance->Components.CreateConstIterator();
-}
-
 bool FAnimNextGraphInstance::HasUpdated() const
 {
 	return bHasUpdatedOnce;
@@ -176,6 +140,8 @@ void FAnimNextGraphInstance::ExecuteLatentPins(const TConstArrayView<UE::AnimNex
 
 	if (URigVM* VM = GetAnimationGraph()->RigVM)
 	{
+		FUAFRigVMComponent& RigVMComponent = GetComponent<FUAFRigVMComponent>();
+		FRigVMExtendedExecuteContext& ExtendedExecuteContext = RigVMComponent.GetExtendedExecuteContext();
 		FAnimNextExecuteContext& AnimNextContext = ExtendedExecuteContext.GetPublicDataSafe<FAnimNextExecuteContext>();
 		if (ModuleInstance)
 		{
@@ -190,6 +156,39 @@ void FAnimNextGraphInstance::ExecuteLatentPins(const TConstArrayView<UE::AnimNex
 	}
 }
 
+void FAnimNextGraphInstance::CopyVariablesToLatentPins(const TConstArrayView<UE::AnimNext::FLatentPropertyHandle>& LatentHandles, void* DestinationBasePtr, bool bIsFrozen)
+{
+	for (UE::AnimNext::FLatentPropertyHandle Handle : LatentHandles)
+	{
+		if (!Handle.IsIndexValid())
+		{
+			// This handle isn't valid
+			continue;
+		}
+
+		if (!Handle.IsOffsetValid())
+		{
+			// This handle isn't valid
+			continue;
+		}
+
+		if (bIsFrozen && Handle.CanFreeze())
+		{
+			// This handle can freeze and we are frozen, no need to update it
+			continue;
+		}
+
+		uint8* DestinationPtr = (uint8*)DestinationBasePtr + Handle.GetLatentPropertyOffset();
+		AccessVariablePropertyByIndex(Handle.GetLatentPropertyIndex(), [DestinationPtr](const FProperty* InProperty, TArrayView<uint8> InData)
+		{
+			// Copy from our source into our destination
+			// We assume the source and destination properties are identical
+			check(InData.Num() == InProperty->GetElementSize());
+			InProperty->CopyCompleteValue(DestinationPtr, InData.GetData());
+		});
+	}
+}
+
 #if WITH_EDITORONLY_DATA
 void FAnimNextGraphInstance::Freeze()
 {
@@ -199,10 +198,7 @@ void FAnimNextGraphInstance::Freeze()
 	}
 
 	GraphInstancePtr.Reset();
-	ExtendedExecuteContext.Reset();
-	Components.Empty();
-	CachedBindings.Reset();
-	PublicVariablesState = PublicVariablesState == EPublicVariablesState::Bound ? EPublicVariablesState::Unbound : EPublicVariablesState::None;
+	ReleaseComponents();
 	bHasUpdatedOnce = false;
 }
 
@@ -210,9 +206,7 @@ void FAnimNextGraphInstance::Thaw()
 {
 	if (const UAnimNextAnimationGraph* AnimationGraph = GetAnimationGraph())
 	{
-		Variables.MigrateToNewBagInstance(AnimationGraph->VariableDefaults);
-
-		ExtendedExecuteContext = AnimationGraph->ExtendedExecuteContext;
+		MigrateVariables();
 
 		{
 			UE::AnimNext::FExecutionContext Context(*this);
@@ -232,310 +226,15 @@ void FAnimNextGraphInstance::Thaw()
 
 void FAnimNextGraphInstance::OnModuleCompiled(UAnimNextModule* InModule)
 {
-	auto DependentModuleCompiled = [this, InModule]()
-	{
-		bool bFound = false;
-		ModuleInstance->ForEachPrerequisite([InModule, &bFound](FAnimNextModuleInstance& InPrerequisite)
-		{
-			if(InPrerequisite.GetModule() == InModule)
-			{
-				bFound = true;
-			}
-		});
-		return bFound;
-	};
-
 	// If we are hosted transitively by the compiled module, or if we could be bound to the compiled module,
-	// invalidate and mark our bindings as needing update. They will be lazily re-bound the next time we run.
+	// migrate our variables (which will re-bind to outer hosts)
 	if(ModuleInstance)
 	{
-		UE::AnimNext::FModuleWriteGuard Guard(ModuleInstance);
-		if(InModule == ModuleInstance->GetModule() || DependentModuleCompiled())
+		if(InModule == ModuleInstance->GetModule())
 		{
-			UnbindPublicVariables();
+			MigrateVariables();
 		}
 	}
 }
 
 #endif
-
-bool FAnimNextGraphInstance::BindToCachedBindings()
-{
-	bool bPublicVariablesBound = false;
-
-	auto FindCachedInterfaceBindings = [this](const UAnimNextDataInterface* InDataInterface) -> const FCachedDataInterfaceBinding*
-	{
-		for(const FCachedDataInterfaceBinding& CachedBinding : CachedBindings)
-		{
-			if(CachedBinding.DataInterface == InDataInterface)
-			{
-				return &CachedBinding;
-			}
-		}
-
-		return nullptr;
-	};
-
-	const UAnimNextAnimationGraph* AnimationGraph = GetAnimationGraph();
-	const UPropertyBag* PropertyBag = Variables.GetPropertyBagStruct();
-	for(const FAnimNextImplementedDataInterface& ImplementedInterface : AnimationGraph->GetImplementedInterfaces())
-	{
-		if(!ImplementedInterface.bAutoBindToHost)
-		{
-			continue;
-		}
-
-		const FCachedDataInterfaceBinding* CachedBinding = FindCachedInterfaceBindings(ImplementedInterface.DataInterface);
-		if(CachedBinding == nullptr)
-		{
-			// Did not cache this binding, so skip
-			continue;
-		}
-
-		if (CachedBinding->CachedBindings.Num() != ImplementedInterface.NumVariables)
-		{
-			UE_LOGFMT(LogAnimation, Error, "Interface size mismatch: {CachedInterfaceName} {CachedSize} vs {ImplementedInterfaceName} {ImplementedSize}. Interface's values will not be updated at runtime.", CachedBinding->DataInterface->GetFName(), CachedBinding->CachedBindings.Num(), ImplementedInterface.DataInterface->GetFName(), ImplementedInterface.NumVariables);
-			continue;
-		}
-
-		// Index into local property bag
-		int32 VariableIndex = ImplementedInterface.VariableIndex;
-		const int32 EndVariableIndex = ImplementedInterface.VariableIndex + ImplementedInterface.NumVariables;
-
-		// Index into interface
-		int32 InterfaceVariableIndex = 0;
-
-		for(; VariableIndex < EndVariableIndex; ++VariableIndex, ++InterfaceVariableIndex)
-		{
-			const FPropertyBagPropertyDesc& Desc = PropertyBag->GetPropertyDescs()[VariableIndex];
-			const FCachedDataInterfaceBinding::FVariable& CachedBindingVariable = CachedBinding->CachedBindings[InterfaceVariableIndex];
-			if(CachedBindingVariable.Memory == nullptr)
-			{
-				UE_LOGFMT(LogAnimation, Error, "FAnimNextGraphInstance::BindToCachedBindings: Did not find cached binding for {VariableName} - variable's values will not be updated at runtime.", Desc.Name);
-				continue;
-			}
-
-			// Validate that interface layout is the same
-			if( Desc.Name == CachedBindingVariable.VariableName &&
-				InterfaceVariableIndex == CachedBindingVariable.InterfaceVariableIndex &&
-				Desc.CachedProperty->GetClass() == CachedBindingVariable.Property->GetClass())
-			{
-				ExtendedExecuteContext.ExternalVariableRuntimeData[VariableIndex].Memory = CachedBindingVariable.Memory;
-				bPublicVariablesBound = true;
-			}
-			else
-			{
-				UE_LOGFMT(LogAnimation, Error, "Interface layout mismatch for {GraphName}. Interface's values will not be updated at runtime. (Have:Need): {HaveName}:{NeedName}, {HaveIndex}:{NeedIndex}, {HaveType}:{NeedType}", AnimationGraph->GetFName(), Desc.Name, CachedBindingVariable.VariableName, InterfaceVariableIndex, CachedBindingVariable.InterfaceVariableIndex, Desc.CachedProperty->GetClass()->GetFName(), CachedBindingVariable.Property->GetClass()->GetFName());
-				continue;
-			}
-		}
-	}
-
-	return bPublicVariablesBound;
-}
-
-void FAnimNextGraphInstance::UpdateCachedBindingsForHost()
-{
-	if(IsRoot())
-	{
-		if(FAnimNextModuleInstance* ModuleHost = GetModuleInstance())
-		{
-			UpdateCachedBindingsForHost(*ModuleHost);
-		}
-	}
-	else if(FAnimNextGraphInstance* ParentGraphInstance = GetParentGraphInstance())
-	{
-		// Just copy parent cached bindings
-		CachedBindings = ParentGraphInstance->CachedBindings;
-	}
-}
-
-void FAnimNextGraphInstance::UpdateCachedBindingsForHost(const FAnimNextModuleInstance& InHost)
-{
-	UpdateCachedBindingsForHostHelper(InHost);
-}
-
-void FAnimNextGraphInstance::UpdateCachedBindingsForHost(const UE::AnimNext::IDataInterfaceHost& InHost)
-{
-	UpdateCachedBindingsForHostHelper(InHost);
-}
-
-template<typename HostType>
-void FAnimNextGraphInstance::UpdateCachedBindingsForHostHelper(const HostType& InHost)
-{
-	auto FindOrAddCachedInterfaceBindings = [this](const UAnimNextDataInterface* InDataInterface, int32 InNumVariables) -> FCachedDataInterfaceBinding&
-	{
-		for(FCachedDataInterfaceBinding& CachedBinding : CachedBindings)
-		{
-			if(CachedBinding.DataInterface == InDataInterface)
-			{
-				return CachedBinding;
-			}
-		}
-
-		FCachedDataInterfaceBinding& NewCachedBinding = CachedBindings.AddDefaulted_GetRef();
-		NewCachedBinding.DataInterface = InDataInterface;
-		NewCachedBinding.CachedBindings.SetNum(InNumVariables);
-		return NewCachedBinding;
-	};
-
-	const UAnimNextDataInterface* HostDataInterface = InHost.GetDataInterface();
-	if(HostDataInterface == nullptr)
-	{
-		return;
-	}
-
-	const UPropertyBag* PropertyBag = HostDataInterface->GetPublicVariableDefaults().GetPropertyBagStruct();
-	if(PropertyBag == nullptr)
-	{
-		return;
-	}
-
-	for(const FAnimNextImplementedDataInterface& ImplementedInterface : HostDataInterface->GetImplementedInterfaces())
-	{
-		FCachedDataInterfaceBinding& CachedBinding = FindOrAddCachedInterfaceBindings(ImplementedInterface.DataInterface, ImplementedInterface.NumVariables);
-
-		if (CachedBinding.CachedBindings.Num() != ImplementedInterface.NumVariables)
-		{
-			UE_LOGFMT(LogAnimation, Error, "UpdateCachedBindingsForHostHelper: Interface size mismatch: {CachedInterfaceName} {CachedSize} vs {ImplementedInterfaceName} {ImplementedSize}. Interface's values will not be updated at runtime.", CachedBinding.DataInterface->GetFName(), CachedBinding.CachedBindings.Num(), ImplementedInterface.DataInterface->GetFName(), ImplementedInterface.NumVariables);
-			continue;
-		}
-
-		int32 VariableIndex = ImplementedInterface.VariableIndex;
-		const int32 EndVariableIndex = ImplementedInterface.VariableIndex + ImplementedInterface.NumVariables;
-		for(int32 InterfaceVariableIndex = 0; VariableIndex < EndVariableIndex && InterfaceVariableIndex < ImplementedInterface.NumVariables; ++VariableIndex, ++InterfaceVariableIndex)
-		{
-			const FPropertyBagPropertyDesc& Desc = PropertyBag->GetPropertyDescs()[VariableIndex];
-			uint8* HostMemory = InHost.GetMemoryForVariable(VariableIndex, Desc.Name, Desc.CachedProperty);
-			uint8* PrevMemory = CachedBinding.CachedBindings.IsValidIndex(InterfaceVariableIndex) 
-				? CachedBinding.CachedBindings[InterfaceVariableIndex].Memory 
-				: nullptr;
-			HostMemory = HostMemory ? HostMemory : PrevMemory;
-
-			CachedBinding.CachedBindings[InterfaceVariableIndex] = FCachedDataInterfaceBinding::FVariable(Desc.Name, InterfaceVariableIndex, Desc.CachedProperty, HostMemory);
-		}
-	}
-}
-
-void FAnimNextGraphInstance::BindPublicVariables(TArrayView<FStructView> InHostStructs, TConstArrayView<UE::AnimNext::IDataInterfaceHost*> InHosts)
-{
-	using namespace UE::AnimNext;
-
-	const UAnimNextAnimationGraph* AnimationGraph = GetAnimationGraph();
-	if (AnimationGraph == nullptr)
-	{
-		return;
-	}
-
-	ensure(CachedBindings.Num() == 0);
-
-	{
-		const int32 NumStructs = InHostStructs.Num();
-		const int32 NumHosts = InHosts.Num();
-
-		// stack allocating uninitialized arrays to contain FDataInterfaceStructAdapter(s)
-		// Note this means we can end up calling FMemory_Alloca(0), which is somewhat implementation specific.
-		// This is OK as we dont every iterate or use the memory in that case
-		TArrayView<FDataInterfaceStructAdapter> Adapters((FDataInterfaceStructAdapter*)FMemory_Alloca(NumStructs * sizeof(FDataInterfaceStructAdapter)), NumStructs);
-		TArrayView<IDataInterfaceHost*> Hosts((IDataInterfaceHost**)FMemory_Alloca((NumStructs + NumHosts) * sizeof(IDataInterfaceHost*)), NumStructs + NumHosts);
-
-		// initializing the FDataInterfaceStructAdapter(s) with their respective HostStruct
-		for (int32 StructIndex = 0; StructIndex < NumStructs; ++StructIndex)
-		{
-			FStructView& StructView = InHostStructs[StructIndex];
-			Hosts[StructIndex] = new (&Adapters[StructIndex]) FDataInterfaceStructAdapter(AnimationGraph, StructView);
-		}
-
-		// Copy incoming hosts to provide a combined set to BindPublicVariables
-		for (int32 HostIndex = 0; HostIndex < NumHosts; ++HostIndex)
-		{
-			Hosts[NumStructs + HostIndex] = InHosts[HostIndex];
-		}
-
-		// Inner binding 
-		BindPublicVariables(Hosts);
-
-		// uninitializing the FDataInterfaceStructAdapter(s)
-		for (int32 StructIndex = 0; StructIndex < NumStructs; ++StructIndex)
-		{
-			Adapters[StructIndex].~FDataInterfaceStructAdapter();
-		}
-
-		// abandoning no longer initialized stack allocated arrays to their demise
-	}
-}
-
-void FAnimNextGraphInstance::BindPublicVariables(TConstArrayView<UE::AnimNext::IDataInterfaceHost*> InHosts)
-{
-	using namespace UE::AnimNext;
-
-	const UAnimNextAnimationGraph* AnimationGraph = GetAnimationGraph();
-	if (AnimationGraph == nullptr)
-	{
-		return;
-	}
-
-	if(PublicVariablesState == EPublicVariablesState::Bound)
-	{
-		return;
-	}
-
-	const UPropertyBag* PropertyBag = Variables.GetPropertyBagStruct();
-	if(PropertyBag == nullptr)
-	{
-		// Nothing to bind
-		PublicVariablesState = EPublicVariablesState::None;
-		return;
-	}
-
-	ensure(CachedBindings.Num() == 0);
-
-	// First cache any bindings to this instance's host
-	UpdateCachedBindingsForHost();
-
-	// Next cache any supplied host interfaces 
-	for(const IDataInterfaceHost* HostInterface : InHosts)
-	{
-		check(HostInterface);
-		UpdateCachedBindingsForHost(*HostInterface);
-	}
-
-	// Bind to cached bindings we built above
-	if(BindToCachedBindings())
-	{
-		// Re-initialize memory handles
-		AnimationGraph->RigVM->InitializeInstance(ExtendedExecuteContext, /* bCopyMemory = */false);
-	}
-
-	PublicVariablesState = EPublicVariablesState::Bound;
-}
-
-void FAnimNextGraphInstance::UnbindPublicVariables()
-{
-	const UAnimNextAnimationGraph* AnimationGraph = GetAnimationGraph();
-	if (AnimationGraph == nullptr)
-	{
-		return;
-	}
-
-	if(PublicVariablesState != EPublicVariablesState::Bound)
-	{
-		return;
-	}
-
-	// Reset external variable ptrs to point to internal public vars
-	const int32 NumVariables = Variables.GetNumPropertiesInBag();
-	TConstArrayView<FPropertyBagPropertyDesc> Descs = Variables.GetPropertyBagStruct()->GetPropertyDescs();
-	uint8* BasePtr = Variables.GetMutableValue().GetMemory();
-	for(int32 VariableIndex = 0; VariableIndex < NumVariables; ++VariableIndex)
-	{
-		ExtendedExecuteContext.ExternalVariableRuntimeData[VariableIndex].Memory = Descs[VariableIndex].CachedProperty->ContainerPtrToValuePtr<uint8>(BasePtr);
-	}
-
-	// Re-initialize memory handles
-	AnimationGraph->RigVM->InitializeInstance(ExtendedExecuteContext, /* bCopyMemory = */false);
-
-	CachedBindings.Reset();
-
-	PublicVariablesState = EPublicVariablesState::Unbound;
-}
diff --git a/Engine/Plugins/Experimental/AnimNextAnimGraph/Source/AnimNextAnimGraph/Private/Graph/AnimNextModuleAnimGraphComponent.cpp b/Engine/Plugins/Experimental/AnimNextAnimGraph/Source/AnimNextAnimGraph/Private/Graph/AnimNextModuleAnimGraphComponent.cpp
--- a/Engine/Plugins/Experimental/AnimNextAnimGraph/Source/AnimNextAnimGraph/Private/Graph/AnimNextModuleAnimGraphComponent.cpp
+++ b/Engine/Plugins/Experimental/AnimNextAnimGraph/Source/AnimNextAnimGraph/Private/Graph/AnimNextModuleAnimGraphComponent.cpp
@@ -4,18 +4,24 @@
 
 #include "Graph/AnimNextGraphInstance.h"
 #include "Graph/AnimNextAnimationGraph.h"
+#include "Module/AnimNextModuleInstance.h"
 
-TWeakPtr<FAnimNextGraphInstance> FAnimNextModuleAnimGraphComponent::AllocateInstance(const UAnimNextAnimationGraph* InAnimationGraph, FAnimNextGraphInstance* InParentInstance, FName InEntryPoint)
+TWeakPtr<FAnimNextGraphInstance> FAnimNextModuleAnimGraphComponent::AllocateInstance(const UAnimNextAnimationGraph* InAnimationGraph, const TSharedPtr<const UE::AnimNext::FVariableOverridesCollection>& InOverrides)
 {
 	check(InAnimationGraph);
 
-	TSharedPtr<FAnimNextGraphInstance> NewInstance = InAnimationGraph->AllocateInstance(&GetModuleInstance(), nullptr, InParentInstance, InEntryPoint);
+	TSharedPtr<FAnimNextGraphInstance> NewInstance = InAnimationGraph->AllocateInstance(
+		{
+			.ModuleInstance = &GetModuleInstance(),
+			.Overrides = InOverrides
+		});
 	if(!NewInstance.IsValid())
 	{
 		return TWeakPtr<FAnimNextGraphInstance>();
 	}
 
-	GraphInstances.Add(NewInstance);
+	// Persist a strong reference to overrides here
+	GraphInstances.Add(NewInstance.ToSharedRef());
 	return NewInstance;
 }
 
@@ -34,7 +40,7 @@ void FAnimNextModuleAnimGraphComponent::AddStructReferencedObjects(class FRefere
 {
 	for (const TSharedPtr<FAnimNextGraphInstance>& GraphInstance : GraphInstances)
 	{
-		GraphInstance->AddStructReferencedObjects(Collector);
+		Collector.AddPropertyReferencesWithStructARO(FAnimNextGraphInstance::StaticStruct(), GraphInstance.Get());
 	}
 }
 
diff --git a/Engine/Plugins/Experimental/AnimNextAnimGraph/Source/AnimNextAnimGraph/Private/Graph/GC_GraphInstanceComponent.cpp b/Engine/Plugins/Experimental/AnimNextAnimGraph/Source/AnimNextAnimGraph/Private/Graph/GC_GraphInstanceComponent.cpp
--- a/Engine/Plugins/Experimental/AnimNextAnimGraph/Source/AnimNextAnimGraph/Private/Graph/GC_GraphInstanceComponent.cpp
+++ b/Engine/Plugins/Experimental/AnimNextAnimGraph/Source/AnimNextAnimGraph/Private/Graph/GC_GraphInstanceComponent.cpp
@@ -5,49 +5,45 @@
 #include "TraitCore/ExecutionContext.h"
 #include "TraitInterfaces/IGarbageCollection.h"
 
-namespace UE::AnimNext
+void FGCGraphInstanceComponent::Register(const UE::AnimNext::FWeakTraitPtr& InTraitPtr)
 {
-	FGCGraphInstanceComponent::FGCGraphInstanceComponent(FAnimNextGraphInstance& InOwnerInstance)
-		: FGraphInstanceComponent(InOwnerInstance)
-	{
-	}
+	TraitsWithReferences.Add(InTraitPtr);
+}
 
-	void FGCGraphInstanceComponent::Register(const FWeakTraitPtr& InTraitPtr)
-	{
-		TraitsWithReferences.Add(InTraitPtr);
-	}
+void FGCGraphInstanceComponent::Unregister(const UE::AnimNext::FWeakTraitPtr& InTraitPtr)
+{
+	using namespace UE::AnimNext;
+	
+	const int32 EntryIndex = TraitsWithReferences.IndexOfByPredicate(
+		[&InTraitPtr](const FWeakTraitPtr& TraitPtr)
+		{
+			return TraitPtr == InTraitPtr;
+		});
 
-	void FGCGraphInstanceComponent::Unregister(const FWeakTraitPtr& InTraitPtr)
+	if (ensure(EntryIndex != INDEX_NONE))
 	{
-		const int32 EntryIndex = TraitsWithReferences.IndexOfByPredicate(
-			[&InTraitPtr](const FWeakTraitPtr& TraitPtr)
-			{
-				return TraitPtr == InTraitPtr;
-			});
-
-		if (ensure(EntryIndex != INDEX_NONE))
-		{
-			TraitsWithReferences.RemoveAtSwap(EntryIndex);
-		}
+		TraitsWithReferences.RemoveAtSwap(EntryIndex);
 	}
+}
 
-	void FGCGraphInstanceComponent::AddReferencedObjects(FReferenceCollector& Collector) const
+void FGCGraphInstanceComponent::AddStructReferencedObjects(FReferenceCollector& Collector)
+{
+	using namespace UE::AnimNext;
+	
+	FExecutionContext Context;
+	FTraitStackBinding TraitStack;
+	TTraitBinding<IGarbageCollection> GCTrait;
+
+	// TODO: If we kept the entries sorted by graph instance, we could re-use the execution context
+	for (const FWeakTraitPtr& TraitPtr : TraitsWithReferences)
 	{
-		FExecutionContext Context;
-		FTraitStackBinding TraitStack;
-		TTraitBinding<IGarbageCollection> GCTrait;
+		Context.BindTo(TraitPtr);
 
-		// TODO: If we kept the entries sorted by graph instance, we could re-use the execution context
-		for (const FWeakTraitPtr& TraitPtr : TraitsWithReferences)
+		if (Context.GetStack(TraitPtr, TraitStack))
 		{
-			Context.BindTo(TraitPtr);
-
-			if (Context.GetStack(TraitPtr, TraitStack))
-			{
-				ensure(TraitStack.GetInterface(GCTrait));
+			ensure(TraitStack.GetInterface(GCTrait));
 
-				GCTrait.AddReferencedObjects(Context, Collector);
-			}
+			GCTrait.AddReferencedObjects(Context, Collector);
 		}
 	}
 }
diff --git a/Engine/Plugins/Experimental/AnimNextAnimGraph/Source/AnimNextAnimGraph/Private/Graph/GC_GraphInstanceComponent.h b/Engine/Plugins/Experimental/AnimNextAnimGraph/Source/AnimNextAnimGraph/Private/Graph/GC_GraphInstanceComponent.h
--- a/Engine/Plugins/Experimental/AnimNextAnimGraph/Source/AnimNextAnimGraph/Private/Graph/GC_GraphInstanceComponent.h
+++ b/Engine/Plugins/Experimental/AnimNextAnimGraph/Source/AnimNextAnimGraph/Private/Graph/GC_GraphInstanceComponent.h
@@ -4,35 +4,43 @@
 
 #include "CoreMinimal.h"
 #include "TraitCore/TraitPtr.h"
-#include "Graph/GraphInstanceComponent.h"
+#include "Graph/UAFGraphInstanceComponent.h"
+#include "GC_GraphInstanceComponent.generated.h"
 
 class FReferenceCollector;
 
-namespace UE::AnimNext
+/**
+ * FGCGraphInstanceComponent
+ *
+ * This component maintains the necessary state to garbage collection.
+ */
+USTRUCT()
+struct FGCGraphInstanceComponent : public FUAFGraphInstanceComponent
 {
-	/**
-	 * FGCGraphInstanceComponent
-	 *
-	 * This component maintains the necessary state to garbage collection.
-	 */
-	struct FGCGraphInstanceComponent : public FGraphInstanceComponent
-	{
-		DECLARE_ANIM_GRAPH_INSTANCE_COMPONENT(FGCGraphInstanceComponent)
+	GENERATED_BODY()
+
+	FGCGraphInstanceComponent() = default;
 
-		explicit FGCGraphInstanceComponent(FAnimNextGraphInstance& InOwnerInstance);
+	// Registers the provided trait with the GC system
+	// Once registered, IGarbageCollection::AddReferencedObjects will be called on it during GC
+	void Register(const UE::AnimNext::FWeakTraitPtr& InTraitPtr);
 
-		// Registers the provided trait with the GC system
-		// Once registered, IGarbageCollection::AddReferencedObjects will be called on it during GC
-		void Register(const FWeakTraitPtr& InTraitPtr);
+	// Unregisters the provided trait from the GC system
+	void Unregister(const UE::AnimNext::FWeakTraitPtr& InTraitPtr);
 
-		// Unregisters the provided trait from the GC system
-		void Unregister(const FWeakTraitPtr& InTraitPtr);
+	// Called during garbage collection to collect strong object references
+	void AddStructReferencedObjects(FReferenceCollector& Collector);
 
-		// Called during garbage collection to collect strong object references
-		void AddReferencedObjects(FReferenceCollector& Collector) const;
+private:
+	// List of trait handles that contain UObject references and implement IGarbageCollection
+	TArray<UE::AnimNext::FWeakTraitPtr> TraitsWithReferences;
+};
 
-	private:
-		// List of trait handles that contain UObject references and implement IGarbageCollection
-		TArray<FWeakTraitPtr> TraitsWithReferences;
+template<>
+struct TStructOpsTypeTraits<FGCGraphInstanceComponent> : public TStructOpsTypeTraitsBase2<FGCGraphInstanceComponent>
+{
+	enum
+	{
+		WithAddStructReferencedObjects = true,
 	};
-}
+};
diff --git a/Engine/Plugins/Experimental/AnimNextAnimGraph/Source/AnimNextAnimGraph/Private/Graph/RigUnit_AnimNextRunAnimationGraph_v1.cpp b/Engine/Plugins/Experimental/AnimNextAnimGraph/Source/AnimNextAnimGraph/Private/Graph/RigUnit_AnimNextRunAnimationGraph_v1.cpp
--- a/Engine/Plugins/Experimental/AnimNextAnimGraph/Source/AnimNextAnimGraph/Private/Graph/RigUnit_AnimNextRunAnimationGraph_v1.cpp
+++ b/Engine/Plugins/Experimental/AnimNextAnimGraph/Source/AnimNextAnimGraph/Private/Graph/RigUnit_AnimNextRunAnimationGraph_v1.cpp
@@ -3,6 +3,7 @@
 #include "Graph/RigUnit_AnimNextRunAnimationGraph_v1.h"
 
 #include "AnimNextAnimGraphStats.h"
+#include "UAFRigVMComponent.h"
 #include "TraitInterfaces/IEvaluate.h"
 #include "TraitInterfaces/IUpdate.h"
 #include "Graph/AnimNextModuleAnimGraphComponent.h"
@@ -46,7 +47,7 @@ FRigUnit_AnimNextRunAnimationGraph_v1_Execute()
 	// Lazily (re-)allocate graph instance if required
 	if(!WorkData.WeakInstance.IsValid())
 	{
-		WorkData.WeakInstance = AnimationGraphComponent.AllocateInstance(Graph);
+		WorkData.WeakInstance = AnimationGraphComponent.AllocateInstance(Graph, Overrides.GetOverrides().Pin());
 	}
 
 	if(!WorkData.WeakInstance.IsValid())
@@ -57,32 +58,14 @@ FRigUnit_AnimNextRunAnimationGraph_v1_Execute()
 	// Take a strong reference to the host instance, we are going to run it
 	TSharedRef<FAnimNextGraphInstance> InstanceRef = WorkData.WeakInstance.Pin().ToSharedRef();
 
-	if(InstanceRef->RequiresPublicVariableBinding())
+	if (InstanceRef->GetAnimationGraph()->GetRigVM())
 	{
-		// Bind any public variables we find
-		TArray<FPublicVariablesTraitToDataInterfaceHostAdapter, TInlineAllocator<4>> Adapters;
-		Adapters.Reserve(ExecuteContext.GetTraits().Num());
-		TArray<IDataInterfaceHost*, TInlineAllocator<4>> Hosts;
-		Hosts.Reserve(ExecuteContext.GetTraits().Num());
-		for(const FRigVMTraitScope& TraitScope : ExecuteContext.GetTraits())
-		{
-			const FRigVMTrait_AnimNextPublicVariables* VariablesTrait = TraitScope.GetTrait<FRigVMTrait_AnimNextPublicVariables>();
-			if(VariablesTrait == nullptr)
-			{
-				continue;
-			}
-
-			FPublicVariablesTraitToDataInterfaceHostAdapter& Adapter = Adapters.Emplace_GetRef(*VariablesTrait, TraitScope);
-			Hosts.Add(&Adapter);
-		}
-
-		InstanceRef->BindPublicVariables(Hosts);
+		// Propagate delta time
+		FUAFRigVMComponent& RigVMComponent = InstanceRef->GetComponent<FUAFRigVMComponent>();
+		FAnimNextExecuteContext& AnimNextExecuteContext = RigVMComponent.GetExtendedExecuteContext().GetPublicDataSafe<FAnimNextExecuteContext>();
+		AnimNextExecuteContext.SetDeltaTime(ExecuteContext.GetDeltaTime());
 	}
 
-	// Propagate delta time
-	FAnimNextExecuteContext& AnimNextExecuteContext = InstanceRef->GetExtendedExecuteContext().GetPublicDataSafe<FAnimNextExecuteContext>();
-	AnimNextExecuteContext.SetDeltaTime(ExecuteContext.GetDeltaTime());
-
 	const UE::AnimNext::FReferencePose& RefPose = ReferencePose.ReferencePose.GetRef<UE::AnimNext::FReferencePose>();
 
 	int32 DesiredLOD = LOD;
diff --git a/Engine/Plugins/Experimental/AnimNextAnimGraph/Source/AnimNextAnimGraph/Private/Graph/RigUnit_AnimNextRunAnimationGraph_v2.cpp b/Engine/Plugins/Experimental/AnimNextAnimGraph/Source/AnimNextAnimGraph/Private/Graph/RigUnit_AnimNextRunAnimationGraph_v2.cpp
--- a/Engine/Plugins/Experimental/AnimNextAnimGraph/Source/AnimNextAnimGraph/Private/Graph/RigUnit_AnimNextRunAnimationGraph_v2.cpp
+++ b/Engine/Plugins/Experimental/AnimNextAnimGraph/Source/AnimNextAnimGraph/Private/Graph/RigUnit_AnimNextRunAnimationGraph_v2.cpp
@@ -12,10 +12,11 @@
 #include "Graph/AnimNextModuleAnimGraphComponent.h"
 #include "Graph/RigVMTrait_AnimNextPublicVariables.h"
 #include "Injection/GraphInstanceInjectionComponent.h"
-#include "Injection/InjectionDataInterfaceHostAdapter.h"
 #include "TraitCore/TraitEventList.h"
 #include "Injection/AnimNextModuleInjectionComponent.h"
 #include "TraceAnimNextGraphInstances.h"
+#include "UAFRigVMComponent.h"
+#include "Factory/AnimGraphFactory.h"
 
 #include UE_INLINE_GENERATED_CPP_BY_NAME(RigUnit_AnimNextRunAnimationGraph_v2)
 
@@ -50,7 +51,7 @@ FRigUnit_AnimNextRunAnimationGraph_v2_Execute()
 	ensure(Result.LODPose.LODLevel == DesiredLOD);
 
 	// Get a host to run this graph
-	const UAnimNextAnimationGraph* HostGraph = Graph.HostGraph ? Graph.HostGraph.Get() : GetDefault<UAnimNextAnimGraphSettings>()->GetDefaultRunGraphHost();// GetHostGraphToRun(ExecuteContext);
+	const UAnimNextAnimationGraph* HostGraph = Graph.HostGraph ? Graph.HostGraph.Get() : FAnimGraphFactory::GetDefaultGraphHost();
 	FAnimNextModuleAnimGraphComponent& AnimationGraphComponent = ModuleInstance.GetComponent<FAnimNextModuleAnimGraphComponent>();
 	if(HostGraph == nullptr)
 	{
@@ -68,7 +69,17 @@ FRigUnit_AnimNextRunAnimationGraph_v2_Execute()
 	// Lazily (re-)allocate graph instance if required
 	if(!WorkData.WeakHost.IsValid())
 	{
-		WorkData.WeakHost = AnimationGraphComponent.AllocateInstance(HostGraph);
+		// Reset our graph reference, we may be allocating a new graph
+		WorkData.InjectedGraphReference.Reset();
+
+		TSharedPtr<FAnimNextGraphInstance> NewGraphInstance = AnimationGraphComponent.AllocateInstance(HostGraph, Overrides.GetOverrides().Pin()).Pin();
+		if (NewGraphInstance.IsValid())
+		{
+			WorkData.WeakHost = NewGraphInstance;
+
+			FGraphInstanceInjectionComponent& HostInjectionComponent = NewGraphInstance->GetComponent<FGraphInstanceInjectionComponent>();
+			WorkData.InjectedGraphReference = HostInjectionComponent.GetInjectionInfo().GetDefaultInjectionSite();
+		}
 	}
 
 	if(!WorkData.WeakHost.IsValid())
@@ -80,52 +91,20 @@ FRigUnit_AnimNextRunAnimationGraph_v2_Execute()
 	TSharedRef<FAnimNextGraphInstance> HostInstanceRef = WorkData.WeakHost.Pin().ToSharedRef();
 	FAnimNextGraphInstance& HostInstance = HostInstanceRef.Get();
 
-	// Apply bindings
-	if(HostInstance.RequiresPublicVariableBinding())
+	// Re-set the injected graph each time, as it may change per-update
+	if (!WorkData.InjectedGraphReference.IsNone())
 	{
-		TArray<FPublicVariablesTraitToDataInterfaceHostAdapter, TInlineAllocator<4>> Adapters;
-		Adapters.Reserve(ExecuteContext.GetTraits().Num());
-		TArray<IDataInterfaceHost*, TInlineAllocator<4>> Hosts;
-		Hosts.Reserve(ExecuteContext.GetTraits().Num() + 1);
-		FInjectionDataInterfaceHostAdapter InjectionAdapter;
-
-		// Bind any pin-overriden public variable traits we find
-		for(const FRigVMTraitScope& TraitScope : ExecuteContext.GetTraits())
-		{
-			const FRigVMTrait_AnimNextPublicVariables* VariablesTrait = TraitScope.GetTrait<FRigVMTrait_AnimNextPublicVariables>();
-			if(VariablesTrait == nullptr)
-			{
-				continue;
-			}
-
-			FPublicVariablesTraitToDataInterfaceHostAdapter& Adapter = Adapters.Emplace_GetRef(*VariablesTrait, TraitScope);
-			Hosts.Add(&Adapter);
-		}
-
-		// Inject the default graph if there is any (_after_ pin-exposed variables so they cant override it)
-		FGraphInstanceInjectionComponent& HostInjectionComponent = HostInstance.GetComponent<FGraphInstanceInjectionComponent>();
-		FName GraphName;
-		TStructView<FAnimNextAnimGraph> GraphInstanceView;
-		if(HostInjectionComponent.GetInjectionInfo().GetDefaultInjectableGraph(GraphName, GraphInstanceView))
-		{
-			WorkData.InjectedGraph = Graph;
-			GraphInstanceView = TStructView<FAnimNextAnimGraph>(WorkData.InjectedGraph);
-			InjectionAdapter = FInjectionDataInterfaceHostAdapter(HostInstance, GraphName, GraphInstanceView);
-			Hosts.Add(&InjectionAdapter);
-		}
-
-		HostInstance.BindPublicVariables(Hosts);
+		HostInstance.SetVariable(WorkData.InjectedGraphReference, Graph);
 	}
-	else
+
+	if (HostInstance.GetAnimationGraph()->GetRigVM())
 	{
-		// Re-copy the graph into the injected work data, as it may have been updated by bindings/injection
-		WorkData.InjectedGraph = Graph;
+		// Propagate delta time
+		FUAFRigVMComponent& RigVMComponent = HostInstance.GetComponent<FUAFRigVMComponent>();
+		FAnimNextExecuteContext& AnimNextExecuteContext = RigVMComponent.GetExtendedExecuteContext().GetPublicDataSafe<FAnimNextExecuteContext>();
+		AnimNextExecuteContext.SetDeltaTime(ExecuteContext.GetDeltaTime());
 	}
 
-	// Propagate delta time
-	FAnimNextExecuteContext& AnimNextExecuteContext = HostInstance.GetExtendedExecuteContext().GetPublicDataSafe<FAnimNextExecuteContext>();
-	AnimNextExecuteContext.SetDeltaTime(ExecuteContext.GetDeltaTime());
-
 	// Every graph in a schedule will see the same input events (if they were queued before the schedule started)
 	FUpdateGraphContext UpdateGraphContext(HostInstance, ExecuteContext.GetDeltaTime());
 	UpdateGraphContext.SetBindingObject(RefPose.SkeletalMeshComponent);
diff --git a/Engine/Plugins/Experimental/AnimNextAnimGraph/Source/AnimNextAnimGraph/Private/Graph/SyncGroup_GraphInstanceComponent.cpp b/Engine/Plugins/Experimental/AnimNextAnimGraph/Source/AnimNextAnimGraph/Private/Graph/SyncGroup_GraphInstanceComponent.cpp
--- a/Engine/Plugins/Experimental/AnimNextAnimGraph/Source/AnimNextAnimGraph/Private/Graph/SyncGroup_GraphInstanceComponent.cpp
+++ b/Engine/Plugins/Experimental/AnimNextAnimGraph/Source/AnimNextAnimGraph/Private/Graph/SyncGroup_GraphInstanceComponent.cpp
@@ -2124,297 +2124,304 @@ namespace UE::AnimNext
 			return ModuleInstance != nullptr ? ModuleInstance->GetObject() : nullptr;
 		}
 	}
+}
 
-	FSyncGroupGraphInstanceComponent::FSyncGroupGraphInstanceComponent(FAnimNextGraphInstance& InOwnerInstance)
-		: FGraphInstanceComponent(InOwnerInstance)
-	{
-	}
+// Defaulted constructors/copy-operators in cpp so fwd decls of members can work
+FSyncGroupGraphInstanceComponent::FSyncGroupGraphInstanceComponent() = default;
+FSyncGroupGraphInstanceComponent::FSyncGroupGraphInstanceComponent(FSyncGroupGraphInstanceComponent&&) = default;
+FSyncGroupGraphInstanceComponent::FSyncGroupGraphInstanceComponent(const FSyncGroupGraphInstanceComponent&) = default;
+FSyncGroupGraphInstanceComponent& FSyncGroupGraphInstanceComponent::operator=(const FSyncGroupGraphInstanceComponent&) = default;
+FSyncGroupGraphInstanceComponent& FSyncGroupGraphInstanceComponent::operator=(FSyncGroupGraphInstanceComponent&&) = default;
 
-	FSyncGroupGraphInstanceComponent::~FSyncGroupGraphInstanceComponent()
-	{
-		// We should have released all used group names before we are destroyed
-		check(UsedUniqueGroupNames.IsEmpty());
+FSyncGroupGraphInstanceComponent::~FSyncGroupGraphInstanceComponent()
+{
+	using namespace UE::AnimNext;
 
-		Private::FSyncGroupUniqueName* Entry = FirstFreeUniqueGroupName;
-		while (Entry != nullptr)
-		{
-			Private::FSyncGroupUniqueName* NextEntry = Entry->NextFreeEntry;
+	// We should have released all used group names before we are destroyed
+	check(UsedUniqueGroupNames.IsEmpty());
 
-			delete Entry;
+	Private::FSyncGroupUniqueName* Entry = FirstFreeUniqueGroupName;
+	while (Entry != nullptr)
+	{
+		Private::FSyncGroupUniqueName* NextEntry = Entry->NextFreeEntry;
 
-			Entry = NextEntry;
-		}
+		delete Entry;
 
-		FirstFreeUniqueGroupName = nullptr;
+		Entry = NextEntry;
 	}
 
-	void FSyncGroupGraphInstanceComponent::RegisterWithGroup(const FSyncGroupParameters& GroupParameters, const FWeakTraitPtr& TraitPtr, const FTraitUpdateState& TraitState)
-	{
-		using namespace Private;
+	FirstFreeUniqueGroupName = nullptr;
+}
 
-		if (!ensure(TraitPtr.IsValid()))
-		{
-			return;
-		}
+void FSyncGroupGraphInstanceComponent::RegisterWithGroup(const UE::AnimNext::FSyncGroupParameters& GroupParameters, const UE::AnimNext::FWeakTraitPtr& TraitPtr, const UE::AnimNext::FTraitUpdateState& TraitState)
+{
+	using namespace UE::AnimNext::Private;
 
-		const int32 GroupIndex = SyncGroupMap.FindOrAdd(GroupParameters.GroupName, SyncGroups.Num());
-		FSyncGroupState& GroupState = GroupIndex == SyncGroups.Num() ? SyncGroups.AddDefaulted_GetRef() : SyncGroups[GroupIndex];
-		GroupState.GroupName = GroupParameters.GroupName;
-		GroupState.bIsActive = true;
+	if (!ensure(TraitPtr.IsValid()))
+	{
+		return;
+	}
 
-		bool bIsMemberNew = true;
-		for (FSyncGroupMember& GroupMember : GroupState.Members)
+	const int32 GroupIndex = SyncGroupMap.FindOrAdd(GroupParameters.GroupName, SyncGroups.Num());
+	FSyncGroupState& GroupState = GroupIndex == SyncGroups.Num() ? SyncGroups.AddDefaulted_GetRef() : SyncGroups[GroupIndex];
+	GroupState.GroupName = GroupParameters.GroupName;
+	GroupState.bIsActive = true;
+
+	bool bIsMemberNew = true;
+	for (FSyncGroupMember& GroupMember : GroupState.Members)
+	{
+		if (GroupMember.TraitPtr == TraitPtr)
 		{
-			if (GroupMember.TraitPtr == TraitPtr)
-			{
-				bIsMemberNew = false;
-				check(!GroupMember.bIsActive);
+			bIsMemberNew = false;
+			check(!GroupMember.bIsActive);
 
-				GroupMember.TraitState = TraitState;
-				GroupMember.GroupParameters = GroupParameters;
-				GroupMember.bIsActive = true;
-				break;
-			}
+			GroupMember.TraitState = TraitState;
+			GroupMember.GroupParameters = GroupParameters;
+			GroupMember.bIsActive = true;
+			break;
 		}
+	}
 
-		if (bIsMemberNew)
-		{
-			GroupState.Members.Add(FSyncGroupMember{ TraitState, FTraitPtr(TraitPtr), GroupParameters, true, true });
+	if (bIsMemberNew)
+	{
+		GroupState.Members.Add(FSyncGroupMember{ TraitState, UE::AnimNext::FTraitPtr(TraitPtr), GroupParameters, true, true });
 
 #if UE_DEBUG_SYNC_GROUPS
-			if (CDebugVarEnableSyncLog.GetValueOnAnyThread())
-			{
-				UE_VLOG_UELOG(GetHostObject(TraitPtr), LogAnimMarkerSync, Verbose,
-					TEXT("[%s] [%p] %s"),
-					*GroupParameters.GroupName.ToString(),
-					TraitPtr.GetNodeInstance(),
-					GroupState.Members.Num() == 1 ? TEXT("Created") : TEXT("Joined"));
-			}
-#endif
+		if (CDebugVarEnableSyncLog.GetValueOnAnyThread())
+		{
+			UE_VLOG_UELOG(GetHostObject(TraitPtr), LogAnimMarkerSync, Verbose,
+				TEXT("[%s] [%p] %s"),
+				*GroupParameters.GroupName.ToString(),
+				TraitPtr.GetNodeInstance(),
+				GroupState.Members.Num() == 1 ? TEXT("Created") : TEXT("Joined"));
 		}
+#endif
 	}
+}
 
-	FName FSyncGroupGraphInstanceComponent::CreateUniqueGroupName()
+FName FSyncGroupGraphInstanceComponent::CreateUniqueGroupName()
+{
+	using namespace UE::AnimNext;
+	
+	// We create a unique group name by incrementing a counter and using it to generate a new FName
+	// We also recycle them to keep the counter as low as possible because in some configurations
+	// FNames are 32-bit as the number is stored in the string table. Recycling the group names
+	// ensures we don't pollute the FName string table.
+	// The counter could be the same across multiple graph instances which is fine since we only
+	// synchronize using groups within a single graph instance.
+
+	Private::FSyncGroupUniqueName* Entry;
+
+	if (FirstFreeUniqueGroupName != nullptr)
 	{
-		// We create a unique group name by incrementing a counter and using it to generate a new FName
-		// We also recycle them to keep the counter as low as possible because in some configurations
-		// FNames are 32-bit as the number is stored in the string table. Recycling the group names
-		// ensures we don't pollute the FName string table.
-		// The counter could be the same across multiple graph instances which is fine since we only
-		// synchronize using groups within a single graph instance.
+		// Recycle an old entry
+		Entry = FirstFreeUniqueGroupName;
+		FirstFreeUniqueGroupName = Entry->NextFreeEntry;
+	}
+	else
+	{
+		// Generate a new unique group name
+		FName UniqueGroupName = FName(Private::NAME_UniqueGroupNamePrefix, UniqueGroupNameCounter);
+		UniqueGroupNameCounter++;
 
-		Private::FSyncGroupUniqueName* Entry;
+		Entry = new Private::FSyncGroupUniqueName();
+		Entry->GroupName = UniqueGroupName;
+	}
 
-		if (FirstFreeUniqueGroupName != nullptr)
-		{
-			// Recycle an old entry
-			Entry = FirstFreeUniqueGroupName;
-			FirstFreeUniqueGroupName = Entry->NextFreeEntry;
-		}
-		else
-		{
-			// Generate a new unique group name
-			FName UniqueGroupName = FName(Private::NAME_UniqueGroupNamePrefix, UniqueGroupNameCounter);
-			UniqueGroupNameCounter++;
+	Entry->NextFreeEntry = nullptr;
 
-			Entry = new Private::FSyncGroupUniqueName();
-			Entry->GroupName = UniqueGroupName;
-		}
+	UsedUniqueGroupNames.Add(Entry->GroupName, Entry);
 
-		Entry->NextFreeEntry = nullptr;
+	return Entry->GroupName;
+}
 
-		UsedUniqueGroupNames.Add(Entry->GroupName, Entry);
+void FSyncGroupGraphInstanceComponent::ReleaseUniqueGroupName(FName GroupName)
+{
+	using namespace UE::AnimNext::Private;
 
-		return Entry->GroupName;
+	FSyncGroupUniqueName* Entry = nullptr;
+	if (ensure(UsedUniqueGroupNames.RemoveAndCopyValue(GroupName, Entry)))
+	{
+		Entry->NextFreeEntry = FirstFreeUniqueGroupName;
+		FirstFreeUniqueGroupName = Entry;
 	}
 
-	void FSyncGroupGraphInstanceComponent::ReleaseUniqueGroupName(FName GroupName)
+	if (const int32* GroupIndex = SyncGroupMap.Find(GroupName))
 	{
-		using namespace Private;
-
-		FSyncGroupUniqueName* Entry = nullptr;
-		if (ensure(UsedUniqueGroupNames.RemoveAndCopyValue(GroupName, Entry)))
-		{
-			Entry->NextFreeEntry = FirstFreeUniqueGroupName;
-			FirstFreeUniqueGroupName = Entry;
-		}
-
-		if (const int32* GroupIndex = SyncGroupMap.Find(GroupName))
-		{
-			FSyncGroupState& GroupState = SyncGroups[*GroupIndex];
+		FSyncGroupState& GroupState = SyncGroups[*GroupIndex];
 
 #if UE_DEBUG_SYNC_GROUPS
-			if (CDebugVarEnableSyncLog.GetValueOnAnyThread())
+		if (CDebugVarEnableSyncLog.GetValueOnAnyThread())
+		{
+			for (int32 GroupMemberIndex = GroupState.Members.Num() - 1; GroupMemberIndex >= 0; GroupMemberIndex--)
 			{
-				for (int32 GroupMemberIndex = GroupState.Members.Num() - 1; GroupMemberIndex >= 0; GroupMemberIndex--)
-				{
-					const FSyncGroupMember& GroupMember = GroupState.Members[GroupMemberIndex];
+				const FSyncGroupMember& GroupMember = GroupState.Members[GroupMemberIndex];
 
-					UE_VLOG_UELOG(GetHostObject(GroupMember.TraitPtr), LogAnimMarkerSync, Verbose,
-						TEXT("[%s] [%p] Left"),
-						*GroupState.GroupName.ToString(),
-						GroupMember.TraitPtr.GetNodeInstance());
-				}
+				UE_VLOG_UELOG(GetHostObject(GroupMember.TraitPtr), LogAnimMarkerSync, Verbose,
+					TEXT("[%s] [%p] Left"),
+					*GroupState.GroupName.ToString(),
+					GroupMember.TraitPtr.GetNodeInstance());
 			}
+		}
 #endif
 
-			// We are releasing this group, clear any lingering member entries
-			// as we might re-use the group name this update
-			GroupState.Members.Reset();
-		}
+		// We are releasing this group, clear any lingering member entries
+		// as we might re-use the group name this update
+		GroupState.Members.Reset();
 	}
+}
 
-	void FSyncGroupGraphInstanceComponent::PreUpdate(FExecutionContext& Context)
-	{
-		using namespace Private;
+void FSyncGroupGraphInstanceComponent::PreUpdate(UE::AnimNext::FExecutionContext& Context)
+{
+	using namespace UE::AnimNext::Private;
 
-		for (FSyncGroupState& GroupState : SyncGroups)
+	for (FSyncGroupState& GroupState : SyncGroups)
+	{
+		for (FSyncGroupMember& GroupMember : GroupState.Members)
 		{
-			for (FSyncGroupMember& GroupMember : GroupState.Members)
-			{
-				GroupMember.bIsActive = false;
-				GroupMember.bJustJoined = false;
-			}
-
-			GroupState.bIsActive = false;
-			GroupState.bJustFormed = GroupState.Members.IsEmpty();
+			GroupMember.bIsActive = false;
+			GroupMember.bJustJoined = false;
 		}
+
+		GroupState.bIsActive = false;
+		GroupState.bJustFormed = GroupState.Members.IsEmpty();
 	}
+}
 
-	void FSyncGroupGraphInstanceComponent::PostUpdate(FExecutionContext& Context)
-	{
-		using namespace Private;
+void FSyncGroupGraphInstanceComponent::PostUpdate(UE::AnimNext::FExecutionContext& Context)
+{
+	using namespace UE::AnimNext::Private;
 
-		// First purge stale entries and update our bookkeeping
-		int32 MaxNumMembers = 0;
+	// First purge stale entries and update our bookkeeping
+	int32 MaxNumMembers = 0;
+	{
+		bool bRebuildGroupIndexMap = false;
+		for (int32 SyncGroupIndex = SyncGroups.Num() - 1; SyncGroupIndex >= 0; --SyncGroupIndex)
 		{
-			bool bRebuildGroupIndexMap = false;
-			for (int32 SyncGroupIndex = SyncGroups.Num() - 1; SyncGroupIndex >= 0; --SyncGroupIndex)
+			FSyncGroupState& GroupState = SyncGroups[SyncGroupIndex];
+
+			if (GroupState.bIsActive)
 			{
-				FSyncGroupState& GroupState = SyncGroups[SyncGroupIndex];
+				// Whether or not we have an active member from a previous update
+				bool bHasOldActiveMember = false;
 
-				if (GroupState.bIsActive)
+				for (int32 GroupMemberIndex = GroupState.Members.Num() - 1; GroupMemberIndex >= 0; GroupMemberIndex--)
 				{
-					// Whether or not we have an active member from a previous update
-					bool bHasOldActiveMember = false;
-
-					for (int32 GroupMemberIndex = GroupState.Members.Num() - 1; GroupMemberIndex >= 0; GroupMemberIndex--)
-					{
-						FSyncGroupMember& GroupMember = GroupState.Members[GroupMemberIndex];
-						if (!GroupMember.bIsActive)
-						{
-#if UE_DEBUG_SYNC_GROUPS
-							if (CDebugVarEnableSyncLog.GetValueOnAnyThread())
-							{
-								UE_VLOG_UELOG(Context.GetHostObject(), LogAnimMarkerSync, Verbose,
-									TEXT("[%s] [%p] Left"),
-									*GroupState.GroupName.ToString(),
-									GroupMember.TraitPtr.GetNodeInstance());
-							}
-#endif
-
-							// This member is no longer active, remove it
-							GroupState.Members.RemoveAtSwap(GroupMemberIndex, EAllowShrinking::No);
-						}
-						else if (!GroupMember.bJustJoined)
-						{
-							bHasOldActiveMember = true;
-						}
-					}
-
-					// If we don't have any lingering members from a previous update, then we are considered
-					// as a fresh new group
-					if (!bHasOldActiveMember && !GroupState.Members.IsEmpty() && !GroupState.bJustFormed)
+					FSyncGroupMember& GroupMember = GroupState.Members[GroupMemberIndex];
+					if (!GroupMember.bIsActive)
 					{
 #if UE_DEBUG_SYNC_GROUPS
 						if (CDebugVarEnableSyncLog.GetValueOnAnyThread())
 						{
 							UE_VLOG_UELOG(Context.GetHostObject(), LogAnimMarkerSync, Verbose,
-								TEXT("[%s] Refreshed"),
-								*GroupState.GroupName.ToString());
+								TEXT("[%s] [%p] Left"),
+								*GroupState.GroupName.ToString(),
+								GroupMember.TraitPtr.GetNodeInstance());
 						}
 #endif
 
-						GroupState.bJustFormed = true;
+						// This member is no longer active, remove it
+						GroupState.Members.RemoveAtSwap(GroupMemberIndex, EAllowShrinking::No);
+					}
+					else if (!GroupMember.bJustJoined)
+					{
+						bHasOldActiveMember = true;
 					}
-
-					MaxNumMembers = FMath::Max<uint32>(MaxNumMembers, GroupState.Members.Num());
 				}
-				else
+
+				// If we don't have any lingering members from a previous update, then we are considered
+				// as a fresh new group
+				if (!bHasOldActiveMember && !GroupState.Members.IsEmpty() && !GroupState.bJustFormed)
 				{
 #if UE_DEBUG_SYNC_GROUPS
 					if (CDebugVarEnableSyncLog.GetValueOnAnyThread())
 					{
-						for (int32 GroupMemberIndex = GroupState.Members.Num() - 1; GroupMemberIndex >= 0; GroupMemberIndex--)
-						{
-							const FSyncGroupMember& GroupMember = GroupState.Members[GroupMemberIndex];
-
-							UE_VLOG_UELOG(Context.GetHostObject(), LogAnimMarkerSync, Verbose,
-								TEXT("[%s] [%p] Left"),
-								*GroupState.GroupName.ToString(),
-								GroupMember.TraitPtr.GetNodeInstance());
-						}
-
 						UE_VLOG_UELOG(Context.GetHostObject(), LogAnimMarkerSync, Verbose,
-							TEXT("[%s] Released"),
+							TEXT("[%s] Refreshed"),
 							*GroupState.GroupName.ToString());
 					}
 #endif
 
-					// This group is no longer active, remove it
-					SyncGroups.RemoveAtSwap(SyncGroupIndex, EAllowShrinking::No);
-					bRebuildGroupIndexMap = true;
+					GroupState.bJustFormed = true;
 				}
-			}
 
-			if (bRebuildGroupIndexMap)
+				MaxNumMembers = FMath::Max<uint32>(MaxNumMembers, GroupState.Members.Num());
+			}
+			else
 			{
-				// We purged one or more inactive groups, rebuild our map
-				SyncGroupMap.Reset();
-
-				const int32 NumSyncGroups = SyncGroups.Num();
-				for (int32 SyncGroupIndex = 0; SyncGroupIndex < NumSyncGroups; ++SyncGroupIndex)
+#if UE_DEBUG_SYNC_GROUPS
+				if (CDebugVarEnableSyncLog.GetValueOnAnyThread())
 				{
-					const FSyncGroupState& GroupState = SyncGroups[SyncGroupIndex];
-					SyncGroupMap.Add(GroupState.GroupName, SyncGroupIndex);
+					for (int32 GroupMemberIndex = GroupState.Members.Num() - 1; GroupMemberIndex >= 0; GroupMemberIndex--)
+					{
+						const FSyncGroupMember& GroupMember = GroupState.Members[GroupMemberIndex];
+
+						UE_VLOG_UELOG(Context.GetHostObject(), LogAnimMarkerSync, Verbose,
+							TEXT("[%s] [%p] Left"),
+							*GroupState.GroupName.ToString(),
+							GroupMember.TraitPtr.GetNodeInstance());
+					}
+
+					UE_VLOG_UELOG(Context.GetHostObject(), LogAnimMarkerSync, Verbose,
+						TEXT("[%s] Released"),
+						*GroupState.GroupName.ToString());
 				}
+#endif
+
+				// This group is no longer active, remove it
+				SyncGroups.RemoveAtSwap(SyncGroupIndex, EAllowShrinking::No);
+				bRebuildGroupIndexMap = true;
 			}
 		}
 
-		TArray<FSyncGroupMemberContext, TMemStackAllocator<>> MemberContexts;
-		MemberContexts.AddDefaulted(MaxNumMembers);
+		if (bRebuildGroupIndexMap)
+		{
+			// We purged one or more inactive groups, rebuild our map
+			SyncGroupMap.Reset();
 
-		FSyncGroupContext GroupContext;
+			const int32 NumSyncGroups = SyncGroups.Num();
+			for (int32 SyncGroupIndex = 0; SyncGroupIndex < NumSyncGroups; ++SyncGroupIndex)
+			{
+				const FSyncGroupState& GroupState = SyncGroups[SyncGroupIndex];
+				SyncGroupMap.Add(GroupState.GroupName, SyncGroupIndex);
+			}
+		}
+	}
 
-		// Now that we have discovered all groups and their memberships, we can perform synchronization
-		for (FSyncGroupState& GroupState : SyncGroups)
-		{
-			// Initialize our context for this new group
-			InitGroup(GroupState, MemberContexts, GroupContext);
+	TArray<FSyncGroupMemberContext, TMemStackAllocator<>> MemberContexts;
+	MemberContexts.AddDefaulted(MaxNumMembers);
 
-			// Find our leader for this group
-			FindLeaderIndex(Context, GroupContext);
+	FSyncGroupContext GroupContext;
+
+	// Now that we have discovered all groups and their memberships, we can perform synchronization
+	for (FSyncGroupState& GroupState : SyncGroups)
+	{
+		// Initialize our context for this new group
+		InitGroup(GroupState, MemberContexts, GroupContext);
 
-			// If the leader has sync markers, collect and filter them
-			// We'll retain only markers common to all members
-			BuildGroupState(Context, GroupContext);
+		// Find our leader for this group
+		FindLeaderIndex(Context, GroupContext);
 
-			// Advance the leader position as it determines the position of its followers
-			AdvanceLeader(Context, GroupContext);
+		// If the leader has sync markers, collect and filter them
+		// We'll retain only markers common to all members
+		BuildGroupState(Context, GroupContext);
 
-			// Advance the follower positions based on their leader
-			AdvanceFollowers(Context, GroupContext);
+		// Advance the leader position as it determines the position of its followers
+		AdvanceLeader(Context, GroupContext);
 
-			// Retain the leader progress as we might need it if a new leader joins the group during the next update
-			GroupState.PreviousLeaderTimelineState = GroupContext.LeaderTimelineState;
-			GroupState.PreviousLeaderSyncMarkers = GroupContext.Members[GroupContext.LeaderIndex].SyncMarkers;
-			GroupState.PreviousLeaderPhasePosition = GroupContext.LeaderPhaseEnd;
+		// Advance the follower positions based on their leader
+		AdvanceFollowers(Context, GroupContext);
+
+		// Retain the leader progress as we might need it if a new leader joins the group during the next update
+		GroupState.PreviousLeaderTimelineState = GroupContext.LeaderTimelineState;
+		GroupState.PreviousLeaderSyncMarkers = GroupContext.Members[GroupContext.LeaderIndex].SyncMarkers;
+		GroupState.PreviousLeaderPhasePosition = GroupContext.LeaderPhaseEnd;
 
 #if UE_DEBUG_SYNC_GROUPS
-			GroupState.PreviousLeaderIndex = GroupContext.LeaderIndex;
+		GroupState.PreviousLeaderIndex = GroupContext.LeaderIndex;
 #endif
-		}
 	}
 }
 
+
 #undef UE_DEBUG_SYNC_GROUPS
diff --git a/Engine/Plugins/Experimental/AnimNextAnimGraph/Source/AnimNextAnimGraph/Private/Graph/SyncGroup_GraphInstanceComponent.h b/Engine/Plugins/Experimental/AnimNextAnimGraph/Source/AnimNextAnimGraph/Private/Graph/SyncGroup_GraphInstanceComponent.h
--- a/Engine/Plugins/Experimental/AnimNextAnimGraph/Source/AnimNextAnimGraph/Private/Graph/SyncGroup_GraphInstanceComponent.h
+++ b/Engine/Plugins/Experimental/AnimNextAnimGraph/Source/AnimNextAnimGraph/Private/Graph/SyncGroup_GraphInstanceComponent.h
@@ -5,7 +5,8 @@
 #include "CoreMinimal.h"
 #include "TraitCore/TraitPtr.h"
 #include "TraitInterfaces/IUpdate.h"
-#include "Graph/GraphInstanceComponent.h"
+#include "Graph/UAFGraphInstanceComponent.h"
+#include "SyncGroup_GraphInstanceComponent.generated.h"
 
 namespace UE::AnimNext
 {
@@ -16,48 +17,56 @@ namespace UE::AnimNext
 	}
 
 	struct FSyncGroupParameters;
+}
 
-	/**
-	 * FSyncGroupGraphInstanceComponent
-	 *
-	 * This component maintains the necessary state to support group based synchronization.
-	 */
-	struct FSyncGroupGraphInstanceComponent : public FGraphInstanceComponent
-	{
-		DECLARE_ANIM_GRAPH_INSTANCE_COMPONENT(FSyncGroupGraphInstanceComponent)
+/**
+ * FSyncGroupGraphInstanceComponent
+ *
+ * This component maintains the necessary state to support group based synchronization.
+ */
+USTRUCT()
+struct FSyncGroupGraphInstanceComponent : public FUAFGraphInstanceComponent
+{
+	GENERATED_BODY()
 
-		explicit FSyncGroupGraphInstanceComponent(FAnimNextGraphInstance& InOwnerInstance);
-		~FSyncGroupGraphInstanceComponent();
+	// Defaulted constructors/copy-operators in cpp so fwd decls of members can work
+	FSyncGroupGraphInstanceComponent();
+	FSyncGroupGraphInstanceComponent(FSyncGroupGraphInstanceComponent&&);
+	FSyncGroupGraphInstanceComponent(const FSyncGroupGraphInstanceComponent&);
+	FSyncGroupGraphInstanceComponent& operator=(const FSyncGroupGraphInstanceComponent&);
+	FSyncGroupGraphInstanceComponent& operator=(FSyncGroupGraphInstanceComponent&&);
 
-		// Registers the specified trait with group based synchronization
-		void RegisterWithGroup(const FSyncGroupParameters& GroupParameters, const FWeakTraitPtr& TraitPtr, const FTraitUpdateState& TraitState);
+	virtual ~FSyncGroupGraphInstanceComponent();
 
-		// Create a unique group name suitable for spawned sub-graphs to self-synchronize
-		// Unique group names are a limited resource, when no longed needed they must be released
-		FName CreateUniqueGroupName();
+	// Registers the specified trait with group based synchronization
+	void RegisterWithGroup(const UE::AnimNext::FSyncGroupParameters& GroupParameters, const UE::AnimNext::FWeakTraitPtr& TraitPtr, const UE::AnimNext::FTraitUpdateState& TraitState);
 
-		// Releases a unique group name that is no longer needed. It will be recycled the
-		// next time one is needed
-		void ReleaseUniqueGroupName(FName GroupName);
+	// Create a unique group name suitable for spawned sub-graphs to self-synchronize
+	// Unique group names are a limited resource, when no longed needed they must be released
+	FName CreateUniqueGroupName();
 
-		// FGraphInstanceComponent impl
-		virtual void PreUpdate(FExecutionContext& Context) override;
-		virtual void PostUpdate(FExecutionContext& Context) override;
+	// Releases a unique group name that is no longer needed. It will be recycled the
+	// next time one is needed
+	void ReleaseUniqueGroupName(FName GroupName);
 
-	private:
-		// A map of sync group name -> group index
-		TMap<FName, int32> SyncGroupMap;
+	// FGraphInstanceComponent impl
+	virtual void PreUpdate(UE::AnimNext::FExecutionContext& Context) override;
+	virtual void PostUpdate(UE::AnimNext::FExecutionContext& Context) override;
 
-		// A list of groups and their data
-		TArray<Private::FSyncGroupState> SyncGroups;
+private:
+	// A map of sync group name -> group index
+	TMap<FName, int32> SyncGroupMap;
 
-		// The first free unique group name
-		Private::FSyncGroupUniqueName* FirstFreeUniqueGroupName = nullptr;
+	// A list of groups and their data
+	TArray<UE::AnimNext::Private::FSyncGroupState> SyncGroups;
 
-		// The map of currently used group names
-		TMap<FName, Private::FSyncGroupUniqueName*> UsedUniqueGroupNames;
+	// The first free unique group name
+	UE::AnimNext::Private::FSyncGroupUniqueName* FirstFreeUniqueGroupName = nullptr;
+
+	// The map of currently used group names
+	TMap<FName, UE::AnimNext::Private::FSyncGroupUniqueName*> UsedUniqueGroupNames;
+
+	// A counter tracking the next unique group name to allocate
+	int32 UniqueGroupNameCounter = 0;
+};
 
-		// A counter tracking the next unique group name to allocate
-		int32 UniqueGroupNameCounter = 0;
-	};
-}
diff --git a/Engine/Plugins/Experimental/AnimNextAnimGraph/Source/AnimNextAnimGraph/Private/Injection/AnimNextModuleInjectionComponent.cpp b/Engine/Plugins/Experimental/AnimNextAnimGraph/Source/AnimNextAnimGraph/Private/Injection/AnimNextModuleInjectionComponent.cpp
--- a/Engine/Plugins/Experimental/AnimNextAnimGraph/Source/AnimNextAnimGraph/Private/Injection/AnimNextModuleInjectionComponent.cpp
+++ b/Engine/Plugins/Experimental/AnimNextAnimGraph/Source/AnimNextAnimGraph/Private/Injection/AnimNextModuleInjectionComponent.cpp
@@ -9,12 +9,11 @@
 #include "Module/AnimNextModuleInstance.h"
 #include "Module/ModuleTaskContext.h"
 
-
-void FAnimNextModuleInjectionComponent::OnInitialize()
+FAnimNextModuleInjectionComponent::FAnimNextModuleInjectionComponent()
 {
 	FAnimNextModuleInstance& ModuleInstance = GetModuleInstance();
 	InjectionInfo = UE::AnimNext::FInjectionInfo(ModuleInstance);
-	
+
 	// Register re-injection for each (user) tick function
 	for(UE::AnimNext::FModuleEventTickFunction& TickFunction : ModuleInstance.GetTickFunctions())
 	{
@@ -47,59 +46,63 @@ void FAnimNextModuleInjectionComponent::OnInjectionEvent(UE::AnimNext::FInjectio
 	}
 
 	const FInjectionRequestArgs& RequestArgs = InEvent.Request->GetArgs();
-
-	FName FoundName = NAME_None;
-	TStructView<FAnimNextAnimGraph> InjectableGraph = InjectionInfo.FindInjectableGraphInstance(RequestArgs.Site, FoundName);
-	if(!InjectableGraph.IsValid())
+	
+	FAnimNextVariableReference FoundSite = InjectionInfo.FindInjectionSite(RequestArgs.Site);
+	if(FoundSite.IsNone())
 	{
-		UE_LOGFMT(LogAnimation, Warning, "Could not find injection site {SiteName} for injection request", RequestArgs.Site.DesiredSiteName);
+		UE_LOGFMT(LogAnimation, Warning, "Could not find injection site '{SiteName}' for injection request", RequestArgs.Site.DesiredSite.GetName());
 		InEvent.MarkConsumed();
 		return;
 	}
 
-	// Correct the name found above as we may have targeted NAME_None (any)
-	if(RequestArgs.Site.DesiredSiteName != FoundName)
+	// Correct the name found above as we may have targeted None (any)
+	if(RequestArgs.Site.DesiredSite != FoundSite)
 	{
-		InEvent.Request->GetMutableArgs().Site = FInjectionSite(FoundName);
+		InEvent.Request->GetMutableArgs().Site = FInjectionSite(FoundSite);
 	}
 
 	// Mark as handled so any additional trait events dont get processed at the module level
 	InEvent.MarkHandled();
 
 	// Store request as it will need to be re-applied each frame to ensure that bindings do not override it
-	FInjectionRecord& InjectionRecord = CurrentRequests.FindOrAdd(InEvent.Request->GetArgs().Site.DesiredSiteName);
+	FInjectionRecord& InjectionRecord = CurrentRequests.FindOrAdd(InEvent.Request->GetArgs().Site.DesiredSite);
 
-	FAnimNextAnimGraph& Graph = InjectableGraph.Get<FAnimNextAnimGraph>();
-	switch(RequestArgs.Type)
+	FAnimNextModuleInstance& ModuleInstance = GetModuleInstance();
+	EPropertyBagResult Result = ModuleInstance.AccessVariable<FAnimNextAnimGraph>(FoundSite, [this, &RequestArgs, &InEvent, &InjectionRecord](FAnimNextAnimGraph& InGraph)
 	{
-	case EAnimNextInjectionType::InjectObject:
-		// Bump serial number to identify this injection routing
-		ensureAlways(RequestArgs.Object != nullptr);
-		Graph.InjectionData.InjectionSerialNumber = IncrementSerialNumber();
-		InEvent.SerialNumber = Graph.InjectionData.InjectionSerialNumber;
-		InjectionRecord.SerialNumber = Graph.InjectionData.InjectionSerialNumber;
-		InjectionRecord.GraphRequest = InEvent.Request;
-
-		// Note we dont consume here, as we want the event to forward to the injection site trait
-		break;
-	case EAnimNextInjectionType::EvaluationModifier:
-		// We dont increment serial number when applying evaluation modifiers as we dont want to trigger a graph instantiation
-		ensureAlways(RequestArgs.EvaluationModifier != nullptr);
-		Graph.InjectionData.EvaluationModifier = RequestArgs.EvaluationModifier;
-		InjectionRecord.ModifierRequest = InEvent.Request;
-
-		// Update status
+		switch(RequestArgs.Type)
 		{
-			auto StatusUpdateEvent = MakeTraitEvent<FInjection_StatusUpdateEvent>();
-			StatusUpdateEvent->Request = InEvent.Request;
-			StatusUpdateEvent->Status = EInjectionStatus::Playing;
-			GetModuleInstance().QueueOutputTraitEvent(StatusUpdateEvent);
+		case EAnimNextInjectionType::InjectObject:
+			// Bump serial number to identify this injection routing
+			ensureAlways(RequestArgs.Object != nullptr);
+			InGraph.InjectionData.InjectionSerialNumber = IncrementSerialNumber();
+			InEvent.SerialNumber = InGraph.InjectionData.InjectionSerialNumber;
+			InjectionRecord.SerialNumber = InGraph.InjectionData.InjectionSerialNumber;
+			InjectionRecord.GraphRequest = InEvent.Request;
+
+			// Note we dont consume here, as we want the event to forward to the injection site trait
+			break;
+		case EAnimNextInjectionType::EvaluationModifier:
+			// We dont increment serial number when applying evaluation modifiers as we dont want to trigger a graph instantiation
+			ensureAlways(RequestArgs.EvaluationModifier != nullptr);
+			InGraph.InjectionData.EvaluationModifier = RequestArgs.EvaluationModifier;
+			InjectionRecord.ModifierRequest = InEvent.Request;
+
+			// Update status
+			{
+				auto StatusUpdateEvent = MakeTraitEvent<FInjection_StatusUpdateEvent>();
+				StatusUpdateEvent->Request = InEvent.Request;
+				StatusUpdateEvent->Status = EInjectionStatus::Playing;
+				GetModuleInstance().QueueOutputTraitEvent(StatusUpdateEvent);
+			}
+
+			// Evaluation modifiers are consumed straight away
+			InEvent.MarkConsumed();
+			break;
 		}
+	});
 
-		// Evaluation modifiers are consumed straight away
-		InEvent.MarkConsumed();
-		break;
-	}
+	ensureAlways(Result == EPropertyBagResult::Success);
 }
 
 void FAnimNextModuleInjectionComponent::OnUninjectionEvent(UE::AnimNext::FInjection_UninjectEvent& InEvent)
@@ -109,79 +112,80 @@ void FAnimNextModuleInjectionComponent::OnUninjectionEvent(UE::AnimNext::FInject
 	const FInjectionRequestArgs& RequestArgs = InEvent.Request->GetArgs();
 
 	// Injection modifiers can have private access to take over low-level tasks
-	FName FoundName = NAME_None;
-	TStructView<FAnimNextAnimGraph> InjectableGraph = InjectionInfo.FindInjectableGraphInstance(RequestArgs.Site, FoundName);
-	if(!InjectableGraph.IsValid())
+	FAnimNextVariableReference FoundSite = InjectionInfo.FindInjectionSite(RequestArgs.Site);
+	if(FoundSite.IsNone())
 	{
-		UE_LOGFMT(LogAnimation, Warning, "Could not find injection site {SiteName} for un-injection request", RequestArgs.Site.DesiredSiteName);
+		UE_LOGFMT(LogAnimation, Warning, "Could not find injection site '{SiteName}' for injection request", RequestArgs.Site.DesiredSite.GetName());
 		InEvent.MarkConsumed();
 		return;
 	}
 
-	FInjectionRecord& InjectionRecord = CurrentRequests.FindOrAdd(FoundName);
+	FInjectionRecord& InjectionRecord = CurrentRequests.FindOrAdd(FoundSite);
 
 	// Update graph and bump serial number to identify this un-injection routing
-	FAnimNextAnimGraph& Graph = InjectableGraph.Get<FAnimNextAnimGraph>();
-	switch(RequestArgs.Type)
+	FAnimNextModuleInstance& ModuleInstance = GetModuleInstance();
+	EPropertyBagResult Result = ModuleInstance.AccessVariable<FAnimNextAnimGraph>(FoundSite, [this, &RequestArgs, &InEvent, &InjectionRecord](FAnimNextAnimGraph& InGraph)
 	{
-	case EAnimNextInjectionType::InjectObject:
-		// Update graph and bump serial number to identify this injection routing
-		Graph.InjectionData.InjectionSerialNumber = IncrementSerialNumber();
-		InEvent.SerialNumber = Graph.InjectionData.InjectionSerialNumber;
-		InjectionRecord.GraphRequest.Reset();
-
-		// Note we dont consume here, as we want the event to forward to the injection site trait
-		break;
-	case EAnimNextInjectionType::EvaluationModifier:
-		// We dont increment serial number when applying evaluation modifiers as we dont want to trigger a graph instantiation
-		Graph.InjectionData.EvaluationModifier = nullptr;
-		InjectionRecord.ModifierRequest.Reset();
-
-		// Update status
+		switch(RequestArgs.Type)
 		{
-			auto StatusUpdateEvent = MakeTraitEvent<FInjection_StatusUpdateEvent>();
-			StatusUpdateEvent->Request = InEvent.Request;
-			StatusUpdateEvent->Status = EInjectionStatus::Completed;
-			GetModuleInstance().QueueOutputTraitEvent(StatusUpdateEvent);
+		case EAnimNextInjectionType::InjectObject:
+			// Update graph and bump serial number to identify this injection routing
+			InGraph.InjectionData.InjectionSerialNumber = IncrementSerialNumber();
+			InEvent.SerialNumber = InGraph.InjectionData.InjectionSerialNumber;
+			InjectionRecord.GraphRequest.Reset();
+
+			// Note we dont consume here, as we want the event to forward to the injection site trait
+			break;
+		case EAnimNextInjectionType::EvaluationModifier:
+			// We dont increment serial number when applying evaluation modifiers as we dont want to trigger a graph instantiation
+			InGraph.InjectionData.EvaluationModifier = nullptr;
+			InjectionRecord.ModifierRequest.Reset();
+
+			// Update status
+			{
+				auto StatusUpdateEvent = MakeTraitEvent<FInjection_StatusUpdateEvent>();
+				StatusUpdateEvent->Request = InEvent.Request;
+				StatusUpdateEvent->Status = EInjectionStatus::Completed;
+				GetModuleInstance().QueueOutputTraitEvent(StatusUpdateEvent);
+			}
+	
+			// Evaluation modifiers are consumed straight away
+			InEvent.MarkConsumed();
+			break;
 		}
-		
-		// Evaluation modifiers are consumed straight away
-		InEvent.MarkConsumed();
-		break;
-	}
+	});
+
+	ensureAlways(Result == EPropertyBagResult::Success);
 
 	if(!InjectionRecord.IsValid())
 	{
-		CurrentRequests.Remove(InEvent.Request->GetArgs().Site.DesiredSiteName);
+		CurrentRequests.Remove(InEvent.Request->GetArgs().Site.DesiredSite);
 	}
 }
 
 void FAnimNextModuleInjectionComponent::OnReapplyInjection(const UE::AnimNext::FModuleTaskContext& InContext)
 {
 	using namespace UE::AnimNext;
-
-	FAnimNextModuleInjectionComponent& Component = InContext.ModuleInstance->GetComponent<FAnimNextModuleInjectionComponent>();
 	
-	for(const TPair<FName, FInjectionRecord>& RequestPair : Component.CurrentRequests)
-	{
-		FName FoundName = NAME_None;
-		TStructView<FAnimNextAnimGraph> InjectableGraph = Component.InjectionInfo.FindInjectableGraphInstance(UE::AnimNext::FInjectionSite(RequestPair.Key), FoundName);
-		if(!InjectableGraph.IsValid())
-		{
-			continue;
-		}
+	FAnimNextModuleInjectionComponent& Component = InContext.ModuleInstance.GetComponent<FAnimNextModuleInjectionComponent>();
 
+	for(const TPair<FAnimNextVariableReference, FInjectionRecord>& RequestPair : Component.CurrentRequests)
+	{
 		// Re-apply this request, as it may have been overwritten by subsequent bindings/calculations
-		FAnimNextAnimGraph& Graph = InjectableGraph.Get<FAnimNextAnimGraph>();
-		if(RequestPair.Value.GraphRequest.IsValid())
+		EPropertyBagResult Result = InContext.ModuleInstance.AccessVariable<FAnimNextAnimGraph>(RequestPair.Key, [&RequestPair](FAnimNextAnimGraph& InGraph)
 		{
-			Graph.InjectionData.InjectionSerialNumber = RequestPair.Value.SerialNumber;
-		}
-		if(RequestPair.Value.ModifierRequest.IsValid())
-		{
-			const FInjectionRequestArgs& ModifierRequestArgs = RequestPair.Value.ModifierRequest->GetArgs();
-			Graph.InjectionData.EvaluationModifier = ModifierRequestArgs.EvaluationModifier;
-		}
+			if(RequestPair.Value.GraphRequest.IsValid())
+			{
+				InGraph.InjectionData.InjectionSerialNumber = RequestPair.Value.SerialNumber;
+			}
+			if(RequestPair.Value.ModifierRequest.IsValid())
+			{
+				const FInjectionRequestArgs& ModifierRequestArgs = RequestPair.Value.ModifierRequest->GetArgs();
+				InGraph.InjectionData.EvaluationModifier = ModifierRequestArgs.EvaluationModifier;
+			}
+		});
+
+		ensureAlways(Result == EPropertyBagResult::Success);
 	}
 }
 
@@ -197,7 +201,7 @@ uint32 FAnimNextModuleInjectionComponent::IncrementSerialNumber()
 
 void FAnimNextModuleInjectionComponent::AddStructReferencedObjects(FReferenceCollector& Collector)
 {
-	for (TPair<FName, FInjectionRecord>& RequestPair : CurrentRequests)
+	for (TPair<FAnimNextVariableReference, FInjectionRecord>& RequestPair : CurrentRequests)
 	{
 		RequestPair.Value.AddReferencedObjects(Collector);
 	}
diff --git a/Engine/Plugins/Experimental/AnimNextAnimGraph/Source/AnimNextAnimGraph/Private/Injection/AnimNextModuleInjectionComponent.h b/Engine/Plugins/Experimental/AnimNextAnimGraph/Source/AnimNextAnimGraph/Private/Injection/AnimNextModuleInjectionComponent.h
--- a/Engine/Plugins/Experimental/AnimNextAnimGraph/Source/AnimNextAnimGraph/Private/Injection/AnimNextModuleInjectionComponent.h
+++ b/Engine/Plugins/Experimental/AnimNextAnimGraph/Source/AnimNextAnimGraph/Private/Injection/AnimNextModuleInjectionComponent.h
@@ -3,9 +3,8 @@
 #pragma once
 
 #include "InjectionInfo.h"
-#include "DataInterface/AnimNextDataInterfaceHost.h"
 #include "Injection/InjectionRequest.h"
-#include "Module/AnimNextModuleInstanceComponent.h"
+#include "Module/UAFModuleInstanceComponent.h"
 #include "UObject/GCObject.h"
 #include "AnimNextModuleInjectionComponent.generated.h"
 
@@ -18,11 +17,11 @@ namespace UE::AnimNext
 
 // Module component that holds info about injection sites and routes injection requests
 USTRUCT()
-struct FAnimNextModuleInjectionComponent : public FAnimNextModuleInstanceComponent
+struct FAnimNextModuleInjectionComponent : public FUAFModuleInstanceComponent
 {
 	GENERATED_BODY()
 
-	FAnimNextModuleInjectionComponent() = default;
+	FAnimNextModuleInjectionComponent();
 
 	// Get the cached injection info for our module 
 	const UE::AnimNext::FInjectionInfo& GetInjectionInfo() const { return InjectionInfo; }
@@ -30,8 +29,7 @@ struct FAnimNextModuleInjectionComponent : public FAnimNextModuleInstanceCompone
 	void AddStructReferencedObjects(class FReferenceCollector& Collector);
 
 private:
-	// FAnimNextModuleInstanceComponent interface
-	virtual void OnInitialize() override;
+	// FUAFModuleInstanceComponent interface
 	virtual void OnTraitEvent(FAnimNextTraitEvent& Event) override;
 
 	void OnInjectionEvent(UE::AnimNext::FInjection_InjectEvent& InEvent);
@@ -60,7 +58,7 @@ private:
 	};
 	
 	// Currently-injected requests
-	TMap<FName, FInjectionRecord> CurrentRequests;
+	TMap<FAnimNextVariableReference, FInjectionRecord> CurrentRequests;
 
 	// Serial number used to identify forwarded requests
 	uint32 SerialNumber = 0;
diff --git a/Engine/Plugins/Experimental/AnimNextAnimGraph/Source/AnimNextAnimGraph/Private/Injection/GraphInstanceInjectionComponent.cpp b/Engine/Plugins/Experimental/AnimNextAnimGraph/Source/AnimNextAnimGraph/Private/Injection/GraphInstanceInjectionComponent.cpp
--- a/Engine/Plugins/Experimental/AnimNextAnimGraph/Source/AnimNextAnimGraph/Private/Injection/GraphInstanceInjectionComponent.cpp
+++ b/Engine/Plugins/Experimental/AnimNextAnimGraph/Source/AnimNextAnimGraph/Private/Injection/GraphInstanceInjectionComponent.cpp
@@ -3,13 +3,9 @@
 #include "GraphInstanceInjectionComponent.h"
 #include "Graph/AnimNextGraphInstance.h"
 
-namespace UE::AnimNext
-{
-
-FGraphInstanceInjectionComponent::FGraphInstanceInjectionComponent(FAnimNextGraphInstance& InOwnerInstance)
-	: FGraphInstanceComponent(InOwnerInstance)
-	, InjectionInfo(InOwnerInstance)
+FGraphInstanceInjectionComponent::FGraphInstanceInjectionComponent()
+	: FUAFGraphInstanceComponent()
+	, InjectionInfo(GetGraphInstance())
 {
 }
 
-}
diff --git a/Engine/Plugins/Experimental/AnimNextAnimGraph/Source/AnimNextAnimGraph/Private/Injection/GraphInstanceInjectionComponent.h b/Engine/Plugins/Experimental/AnimNextAnimGraph/Source/AnimNextAnimGraph/Private/Injection/GraphInstanceInjectionComponent.h
--- a/Engine/Plugins/Experimental/AnimNextAnimGraph/Source/AnimNextAnimGraph/Private/Injection/GraphInstanceInjectionComponent.h
+++ b/Engine/Plugins/Experimental/AnimNextAnimGraph/Source/AnimNextAnimGraph/Private/Injection/GraphInstanceInjectionComponent.h
@@ -4,26 +4,25 @@
 
 #include "CoreMinimal.h"
 #include "InjectionInfo.h"
-#include "Graph/GraphInstanceComponent.h"
+#include "Graph/UAFGraphInstanceComponent.h"
+#include "GraphInstanceInjectionComponent.generated.h"
 
-namespace UE::AnimNext
+/**
+ * FGraphInstanceInjectionComponent
+ * This component maintains injection info for a graph
+ */
+USTRUCT()
+struct FGraphInstanceInjectionComponent : public FUAFGraphInstanceComponent
 {
-	/**
-	 * FGraphInstanceInjectionComponent
-	 * This component maintains injection info for a graph
-	 */
-	struct FGraphInstanceInjectionComponent : public FGraphInstanceComponent
-	{
-		DECLARE_ANIM_GRAPH_INSTANCE_COMPONENT(FGraphInstanceInjectionComponent)
+	GENERATED_BODY()
 
-		explicit FGraphInstanceInjectionComponent(FAnimNextGraphInstance& InOwnerInstance);
+	FGraphInstanceInjectionComponent();
 
-		const FInjectionInfo& GetInjectionInfo() const
-		{
-			return InjectionInfo;
-		}
+	const UE::AnimNext::FInjectionInfo& GetInjectionInfo() const
+	{
+		return InjectionInfo;
+	}
 
-	private:
-		FInjectionInfo InjectionInfo;
-	};
-}
+private:
+	UE::AnimNext::FInjectionInfo InjectionInfo;
+};
diff --git a/Engine/Plugins/Experimental/AnimNextAnimGraph/Source/AnimNextAnimGraph/Private/Injection/InjectionCallbackProxy.cpp b/Engine/Plugins/Experimental/AnimNextAnimGraph/Source/AnimNextAnimGraph/Private/Injection/InjectionCallbackProxy.cpp
--- a/Engine/Plugins/Experimental/AnimNextAnimGraph/Source/AnimNextAnimGraph/Private/Injection/InjectionCallbackProxy.cpp
+++ b/Engine/Plugins/Experimental/AnimNextAnimGraph/Source/AnimNextAnimGraph/Private/Injection/InjectionCallbackProxy.cpp
@@ -2,6 +2,7 @@
 
 #include "Injection/InjectionCallbackProxy.h"
 #include "Animation/AnimSequence.h"
+#include "Blueprint/BlueprintExceptionInfo.h"
 #include "Component/AnimNextComponent.h"
 #include "GraphInterfaces/AnimNextNativeDataInterface_AnimSequencePlayer.h"
 #include "Injection/InjectionUtils.h"
@@ -10,25 +11,23 @@
 
 UInjectionCallbackProxy* UInjectionCallbackProxy::CreateProxyObjectForInjection(
 	UAnimNextComponent* AnimNextComponent,
-	FName SiteName,
+	FAnimNextVariableReference Site,
 	UObject* Object,
-	UAnimNextComponent* BindingComponent,
-	FInstancedStruct Payload,
+	FAnimNextFactoryParams FactoryParams,
 	UE::AnimNext::FInjectionBlendSettings BlendInSettings,
 	UE::AnimNext::FInjectionBlendSettings BlendOutSettings)
 {
 	UInjectionCallbackProxy* Proxy = NewObject<UInjectionCallbackProxy>();
 	Proxy->SetFlags(RF_StrongRefOnFrame);
-	Proxy->Inject(AnimNextComponent, SiteName, Object, BindingComponent, MoveTemp(Payload), BlendInSettings, BlendOutSettings);
+	Proxy->Inject(AnimNextComponent, Site, Object, MoveTemp(FactoryParams), BlendInSettings, BlendOutSettings);
 	return Proxy;
 }
 
 bool UInjectionCallbackProxy::Inject(
 	UAnimNextComponent* AnimNextComponent,
-	FName SiteName,
+	FAnimNextVariableReference Site,
 	UObject* Object,
-	UAnimNextComponent* BindingComponent,
-	FInstancedStruct&& Payload,
+	FAnimNextFactoryParams&& FactoryParams,
 	const UE::AnimNext::FInjectionBlendSettings& BlendInSettings,
 	const UE::AnimNext::FInjectionBlendSettings& BlendOutSettings)
 {
@@ -38,12 +37,11 @@ bool UInjectionCallbackProxy::Inject(
 	}
 
 	UE::AnimNext::FInjectionRequestArgs RequestArgs;
-	RequestArgs.Site = FAnimNextInjectionSite(SiteName);
+	RequestArgs.Site = FAnimNextInjectionSite(Site);
 	RequestArgs.Object = Object;
 	RequestArgs.BlendInSettings = BlendInSettings;
 	RequestArgs.BlendOutSettings = BlendOutSettings;
-	RequestArgs.BindingModuleHandle = BindingComponent ? BindingComponent->GetModuleHandle() : UE::AnimNext::FModuleHandle();
-	RequestArgs.Payload.AddNative(MoveTemp(Payload));
+	RequestArgs.FactoryParams = MoveTemp(FactoryParams);
 
 	UE::AnimNext::FInjectionLifetimeEvents LifetimeEvents;
 	LifetimeEvents.OnCompleted.BindUObject(this, &UInjectionCallbackProxy::OnInjectionCompleted);
@@ -75,6 +73,85 @@ EUninjectionResult UInjectionCallbackProxy::Uninject()
 	return EUninjectionResult::Succeeded;
 }
 
+void UInjectionCallbackProxy::SetVariable(const FAnimNextVariableReference& Variable, const int32& Value)
+{
+	checkNoEntry();
+}
+
+DEFINE_FUNCTION(UInjectionCallbackProxy::execSetVariable)
+{
+	using namespace UE::AnimNext;
+
+	if(!P_THIS->PlayingRequest.IsValid())
+	{
+		P_FINISH;
+
+		FBlueprintExceptionInfo ExceptionInfo(
+			EBlueprintExceptionType::NonFatalError,
+			NSLOCTEXT("InjectionCallbackProxy", "InjectionCallbackProxy_SetVariableInvalidRequestWarning", "Invalid injection request when calling Set Variable")
+		);
+
+		FBlueprintCoreDelegates::ThrowScriptException(P_THIS, Stack, ExceptionInfo);
+		return;
+	}
+	
+	P_GET_STRUCT_REF(FAnimNextVariableReference, Variable);
+
+	const FProperty* VariableProperty = Variable.ResolveProperty();
+	if (VariableProperty == nullptr)
+	{
+		P_FINISH;
+
+		FBlueprintExceptionInfo ExceptionInfo(
+			EBlueprintExceptionType::AbortExecution,
+			FText::Format(NSLOCTEXT("InjectionCallbackProxy", "InjectionCallbackProxy_SetVariablePropertyError", "Failed to resolve the property {0} of variable reference"), FText::FromName(Variable.GetName()))
+		);
+
+		FBlueprintCoreDelegates::ThrowScriptException(P_THIS, Stack, ExceptionInfo);
+		return;
+	}
+
+	if (Variable.GetName() == NAME_None || Variable.GetObject() == nullptr)
+	{
+		P_FINISH;
+
+		FBlueprintExceptionInfo ExceptionInfo(
+			EBlueprintExceptionType::NonFatalError,
+			NSLOCTEXT("InjectionCallbackProxy", "InjectionCallbackProxy_SetVariableInvalidWarning", "Invalid variable supplied to Set Variable")
+		);
+
+		FBlueprintCoreDelegates::ThrowScriptException(P_THIS, Stack, ExceptionInfo);
+		return;
+	}
+
+	// Read wildcard Value input.
+	Stack.MostRecentPropertyAddress = nullptr;
+	Stack.MostRecentPropertyContainer = nullptr;
+
+	TUniquePtr<uint8> TempStorage((uint8*)FMemory::Malloc(VariableProperty->GetElementSize(), VariableProperty->GetMinAlignment()));
+	VariableProperty->InitializeValue(TempStorage.Get());
+
+	Stack.StepCompiledIn(TempStorage.Get(), VariableProperty->GetClass());
+
+	P_FINISH;
+
+	P_NATIVE_BEGIN;
+
+	P_THIS->PlayingRequest->QueueTask([Variable, Storage = MoveTemp(TempStorage)](const FInstanceTaskContext& InContext)
+	{
+		const FProperty* VariableProperty = Variable.ResolveProperty();
+		check(VariableProperty);	// If this fires the property has been destroyed before we can destroy the temp memory storage
+
+		FAnimNextParamType Type = FAnimNextParamType::FromProperty(VariableProperty);
+		InContext.SetVariableInternal(Variable, Type, TConstArrayView<uint8>(Storage.Get(), VariableProperty->GetElementSize()));
+
+		// Clean up
+		VariableProperty->DestroyValue(Storage.Get());
+	});
+
+	P_NATIVE_END;
+}
+
 void UInjectionCallbackProxy::Cancel()
 {
 	Super::Cancel();
diff --git a/Engine/Plugins/Experimental/AnimNextAnimGraph/Source/AnimNextAnimGraph/Private/Injection/InjectionInfo.cpp b/Engine/Plugins/Experimental/AnimNextAnimGraph/Source/AnimNextAnimGraph/Private/Injection/InjectionInfo.cpp
--- a/Engine/Plugins/Experimental/AnimNextAnimGraph/Source/AnimNextAnimGraph/Private/Injection/InjectionInfo.cpp
+++ b/Engine/Plugins/Experimental/AnimNextAnimGraph/Source/AnimNextAnimGraph/Private/Injection/InjectionInfo.cpp
@@ -1,16 +1,16 @@
 ﻿// Copyright Epic Games, Inc. All Rights Reserved.
 
 #include "InjectionInfo.h"
-#include "Injection/InjectionSite.h"
 
-#include "DataInterface/AnimNextDataInterfaceInstance.h"
+#include "AnimNextRigVMAsset.h"
+#include "Injection/InjectionSite.h"
+#include "UAFAssetInstance.h"
 #include "Graph/AnimNextAnimGraph.h"
-#include "StructUtils/PropertyBag.h"
 
 namespace UE::AnimNext
 {
 
-FInjectionInfo::FInjectionInfo(const FAnimNextDataInterfaceInstance& InInstance)
+FInjectionInfo::FInjectionInfo(const FUAFAssetInstance& InInstance)
 	: Instance(&InInstance)
 {
 	CacheInfo();
@@ -19,92 +19,40 @@ FInjectionInfo::FInjectionInfo(const FAnimNextDataInterfaceInstance& InInstance)
 void FInjectionInfo::CacheInfo() const
 {
 	check(Instance);
-	check(Instance->GetDataInterface());
-
-	DefaultInjectableGraphIndex = INDEX_NONE;
-	InjectableGraphs.Reset();
-
-	const UPropertyBag* Struct = Instance->Variables.GetPropertyBagStruct();
-	const uint8* Memory = Instance->Variables.GetValue().GetMemory();
-	if(Struct == nullptr || Memory == nullptr)
-	{
-		return;
-	}
-
-	// Find the default value in our set of variables
-	// Note we dont recurse here as the struct views we cache cannot be relocated (e.g. held in an array) due to them being referenced
-	// elsewhere via raw ptrs in RigVM memory handles
-	int32 VariableIndex = 0;
-	for (TPropertyValueIterator<FProperty> It(Struct, Memory, EPropertyValueIteratorFlags::NoRecursion); It; ++It, ++VariableIndex)
-	{
-		const FStructProperty* Property = CastField<FStructProperty>(It.Key());
-		if(Property == nullptr)
-		{
-			continue;
-		}
-
-		if(Property->Struct->IsChildOf(FAnimNextAnimGraph::StaticStruct()))
-		{
-			const FAnimNextAnimGraph* Value = static_cast<const FAnimNextAnimGraph*>(It.Value());
-			FName GraphName = Property->GetFName();
-			TStructView<FAnimNextAnimGraph> GraphStructView(const_cast<FAnimNextAnimGraph&>(*Value));
+	check(Instance->GetAsset<UAnimNextRigVMAsset>());
+	DefaultInjectionSite = Instance->GetAsset<UAnimNextRigVMAsset>()->DefaultInjectionSite;
 
-			int32 GraphIndex = InjectableGraphs.Emplace(GraphName, GraphStructView);
+	InjectionSites.Reset();
+	Instance->GetAllVariablesOfType<FAnimNextAnimGraph>(InjectionSites);
 
-			if(VariableIndex == Instance->GetDataInterface()->DefaultInjectionSiteIndex)
-			{
-				DefaultInjectableGraphIndex = GraphIndex;
-			}
-		}
-	}
-
-	if(DefaultInjectableGraphIndex == INDEX_NONE && InjectableGraphs.Num() > 0)
-	{
-		DefaultInjectableGraphIndex = 0;
-	}
-}
-
-bool FInjectionInfo::GetDefaultInjectableGraph(FName& OutName, TStructView<FAnimNextAnimGraph>& OutGraph) const
-{
-	if(DefaultInjectableGraphIndex == INDEX_NONE || !InjectableGraphs.IsValidIndex(DefaultInjectableGraphIndex))
+	if(DefaultInjectionSite.IsNone() && InjectionSites.Num() > 0)
 	{
-		return false;
+		DefaultInjectionSite = InjectionSites[0];
 	}
-
-	const FInjectableGraphInfo& InjectableGraph = InjectableGraphs[DefaultInjectableGraphIndex];
-	OutName = InjectableGraph.Name;
-	OutGraph = InjectableGraph.StructView;
-
-	return OutName != NAME_None && OutGraph.IsValid();
 }
 
-TStructView<FAnimNextAnimGraph> FInjectionInfo::FindInjectableGraphInstance(const FInjectionSite& InSite, FName& OutRealName) const
+FAnimNextVariableReference FInjectionInfo::FindInjectionSite(const FInjectionSite& InSite) const
 {
-	FName SiteName = InSite.DesiredSiteName;
-	if(SiteName != NAME_None)
+	if(!InSite.DesiredSite.IsNone())
 	{
-		// Linear search all the injectable graphs for the name
-		const int32 InjectableGraphNum = InjectableGraphs.Num();
-		for(int32 InjectableGraphIndex = 0; InjectableGraphIndex < InjectableGraphNum; ++InjectableGraphIndex)
+		// Linear search all the injection sites
+		for(const FAnimNextVariableReference& InjectionSite : InjectionSites)
 		{
-			const FInjectableGraphInfo& InjectableGraphInfo = InjectableGraphs[InjectableGraphIndex];
-			if(SiteName == InjectableGraphInfo.Name)
+			if(InSite.DesiredSite == InjectionSite)
 			{
-				OutRealName = SiteName;
-				return InjectableGraphInfo.StructView;
+				return InjectionSite;
 			}
 		}
 	}
 
-	const bool bUseModuleDefault = SiteName == NAME_None || InSite.bUseModuleFallback;
-	if(bUseModuleDefault && InjectableGraphs.IsValidIndex(DefaultInjectableGraphIndex))
+	// Not found - optionally use fallback
+	const bool bUseModuleDefault = InSite.DesiredSite.IsNone() || InSite.bUseModuleFallback;
+	if(bUseModuleDefault && !DefaultInjectionSite.IsNone())
 	{
-		const FInjectableGraphInfo& InjectableGraphInfo = InjectableGraphs[DefaultInjectableGraphIndex];
-		OutRealName = InjectableGraphInfo.Name;
-		return InjectableGraphInfo.StructView;
+		return DefaultInjectionSite;
 	}
 
-	return TStructView<FAnimNextAnimGraph>();
+	return FAnimNextVariableReference();
 }
 
 }
diff --git a/Engine/Plugins/Experimental/AnimNextAnimGraph/Source/AnimNextAnimGraph/Private/Injection/InjectionInfo.h b/Engine/Plugins/Experimental/AnimNextAnimGraph/Source/AnimNextAnimGraph/Private/Injection/InjectionInfo.h
--- a/Engine/Plugins/Experimental/AnimNextAnimGraph/Source/AnimNextAnimGraph/Private/Injection/InjectionInfo.h
+++ b/Engine/Plugins/Experimental/AnimNextAnimGraph/Source/AnimNextAnimGraph/Private/Injection/InjectionInfo.h
@@ -3,9 +3,9 @@
 #pragma once
 
 #include "CoreMinimal.h"
-#include "StructUtils/StructView.h"
+#include "Variables/AnimNextVariableReference.h"
 
-struct FAnimNextDataInterfaceInstance;
+struct FUAFAssetInstance;
 struct FAnimNextAnimGraph;
 struct FAnimNextModuleInjectionComponent;
 struct FAnimNextInjectionSite;
@@ -14,61 +14,38 @@ namespace UE::AnimNext
 {
 using FInjectionSite = FAnimNextInjectionSite;
 
-// Info used to track injection sites for a data interface instance (graph, module etc)
+// Info used to track injection sites for an AnimNext instance (graph, module etc)
 struct FInjectionInfo
 {
 	FInjectionInfo() = default;
 
-	explicit FInjectionInfo(const FAnimNextDataInterfaceInstance& InInstance);
+	explicit FInjectionInfo(const FUAFAssetInstance& InInstance);
 
-	// Get the default injectable graph data
-	bool GetDefaultInjectableGraph(FName& OutName, TStructView<FAnimNextAnimGraph>& OutGraph) const;
-
-	// Find an injectable graph instance by name.
-	// @param   InSite            The injection site. If bUseModuleFallback is true, the default site will be returned if the site does not exist
-	// @param   OutRealName       The name of the actual injection site we found (in the case NAME_None was passed)
-	TStructView<FAnimNextAnimGraph> FindInjectableGraphInstance(const FInjectionSite& InSite, FName& OutRealName) const;
-
-	// Iterate each injectable graph instance
-	template<typename PredicateType>
-	void ForEachInjectableGraphInstance(PredicateType InPredicate) const
+	// Get the default injection site
+	FAnimNextVariableReference GetDefaultInjectionSite() const
 	{
-		for(const FInjectableGraphInfo& InjectableGraphInfo : InjectableGraphs)
-		{
-			InPredicate(InjectableGraphInfo.Name, InjectableGraphInfo.StructView);
-		}
+		return DefaultInjectionSite;
 	}
 
+	// Find an injectable graph variable by name.
+	// @param   InSite            The injection site.
+	// @return The actual injection site we found (in the case an injection site of None was passed)
+	FAnimNextVariableReference FindInjectionSite(const FInjectionSite& InSite) const;
+
 private:
 	void CacheInfo() const;
 
 private:
 	friend ::FAnimNextModuleInjectionComponent;
 
-	struct FInjectableGraphInfo
-	{
-		FInjectableGraphInfo(FName InName, TStructView<FAnimNextAnimGraph> InStructView)
-			: StructView(InStructView)
-			, Name(InName)
-		{}
-
-		TStructView<FAnimNextAnimGraph> StructView;
-		FName Name;
-	};
-
-	// Lookup of injectable graphs by name
-	// Lookup is a linear search for now as we expect there to be 'not many' of these things. If we find this is a bottleneck, we can evaluate the
-	// performance/memory cost of a TMap
-	mutable TArray<FInjectableGraphInfo> InjectableGraphs;
+	// All injection site variables
+	mutable TArray<FAnimNextVariableReference> InjectionSites;
 
 	// Instance we are tracking
-	const FAnimNextDataInterfaceInstance* Instance = nullptr;
-
-	// The name of the default injectable graph for our module (user-adjustable via SetDefaultInjectionSite) 
-	FName DefaultInjectableGraphName;
+	const FUAFAssetInstance* Instance = nullptr;
 
-	// Index into InjectableGraphs of the default graph
-	mutable int32 DefaultInjectableGraphIndex = INDEX_NONE;
+	// The default injection site for our instance (user-adjustable via SetDefaultInjectionSite) 
+	mutable FAnimNextVariableReference DefaultInjectionSite;
 };
 
 }
diff --git a/Engine/Plugins/Experimental/AnimNextAnimGraph/Source/AnimNextAnimGraph/Private/Injection/InjectionRequest.cpp b/Engine/Plugins/Experimental/AnimNextAnimGraph/Source/AnimNextAnimGraph/Private/Injection/InjectionRequest.cpp
--- a/Engine/Plugins/Experimental/AnimNextAnimGraph/Source/AnimNextAnimGraph/Private/Injection/InjectionRequest.cpp
+++ b/Engine/Plugins/Experimental/AnimNextAnimGraph/Source/AnimNextAnimGraph/Private/Injection/InjectionRequest.cpp
@@ -35,8 +35,6 @@ namespace UE::AnimNext
 			return false;	// No world to play in
 		}
 
-		SetupBindingModuleDependencies(World, InHandle, InRequestArgs);
-
 		RequestArgs = MoveTemp(InRequestArgs);
 		LifetimeEvents = MoveTemp(InLifetimeEvents);
 		WeakHost = InHost;
@@ -91,9 +89,9 @@ namespace UE::AnimNext
 		switch(InRequestArgs.Type)
 		{
 		case EAnimNextInjectionType::InjectObject:
-			if(InRequestArgs.Site.DesiredSiteName == NAME_None)
+			if(InRequestArgs.Site.DesiredSite.IsNone())
 			{
-				UE_LOGFMT(LogAnimation, Warning, "FInjectionRequest::ValidateArgs: Missing injection site name");
+				UE_LOGFMT(LogAnimation, Warning, "FInjectionRequest::ValidateArgs: Missing injection site");
 				return false;
 			}
 			if(InRequestArgs.Object == nullptr)
@@ -114,32 +112,6 @@ namespace UE::AnimNext
 		return true;
 	}
 
-	void FInjectionRequest::SetupBindingModuleDependencies(UWorld* InWorld, FModuleHandle InHandle, const FInjectionRequestArgs& InRequestArgs)
-	{
-		check(IsInGameThread());
-		check(InWorld);
-
-		if(InRequestArgs.BindingModuleHandle.IsValid())
-		{
-			UAnimNextWorldSubsystem* Subsystem = InWorld->GetSubsystem<UAnimNextWorldSubsystem>();
-			check(Subsystem);
-			Subsystem->AddModuleDependencyHandle(InRequestArgs.BindingModuleHandle, InHandle);
-		}
-	}
-
-	void FInjectionRequest::RemoveBindingModuleDependencies(UWorld* InWorld, FModuleHandle InHandle, const FInjectionRequestArgs& InRequestArgs)
-	{
-		check(IsInGameThread());
-		check(InWorld);
-
-		if(InRequestArgs.BindingModuleHandle.IsValid())
-		{
-			UAnimNextWorldSubsystem* Subsystem = InWorld->GetSubsystem<UAnimNextWorldSubsystem>();
-			check(Subsystem);
-			Subsystem->RemoveModuleDependencyHandle(InRequestArgs.BindingModuleHandle, InHandle);
-		}
-	}
-
 	const FInjectionRequestArgs& FInjectionRequest::GetArgs() const
 	{
 		return RequestArgs;
@@ -228,11 +200,6 @@ namespace UE::AnimNext
 			// Maintain our interrupted status if it was present
 			Status = EInjectionStatus::Completed | (Status & EInjectionStatus::Interrupted);
 			LifetimeEvents.OnCompleted.ExecuteIfBound(*this);
-			// Remove any dependencies we have set up
-			if (UWorld* World = WeakWorld.Get())
-			{
-				RemoveBindingModuleDependencies(World, Handle, RequestArgs);
-			}
 			break;
 		case EInjectionStatus::Expired:
 			ensureMsgf(Status == EInjectionStatus::Pending, TEXT("Expected Injection status to be pending, found: %u"), (uint32)Status);
@@ -260,4 +227,20 @@ namespace UE::AnimNext
 	{
 		Collector.AddPropertyReferencesWithStructARO(FInjectionRequestArgs::StaticStruct(), &RequestArgs);
 	}
+
+	void FInjectionRequest::QueueTask(FUniqueInstanceTask&& InTask)
+	{
+		TaskQueue.Enqueue(MoveTemp(InTask));
+	}
+
+	void FInjectionRequest::FlushTasks(const FInstanceTaskContext& InContext)
+	{
+		// Flush any queued tasks
+		while (!TaskQueue.IsEmpty())
+		{
+			TOptional<FUniqueInstanceTask> Task = TaskQueue.Dequeue();
+			check(Task.IsSet());
+			Task.GetValue()(InContext);
+		}
+	}
 }
