AnimNext: Pass-by-ref runtime refactor

diff --git a/Engine/Plugins/Experimental/AnimNextAnimGraph/Source/AnimNextAnimGraph/Private/Injection/InjectionSiteTrait.cpp b/Engine/Plugins/Experimental/AnimNextAnimGraph/Source/AnimNextAnimGraph/Private/Injection/InjectionSiteTrait.cpp
--- a/Engine/Plugins/Experimental/AnimNextAnimGraph/Source/AnimNextAnimGraph/Private/Injection/InjectionSiteTrait.cpp
+++ b/Engine/Plugins/Experimental/AnimNextAnimGraph/Source/AnimNextAnimGraph/Private/Injection/InjectionSiteTrait.cpp
@@ -2,8 +2,9 @@
 
 #include "Injection/InjectionSiteTrait.h"
 
-#include "AnimNextDataInterfacePayload.h"
 #include "AnimNextModuleInjectionComponent.h"
+#include "InstanceTaskContext.h"
+#include "Factory/AnimGraphFactory.h"
 #include "TraitCore/ExecutionContext.h"
 #include "TraitCore/NodeInstance.h"
 #include "TraitInterfaces/ITimeline.h"
@@ -94,7 +95,7 @@ namespace UE::AnimNext
 		if (Graph.InjectionData.InjectionSerialNumber == Event.SerialNumber)
 		{
 			const FAnimNextInjectionRequestArgs& RequestArgs = Event.Request->GetArgs();
-			ensureMsgf(!InstanceData->PendingRequest.IsValid(), TEXT("Injection site %s already contained a pending request, it will be overwritten"), *RequestArgs.Site.DesiredSiteName.ToString());
+			ensureMsgf(!InstanceData->PendingRequest.IsValid(), TEXT("Injection site %s already contained a pending request, it will be overwritten"), *RequestArgs.Site.DesiredSite.GetName().ToString());
 
 			// Overwrite any request we might have, we'll pick it up on the next update
 			InstanceData->PendingRequest.Reset();
@@ -167,23 +168,19 @@ namespace UE::AnimNext
 			const FInjectionRequestPtr Request = InstanceData->PendingRequest.Request;
 			InstanceData->PendingRequest.Reset();
 
+			bool bRequestValid = false;
 			IBlendStack::FGraphRequest BlendStackRequest;
 			EInjectionTrackerState State = EInjectionTrackerState::Inactive;
 			if (Request && Request->GetArgs().Object)
 			{
 				// Push request from an event
-				const FAnimNextInjectionRequestArgs& RequestArgs = Request->GetArgs();
 				State = EInjectionTrackerState::ActiveEvent;
 
-				BlendStackRequest.GraphPayload = MoveTemp(Request->GetMutableArgs().Payload);
-
+				FAnimNextInjectionRequestArgs& RequestArgs = Request->GetMutableArgs();
 				BlendStackRequest.Type = IBlendStack::EGraphRequestType::Owned;
-				BlendStackRequest.AnimationGraph = IGraphFactory::GetGraphFromObjectWithFallback(Context, Binding, RequestArgs.Object, BlendStackRequest.GraphPayload);
-				if (BlendStackRequest.AnimationGraph == nullptr)
-				{
-					BlendStackRequest.GraphPayload = FAnimNextDataInterfacePayload();
-				}
-				BlendStackRequest.BindingModuleHandle = RequestArgs.BindingModuleHandle;
+				BlendStackRequest.FactoryParams = MoveTemp(RequestArgs.FactoryParams);
+				BlendStackRequest.AnimationGraph = IGraphFactory::GetOrBuildGraph(Context, Binding, RequestArgs.Object, BlendStackRequest.FactoryParams);
+				bRequestValid = BlendStackRequest.AnimationGraph != nullptr;
 				BlendStackRequest.BlendArgs = RequestArgs.BlendInSettings.Blend;
 				BlendStackRequest.BlendMode = TranslateBlendMode(RequestArgs.BlendInSettings.BlendMode);
 			}
@@ -193,12 +190,8 @@ namespace UE::AnimNext
 				State = EInjectionTrackerState::ActivePin;
 
 				BlendStackRequest.Type = IBlendStack::EGraphRequestType::Owned;
-				BlendStackRequest.AnimationGraph = IGraphFactory::GetGraphFromObjectWithFallback(Context, Binding, Graph.Asset, BlendStackRequest.GraphPayload);
-				if (BlendStackRequest.AnimationGraph == nullptr)
-				{
-					BlendStackRequest.GraphPayload = FAnimNextDataInterfacePayload();
-				}
-
+				BlendStackRequest.AnimationGraph = IGraphFactory::GetOrBuildGraph(Context, Binding, Graph.Asset, BlendStackRequest.FactoryParams);
+				bRequestValid = BlendStackRequest.AnimationGraph != nullptr;
 				const FAnimNextInjectionBlendSettings& DefaultBlendInSettings = SharedData->GetDefaultBlendInSettings(Binding);
 				BlendStackRequest.BlendArgs = DefaultBlendInSettings.Blend;
 				BlendStackRequest.BlendMode = TranslateBlendMode(DefaultBlendInSettings.BlendMode);
@@ -208,6 +201,7 @@ namespace UE::AnimNext
 				// No request, or invalid/unhandled, so push our source child
 				State = EInjectionTrackerState::ActiveSource;
 
+				bRequestValid = true;
 				BlendStackRequest.Type = IBlendStack::EGraphRequestType::Child;
 				BlendStackRequest.ChildPtr = Context.AllocateNodeInstance(Binding, SharedData->Source);
 
@@ -216,45 +210,47 @@ namespace UE::AnimNext
 				BlendStackRequest.BlendMode = TranslateBlendMode(DefaultBlendInSettings.BlendMode);
 			}
 
-			if (bHasActiveSubGraph)
+			if (bRequestValid)
 			{
-				ensure(InstanceData->InjectionTrackers.IsValidIndex(ActiveChildIndex));
-
-				// Queue our status update
-				const FInjectionTracker& ActiveInjectionTracker = InstanceData->InjectionTrackers[ActiveChildIndex];
-				if (ActiveInjectionTracker.State == EInjectionTrackerState::ActiveEvent)
+				if (bHasActiveSubGraph)
 				{
-					auto StatusUpdateEvent = MakeTraitEvent<FInjection_StatusUpdateEvent>();
-					StatusUpdateEvent->Request = ActiveInjectionTracker.Request;
-					StatusUpdateEvent->Status = EInjectionStatus::Playing | EInjectionStatus::Interrupted;
+					ensure(InstanceData->InjectionTrackers.IsValidIndex(ActiveChildIndex));
 
-					Context.RaiseOutputTraitEvent(StatusUpdateEvent);
+					// Queue our status update
+					const FInjectionTracker& ActiveInjectionTracker = InstanceData->InjectionTrackers[ActiveChildIndex];
+					if (ActiveInjectionTracker.State == EInjectionTrackerState::ActiveEvent)
+					{
+						auto StatusUpdateEvent = MakeTraitEvent<FInjection_StatusUpdateEvent>();
+						StatusUpdateEvent->Request = ActiveInjectionTracker.Request;
+						StatusUpdateEvent->Status = EInjectionStatus::Playing | EInjectionStatus::Interrupted;
+
+						Context.RaiseOutputTraitEvent(StatusUpdateEvent);
+					}
+				}
+				else
+				{
+					// No current active subgraph, so blend instantly so we either:
+					// - Ensure that source doesnt blend from refpose on first becoming relevant
+					// - Dont blend from source when we first become relevant with a valid request
+					BlendStackRequest.BlendArgs.BlendTime = 0.0f;
 				}
-			}
-			else
-			{
-				// No current active subgraph, so blend instantly so we either:
-				// - Ensure that source doesnt blend from refpose on first becoming relevant
-				// - Dont blend from source when we first become relevant with a valid request
-				BlendStackRequest.BlendArgs.BlendTime = 0.0f;
-			}
 
-			// Track a new request
-			InstanceData->PendingTracker.Initialize(Request, State);
-			ActiveChildIndex = BlendStackTrait.PushGraph(Context, MoveTemp(BlendStackRequest));
+				// Track a new request
+				InstanceData->PendingTracker.Initialize(Request, State);
+				ActiveChildIndex = BlendStackTrait.PushGraph(Context, MoveTemp(BlendStackRequest));
 
-			// Make sure we tracked it after blending
-			ensure(!InstanceData->PendingTracker.IsRequestValid());
-			ensure(InstanceData->InjectionTrackers.IsValidIndex(ActiveChildIndex) &&
-				InstanceData->InjectionTrackers[ActiveChildIndex].Request == Request && 
-				InstanceData->InjectionTrackers[ActiveChildIndex].State == State);
+				// Make sure we tracked it after blending
+				ensure(!InstanceData->PendingTracker.IsRequestValid());
+				ensure(InstanceData->InjectionTrackers.IsValidIndex(ActiveChildIndex) &&
+					InstanceData->InjectionTrackers[ActiveChildIndex].Request == Request && 
+					InstanceData->InjectionTrackers[ActiveChildIndex].State == State);
 
-			bJustTransitioned = true;
+				bJustTransitioned = true;
+			}
 		}
 
 		float CurrentRequestTimeLeft = FLT_MAX;
 
-		// Broadcast our timeline progress
 		const int32 NumInjectionTrackers = InstanceData->InjectionTrackers.Num();
 		for (int32 TrackerIndex = 0; TrackerIndex < NumInjectionTrackers; ++TrackerIndex)
 		{
@@ -264,6 +260,13 @@ namespace UE::AnimNext
 				continue;	// We don't care about this injection request
 			}
 
+			// Flush any pending tasks (allows variable sets etc. queued from GT)
+			FAnimNextGraphInstance* GraphInstance = BlendStackTrait.GetGraphInstance(Context, TrackerIndex);
+			check(GraphInstance != nullptr);
+			FInstanceTaskContext TaskContext(*GraphInstance);
+			InjectionTracker.Request->FlushTasks(TaskContext);
+
+			// Broadcast our timeline progress
 			TTraitBinding<ITimeline> TimelineTrait;
 			if (Binding.GetStackInterface(TimelineTrait))
 			{
@@ -326,28 +329,33 @@ namespace UE::AnimNext
 					}
 
 					// Blend to source child or do we still have an active pin input?
+					bool bBlendedToPin = false;
 					if(Graph.Asset != nullptr)
 					{
 						// Add a request to blend in our fallback graph using blend out time
+						bool bRequestValid = true;
 						IBlendStack::FGraphRequest BlendStackRequest;
 						BlendStackRequest.Type = IBlendStack::EGraphRequestType::Owned;
-						BlendStackRequest.AnimationGraph = IGraphFactory::GetGraphFromObjectWithFallback(Context, Binding, Graph.Asset, BlendStackRequest.GraphPayload);
-						if (BlendStackRequest.AnimationGraph == nullptr)
-						{
-							BlendStackRequest.GraphPayload = FAnimNextDataInterfacePayload();
-						}
+						BlendStackRequest.AnimationGraph = IGraphFactory::GetOrBuildGraph(Context, Binding, Graph.Asset, BlendStackRequest.FactoryParams);
+						bRequestValid = BlendStackRequest.AnimationGraph != nullptr;
 						BlendStackRequest.BlendArgs.BlendTime = BlendOutTime;
 						BlendStackRequest.BlendMode = TranslateBlendMode(RequestArgs.BlendOutSettings.BlendMode);
 
-						InstanceData->PendingTracker.Initialize(FInjectionRequestPtr(), EInjectionTrackerState::ActivePin);
-						ActiveChildIndex = BlendStackTrait.PushGraph(Context, MoveTemp(BlendStackRequest));
+						if (bRequestValid)
+						{
+							InstanceData->PendingTracker.Initialize(FInjectionRequestPtr(), EInjectionTrackerState::ActivePin);
+							ActiveChildIndex = BlendStackTrait.PushGraph(Context, MoveTemp(BlendStackRequest));
 
-						// Make sure we tracked it after blending
-						ensure(InstanceData->InjectionTrackers.IsValidIndex(ActiveChildIndex) &&
-							InstanceData->InjectionTrackers[ActiveChildIndex].Request == FInjectionRequestPtr() && 
-							InstanceData->InjectionTrackers[ActiveChildIndex].State == EInjectionTrackerState::ActivePin);
+							// Make sure we tracked it after blending
+							ensure(InstanceData->InjectionTrackers.IsValidIndex(ActiveChildIndex) &&
+								InstanceData->InjectionTrackers[ActiveChildIndex].Request == FInjectionRequestPtr() && 
+								InstanceData->InjectionTrackers[ActiveChildIndex].State == EInjectionTrackerState::ActivePin);
+
+							bBlendedToPin = true;
+						}
 					}
-					else
+
+					if (!bBlendedToPin)
 					{
 						// Add a request to blend in our source using blend out time
 						IBlendStack::FGraphRequest BlendStackRequest;
diff --git a/Engine/Plugins/Experimental/AnimNextAnimGraph/Source/AnimNextAnimGraph/Private/Injection/InjectionUtils.cpp b/Engine/Plugins/Experimental/AnimNextAnimGraph/Source/AnimNextAnimGraph/Private/Injection/InjectionUtils.cpp
--- a/Engine/Plugins/Experimental/AnimNextAnimGraph/Source/AnimNextAnimGraph/Private/Injection/InjectionUtils.cpp
+++ b/Engine/Plugins/Experimental/AnimNextAnimGraph/Source/AnimNextAnimGraph/Private/Injection/InjectionUtils.cpp
@@ -4,10 +4,7 @@
 #include "GraphInterfaces/AnimNextNativeDataInterface_AnimSequencePlayer.h"
 #include "Animation/AnimSequence.h"
 #include "Component/AnimNextComponent.h"
-#include "DataInterface/AnimNextDataInterfaceInstance.h"
-#include "Graph/AnimNextGraphInstance.h"
 #include "Logging/StructuredLog.h"
-#include "UObject/ObjectKey.h"
 
 namespace UE::AnimNext
 {
@@ -65,7 +62,7 @@ namespace UE::AnimNext
 		UAnimNextComponent* InComponent,
 		const FInjectionSite& InSite,
 		UAnimSequence* InAnimSequence,
-		const FPlayAnimArgs& InArgs,
+		FPlayAnimArgs&& InArgs,
 		const FInjectionBlendSettings& InBlendInSettings,
 		const FInjectionBlendSettings& InBlendOutSettings,
 		FInjectionLifetimeEvents&& InLifetimeEvents)
@@ -75,7 +72,7 @@ namespace UE::AnimNext
 			InComponent->GetModuleHandle(),
 			InSite,
 			InAnimSequence,
-			InArgs,
+			MoveTemp(InArgs),
 			InBlendInSettings,
 			InBlendOutSettings,
 			MoveTemp(InLifetimeEvents));
@@ -86,41 +83,50 @@ namespace UE::AnimNext
 		FModuleHandle InModuleHandle,
 		const FInjectionSite& InSite,
 		UAnimSequence* InAnimSequence,
-		const FPlayAnimArgs& InArgs,
+		FPlayAnimArgs&& InArgs,
 		const FInjectionBlendSettings& InBlendInSettings,
 		const FInjectionBlendSettings& InBlendOutSettings,
 		FInjectionLifetimeEvents&& InLifetimeEvents)
 	{
 		check(IsInGameThread());
 
-		FInstancedStruct Payload;
-		Payload.InitializeAs<FAnimNextNativeDataInterface_AnimSequencePlayer>();
-		FAnimNextNativeDataInterface_AnimSequencePlayer& PlayNativeInterface = Payload.GetMutable<FAnimNextNativeDataInterface_AnimSequencePlayer>();
-		PlayNativeInterface.AnimSequence = InAnimSequence;
-		PlayNativeInterface.PlayRate = InArgs.PlayRate;
-		PlayNativeInterface.StartPosition = InArgs.StartPosition;
-		switch (InArgs.LoopMode)
-		{
-		case ELoopMode::Auto:
-			PlayNativeInterface.Loop = InAnimSequence->bLoop;
-			break;
-		case ELoopMode::ForceLoop:
-			PlayNativeInterface.Loop = true;
-			break;
-		case ELoopMode::ForceNonLoop:
-			PlayNativeInterface.Loop = false;
-			break;
-		}
-
 		FInjectionRequestArgs RequestArgs;
 		RequestArgs.Site = InSite;
 		RequestArgs.Object = InAnimSequence;
 		RequestArgs.Type = EAnimNextInjectionType::InjectObject;
 		RequestArgs.BlendInSettings = InBlendInSettings;
 		RequestArgs.BlendOutSettings = InBlendOutSettings;
-		RequestArgs.Payload.AddNative(MoveTemp(Payload));
 		RequestArgs.LifetimeType = InArgs.LifetimeType;
 
+		if (!InArgs.FactoryParams.IsValid())
+		{
+			RequestArgs.FactoryParams =
+				FAnimNextFactoryParams()
+				.PushPublicTrait(FSequencePlayerData(
+					{
+						.AnimSequence = InAnimSequence,
+						.PlayRate = InArgs.PlayRate,
+						.StartPosition = InArgs.StartPosition,
+						.bLoop = [&InArgs, &InAnimSequence]()
+						{
+							switch (InArgs.LoopMode)
+							{
+							case ELoopMode::Auto:
+								return InAnimSequence->bLoop;
+							case ELoopMode::ForceLoop:
+								return true;
+							case ELoopMode::ForceNonLoop:
+								return false;
+							}
+							return false;
+						}()
+					}));
+		}
+		else
+		{
+			RequestArgs.FactoryParams = MoveTemp(InArgs.FactoryParams);
+		}
+
 		return Inject(InHost, InModuleHandle, MoveTemp(RequestArgs), MoveTemp(InLifetimeEvents));
 	}
 
@@ -128,8 +134,7 @@ namespace UE::AnimNext
 		UAnimNextComponent* InComponent,
 		const FInjectionSite& InSite,
 		UObject* InAsset,
-		FAnimNextDataInterfacePayload&& InPayload,
-		UAnimNextComponent* InBindingComponent,
+		FAnimNextFactoryParams&& InFactoryParams,
 		const FInjectionBlendSettings& InBlendInSettings,
 		const FInjectionBlendSettings& InBlendOutSettings,
 		FInjectionLifetimeEvents&& InLifetimeEvents)
@@ -139,8 +144,7 @@ namespace UE::AnimNext
 			InComponent->GetModuleHandle(),
 			InSite,
 			InAsset,
-			MoveTemp(InPayload),
-			InBindingComponent ? InBindingComponent->GetModuleHandle() : FModuleHandle(),
+			MoveTemp(InFactoryParams),
 			InBlendInSettings,
 			InBlendOutSettings,
 			MoveTemp(InLifetimeEvents));
@@ -151,8 +155,7 @@ namespace UE::AnimNext
 		FModuleHandle InModuleHandle,
 		const FInjectionSite& InSite,
 		UObject* InAsset,
-		FAnimNextDataInterfacePayload&& InPayload,
-		FModuleHandle InBindingModuleHandle,
+		FAnimNextFactoryParams&& InFactoryParams,
 		const FInjectionBlendSettings& InBlendInSettings,
 		const FInjectionBlendSettings& InBlendOutSettings,
 		FInjectionLifetimeEvents&& InLifetimeEvents)
@@ -162,11 +165,10 @@ namespace UE::AnimNext
 		FInjectionRequestArgs RequestArgs;
 		RequestArgs.Site = InSite;
 		RequestArgs.Object = InAsset;
-		RequestArgs.BindingModuleHandle = InBindingModuleHandle;
 		RequestArgs.Type = EAnimNextInjectionType::InjectObject;
 		RequestArgs.BlendInSettings = InBlendInSettings;
 		RequestArgs.BlendOutSettings = InBlendOutSettings;
-		RequestArgs.Payload = MoveTemp(InPayload);
+		RequestArgs.FactoryParams = MoveTemp(InFactoryParams);
 
 		return Inject(InHost, InModuleHandle, MoveTemp(RequestArgs), MoveTemp(InLifetimeEvents));
 	}
diff --git a/Engine/Plugins/Experimental/AnimNextAnimGraph/Source/AnimNextAnimGraph/Private/PlayAnim/PlayAnimCallbackProxy.cpp b/Engine/Plugins/Experimental/AnimNextAnimGraph/Source/AnimNextAnimGraph/Private/PlayAnim/PlayAnimCallbackProxy.cpp
--- a/Engine/Plugins/Experimental/AnimNextAnimGraph/Source/AnimNextAnimGraph/Private/PlayAnim/PlayAnimCallbackProxy.cpp
+++ b/Engine/Plugins/Experimental/AnimNextAnimGraph/Source/AnimNextAnimGraph/Private/PlayAnim/PlayAnimCallbackProxy.cpp
@@ -5,6 +5,7 @@
 #include "Component/AnimNextComponent.h"
 #include "GraphInterfaces/AnimNextNativeDataInterface_AnimSequencePlayer.h"
 #include "Injection/InjectionUtils.h"
+#include "Variables/AnimNextVariableReference.h"
 
 #include UE_INLINE_GENERATED_CPP_BY_NAME(PlayAnimCallbackProxy)
 
@@ -48,9 +49,10 @@ bool UPlayAnimCallbackProxy::Play(
 		LifetimeEvents.OnInterrupted.BindUObject(this, &UPlayAnimCallbackProxy::OnPlayAnimInterrupted);
 		LifetimeEvents.OnBlendingOut.BindUObject(this, &UPlayAnimCallbackProxy::OnPlayAnimBlendingOut);
 
+		PRAGMA_DISABLE_DEPRECATION_WARNINGS
 		PlayingRequest = UE::AnimNext::FInjectionUtils::PlayAnim(
 			AnimNextComponent,
-			UE::AnimNext::FInjectionSite(SiteName),
+			UE::AnimNext::FInjectionSite(FAnimNextVariableReference(SiteName)),
 			AnimSequence,
 			{ 
 				.PlayRate = PlayRate,
@@ -59,6 +61,7 @@ bool UPlayAnimCallbackProxy::Play(
 			BlendInSettings,
 			BlendOutSettings,
 			MoveTemp(LifetimeEvents));
+		PRAGMA_ENABLE_DEPRECATION_WARNINGS
 		bWasInterrupted = false;
 	}
 
diff --git a/Engine/Plugins/Experimental/AnimNextAnimGraph/Source/AnimNextAnimGraph/Private/PlayAnim/PlayAnimSlotTrait.cpp b/Engine/Plugins/Experimental/AnimNextAnimGraph/Source/AnimNextAnimGraph/Private/PlayAnim/PlayAnimSlotTrait.cpp
--- a/Engine/Plugins/Experimental/AnimNextAnimGraph/Source/AnimNextAnimGraph/Private/PlayAnim/PlayAnimSlotTrait.cpp
+++ b/Engine/Plugins/Experimental/AnimNextAnimGraph/Source/AnimNextAnimGraph/Private/PlayAnim/PlayAnimSlotTrait.cpp
@@ -3,10 +3,8 @@
 #include "PlayAnim/PlayAnimSlotTrait.h"
 
 #include "AnimNextAnimGraphSettings.h"
-#include "AnimNextDataInterfacePayload.h"
 #include "Animation/AnimSequence.h"
-#include "DataInterface/AnimNextDataInterfaceHost.h"
-#include "DataInterface/DataInterfaceStructAdapter.h"
+#include "Factory/AnimGraphFactory.h"
 #include "TraitCore/ExecutionContext.h"
 #include "TraitCore/NodeInstance.h"
 #include "TraitInterfaces/ITimeline.h"
@@ -37,12 +35,11 @@ namespace UE::AnimNext
 	#undef TRAIT_INTERFACE_ENUMERATOR
 	#undef TRAIT_EVENT_ENUMERATOR
 
-	void FPlayAnimSlotTrait::FPlayAnimSlotRequest::Initialize(FInjectionRequestPtr InRequest, const FInjectionBlendSettings& InBlendSettings, const UAnimNextAnimationGraph* InAnimationGraph, FAnimNextDataInterfacePayload&& InPayload)
+	void FPlayAnimSlotTrait::FPlayAnimSlotRequest::Initialize(FInjectionRequestPtr InRequest, const FInjectionBlendSettings& InBlendSettings, const UAnimNextAnimationGraph* InAnimationGraph)
 	{
 		Request = InRequest;
 		BlendSettings = InBlendSettings;
 		AnimationGraph = InAnimationGraph;
-		Payload = MoveTemp(InPayload);
 
 		// If no input is provided, we'll use the source
 		State = InAnimationGraph != nullptr ? EPlayAnimRequestState::Active : EPlayAnimRequestState::ActiveSource;
@@ -83,14 +80,16 @@ namespace UE::AnimNext
 	ETraitStackPropagation FPlayAnimSlotTrait::OnInjectEvent(const FExecutionContext& Context, FTraitBinding& Binding, FInjection_InjectEvent& Event) const
 	{
 		const FSharedData* SharedData = Binding.GetSharedData<FSharedData>();
-		const FName SlotName = SharedData->GetSlotName(Binding);
+		PRAGMA_DISABLE_DEPRECATION_WARNINGS
+		const FAnimNextVariableReference InjectionSite(SharedData->GetSlotName(Binding));
+		PRAGMA_ENABLE_DEPRECATION_WARNINGS
 
 		const FAnimNextInjectionRequestArgs& RequestArgs = Event.Request->GetArgs();
-		if (SlotName == RequestArgs.Site.DesiredSiteName)
+		if (InjectionSite == RequestArgs.Site.DesiredSite)
 		{
 			FInstanceData* InstanceData = Binding.GetInstanceData<FInstanceData>();
 
-			ensureMsgf(!InstanceData->PendingRequest.IsValid(), TEXT("PlayAnim slot %s already contained a pending request, it will be overwritten"), *SlotName.ToString());
+			ensureMsgf(!InstanceData->PendingRequest.IsValid(), TEXT("PlayAnim slot %s already contained a pending request, it will be overwritten"), *InjectionSite.GetName().ToString());
 
 			// Overwrite any request we might have, we'll pick it up on the next update
 			InstanceData->PendingRequest.Reset();
@@ -105,10 +104,12 @@ namespace UE::AnimNext
 	ETraitStackPropagation FPlayAnimSlotTrait::OnUninjectEvent(const FExecutionContext& Context, FTraitBinding& Binding, FInjection_UninjectEvent& Event) const
 	{
 		const FSharedData* SharedData = Binding.GetSharedData<FSharedData>();
-		const FName SlotName = SharedData->GetSlotName(Binding);
+		PRAGMA_DISABLE_DEPRECATION_WARNINGS
+		const FAnimNextVariableReference InjectionSite(SharedData->GetSlotName(Binding));
+		PRAGMA_ENABLE_DEPRECATION_WARNINGS
 
 		const FAnimNextInjectionRequestArgs& RequestArgs = Event.Request->GetArgs();
-		if (SlotName == RequestArgs.Site.DesiredSiteName)
+		if (InjectionSite == RequestArgs.Site.DesiredSite)
 		{
 			FInstanceData* InstanceData = Binding.GetInstanceData<FInstanceData>();
 
@@ -163,7 +164,6 @@ namespace UE::AnimNext
 
 			FInjectionBlendSettings BlendSettings;
 			const UAnimNextAnimationGraph* AnimationGraph = nullptr;
-			FAnimNextDataInterfacePayload Payload;
 
 			if (Request)
 			{
@@ -171,15 +171,8 @@ namespace UE::AnimNext
 				FAnimNextInjectionRequestArgs& RequestArgs = Request->GetMutableArgs();
 				if (RequestArgs.Object)
 				{
-					Payload = MoveTemp(RequestArgs.Payload);
-
-					AnimationGraph = IGraphFactory::GetGraphFromObjectWithFallback(Context, Binding, RequestArgs.Object, Payload);
-
-					if (AnimationGraph == nullptr)
-					{
-						Payload = FAnimNextDataInterfacePayload();
-					}
-					else
+					AnimationGraph = IGraphFactory::GetOrBuildGraph(Context, Binding, RequestArgs.Object, RequestArgs.FactoryParams);
+					if (AnimationGraph != nullptr)
 					{
 						// Check for re-entrancy and early-out if we are linking back to the current instance or one of its parents
 						const FName EntryPoint = AnimationGraph->DefaultEntryPoint;
@@ -218,7 +211,7 @@ namespace UE::AnimNext
 			const int32 FreeRequestIndex = FindFreeRequestIndexOrAdd(*InstanceData);
 
 			FPlayAnimSlotRequest& SlotRequest = InstanceData->SlotRequests[FreeRequestIndex];
-			SlotRequest.Initialize(Request, BlendSettings, AnimationGraph, MoveTemp(Payload));
+			SlotRequest.Initialize(Request, BlendSettings, AnimationGraph);
 
 			const int32 OldChildIndex = InstanceData->CurrentlyActiveRequestIndex;
 			const int32 NewChildIndex = FreeRequestIndex;
@@ -307,7 +300,7 @@ namespace UE::AnimNext
 					const int32 FreeRequestIndex = FindFreeRequestIndexOrAdd(*InstanceData);
 
 					FPlayAnimSlotRequest& FreeSlotRequest = InstanceData->SlotRequests[FreeRequestIndex];
-					FreeSlotRequest.Initialize(FInjectionRequestPtr(), RequestArgs.BlendOutSettings, nullptr, FAnimNextDataInterfacePayload());
+					FreeSlotRequest.Initialize(FInjectionRequestPtr(), RequestArgs.BlendOutSettings, nullptr);
 
 					const int32 OldChildIndex = InstanceData->CurrentlyActiveRequestIndex;
 					const int32 NewChildIndex = FreeRequestIndex;
@@ -401,8 +394,13 @@ namespace UE::AnimNext
 			{
 				const FName EntryPoint = SlotRequest.AnimationGraph->DefaultEntryPoint;
 				FAnimNextGraphInstance& Owner = Binding.GetTraitPtr().GetNodeInstance()->GetOwner();
-				SlotRequest.GraphInstance = SlotRequest.AnimationGraph->AllocateInstance(Owner.GetModuleInstance(), &Context, &Binding.GetTraitPtr().GetNodeInstance()->GetOwner(), EntryPoint);
-				SlotRequest.GraphInstance->BindPublicVariables(SlotRequest.Payload.Get());
+				SlotRequest.GraphInstance = SlotRequest.AnimationGraph->AllocateInstance(
+					{
+						.ModuleInstance = Owner.GetModuleInstance(),
+						.ParentContext = &Context,
+						.ParentGraphInstance = &Binding.GetTraitPtr().GetNodeInstance()->GetOwner(),
+						.EntryPoint = EntryPoint
+					});
 
 				// TODO: Validate that our child implements the ITimeline interface
 
diff --git a/Engine/Plugins/Experimental/AnimNextAnimGraph/Source/AnimNextAnimGraph/Private/PlayAnim/PlayAnimSlotTrait.h b/Engine/Plugins/Experimental/AnimNextAnimGraph/Source/AnimNextAnimGraph/Private/PlayAnim/PlayAnimSlotTrait.h
--- a/Engine/Plugins/Experimental/AnimNextAnimGraph/Source/AnimNextAnimGraph/Private/PlayAnim/PlayAnimSlotTrait.h
+++ b/Engine/Plugins/Experimental/AnimNextAnimGraph/Source/AnimNextAnimGraph/Private/PlayAnim/PlayAnimSlotTrait.h
@@ -94,9 +94,6 @@ namespace UE::AnimNext
 			// The module used by the graph instance, as selected by the chooser
 			TObjectPtr<const UAnimNextAnimationGraph> AnimationGraph;
 
-			// Payloads for the graph instance when instantiated
-			FAnimNextDataInterfacePayload Payload;
-
 			// The graph instance
 			TSharedPtr<FAnimNextGraphInstance> GraphInstance;
 
@@ -111,7 +108,7 @@ namespace UE::AnimNext
 			bool bWasRelevant = false;
 
 			// Initializes a request to begin playing
-			void Initialize(FInjectionRequestPtr InRequest, const FInjectionBlendSettings& InBlendSettings, const UAnimNextAnimationGraph* InAnimationGraph, FAnimNextDataInterfacePayload&& InPayload);
+			void Initialize(FInjectionRequestPtr InRequest, const FInjectionBlendSettings& InBlendSettings, const UAnimNextAnimationGraph* InAnimationGraph);
 
 			FWeakTraitPtr GetChildPtr() const
 			{
diff --git a/Engine/Plugins/Experimental/AnimNextAnimGraph/Source/AnimNextAnimGraph/Private/TraitCore/ExecutionContext.cpp b/Engine/Plugins/Experimental/AnimNextAnimGraph/Source/AnimNextAnimGraph/Private/TraitCore/ExecutionContext.cpp
--- a/Engine/Plugins/Experimental/AnimNextAnimGraph/Source/AnimNextAnimGraph/Private/TraitCore/ExecutionContext.cpp
+++ b/Engine/Plugins/Experimental/AnimNextAnimGraph/Source/AnimNextAnimGraph/Private/TraitCore/ExecutionContext.cpp
@@ -512,22 +512,10 @@ namespace UE::AnimNext
 		}
 	}
 
-	FGraphInstanceComponent* FExecutionContext::TryGetComponent(int32 ComponentNameHash, FName ComponentName) const
+	void FExecutionContext::ForEachComponent(TFunctionRef<bool(FUAFGraphInstanceComponent&)> InFunction) const
 	{
 		check(IsBound());
-		return RootGraphInstance->TryGetComponent(ComponentNameHash, ComponentName);
-	}
-
-	FGraphInstanceComponent& FExecutionContext::AddComponent(int32 ComponentNameHash, FName ComponentName, TSharedPtr<FGraphInstanceComponent>&& Component) const
-	{
-		check(IsBound());
-		return RootGraphInstance->AddComponent(ComponentNameHash, ComponentName, MoveTemp(Component));
-	}
-
-	GraphInstanceComponentMapType::TConstIterator FExecutionContext::GetComponentIterator() const
-	{
-		check(IsBound());
-		return RootGraphInstance->GetComponentIterator();
+		RootGraphInstance->ForEachComponent<FUAFGraphInstanceComponent>(InFunction);
 	}
 
 	void FExecutionContext::RaiseInputTraitEvent(FAnimNextTraitEventPtr Event)
diff --git a/Engine/Plugins/Experimental/AnimNextAnimGraph/Source/AnimNextAnimGraph/Private/TraitCore/NodeTemplateBuilder.cpp b/Engine/Plugins/Experimental/AnimNextAnimGraph/Source/AnimNextAnimGraph/Private/TraitCore/NodeTemplateBuilder.cpp
--- a/Engine/Plugins/Experimental/AnimNextAnimGraph/Source/AnimNextAnimGraph/Private/TraitCore/NodeTemplateBuilder.cpp
+++ b/Engine/Plugins/Experimental/AnimNextAnimGraph/Source/AnimNextAnimGraph/Private/TraitCore/NodeTemplateBuilder.cpp
@@ -18,7 +18,7 @@ namespace UE::AnimNext
 		return BuildNodeTemplate(TraitUIDs, NodeTemplateBuffer);
 	}
 
-	FNodeTemplate* FNodeTemplateBuilder::BuildNodeTemplate(const TArray<FTraitUID>& InTraitUIDs, TArray<uint8>& NodeTemplateBuffer)
+	FNodeTemplate* FNodeTemplateBuilder::BuildNodeTemplate(TConstArrayView<FTraitUID> InTraitUIDs, TArray<uint8>& NodeTemplateBuffer)
 	{
 		NodeTemplateBuffer.Reset();
 
@@ -46,7 +46,7 @@ namespace UE::AnimNext
 		TraitUIDs.Reset();
 	}
 
-	uint32 FNodeTemplateBuilder::GetNodeTemplateUID(const TArray<FTraitUID>& InTraitUIDs)
+	uint32 FNodeTemplateBuilder::GetNodeTemplateUID(TConstArrayView<FTraitUID> InTraitUIDs)
 	{
 		uint32 NodeTemplateUID = 0;
 
@@ -59,7 +59,7 @@ namespace UE::AnimNext
 	}
 
 	void FNodeTemplateBuilder::AppendTemplateTrait(
-		const TArray<FTraitUID>& InTraitUIDs, int32 TraitIndex,
+		TConstArrayView<FTraitUID> InTraitUIDs, int32 TraitIndex,
 		TArray<uint8>& NodeTemplateBuffer)
 	{
 		const FTraitRegistry& TraitRegistry = FTraitRegistry::Get();
diff --git a/Engine/Plugins/Experimental/AnimNextAnimGraph/Source/AnimNextAnimGraph/Private/TraitCore/Trait.cpp b/Engine/Plugins/Experimental/AnimNextAnimGraph/Source/AnimNextAnimGraph/Private/TraitCore/Trait.cpp
--- a/Engine/Plugins/Experimental/AnimNextAnimGraph/Source/AnimNextAnimGraph/Private/TraitCore/Trait.cpp
+++ b/Engine/Plugins/Experimental/AnimNextAnimGraph/Source/AnimNextAnimGraph/Private/TraitCore/Trait.cpp
@@ -199,6 +199,83 @@ namespace UE::AnimNext
 	}
 #endif
 
+	uint32 FTrait::GetVariableMappedLatentPropertyHandles(
+		const FAnimNextTraitSharedData* InSharedData,
+		TArray<FLatentPropertyMetadata>& OutLatentPropertyHandles,
+		bool bFilterEditorOnly,
+		const TFunction<uint16(FName PropertyName)>& GetTraitVariableMappingIndex) const
+	{
+		uint32 NumHandlesAdded = 0;
+
+		const UStruct* BaseStruct = GetTraitSharedDataStruct();
+
+		// The property linked list on UScriptStruct iterates over the properties starting in the derived type
+		// but with latent properties, the base type should be the first to be visited.
+		// Gather our struct hierarchy from most derived to base
+		TArray<const UStruct*> StructHierarchy;
+
+		do
+		{
+			StructHierarchy.Add(BaseStruct);
+			BaseStruct = BaseStruct->GetSuperStruct();
+		}
+		while (BaseStruct != nullptr);
+
+		// Gather our latent properties from base to most derived
+		for (auto It = StructHierarchy.rbegin(); It != StructHierarchy.rend(); ++It)
+		{
+			const UStruct* SharedDataStruct = *It;
+			for (const FField* Field = SharedDataStruct->ChildProperties; Field != nullptr; Field = Field->Next)
+			{
+				const FProperty* Property = CastField<FProperty>(Field);
+
+				if (bFilterEditorOnly && Property->IsEditorOnlyProperty())
+				{
+					continue;	// Skip editor only properties if we don't need them
+				}
+
+				// By default, properties are latent
+				// However, there are exceptions:
+				//     - Properties of trait handle type are never lazy since they just encode graph connectivity
+				
+				bool bIsPotentiallyLatent = true;
+				const FStructProperty* StructProperty = CastField<FStructProperty>(Property);
+				if (StructProperty != nullptr && StructProperty->Struct == FAnimNextTraitHandle::StaticStruct())
+				{
+					bIsPotentiallyLatent = false;
+				}
+
+				if (!bIsPotentiallyLatent)
+				{
+					continue;	// Skip non-latent properties
+				}
+
+				if (!IsPropertyLatent(*InSharedData, Property->GetFName()))
+				{
+					// Skip properties not part of the latent macro enumerator
+					// This can occur if the property is missing from the enumerator by mistake
+					// in which case we'll warn during compilation
+					continue;
+				}
+
+				FLatentPropertyMetadata Metadata;
+				Metadata.Name = Property->GetFName();
+				Metadata.RigVMIndex = GetTraitVariableMappingIndex(Property->GetFName());
+
+				// Always false for now, we don't support freezing yet
+				Metadata.bCanFreeze = false;
+
+				// Always true for now when calling this function
+				Metadata.bUsesVariableCopy = true;
+
+				OutLatentPropertyHandles.Add(Metadata);
+				NumHandlesAdded++;
+			}
+		}
+
+		return NumHandlesAdded;
+	}
+
 	FTraitStaticInitHook::FTraitStaticInitHook(TraitConstructorFunc InTraitConstructor)
 		: TraitConstructor(InTraitConstructor)
 	{
diff --git a/Engine/Plugins/Experimental/AnimNextAnimGraph/Source/AnimNextAnimGraph/Private/TraitCore/TraitEventRaising.cpp b/Engine/Plugins/Experimental/AnimNextAnimGraph/Source/AnimNextAnimGraph/Private/TraitCore/TraitEventRaising.cpp
--- a/Engine/Plugins/Experimental/AnimNextAnimGraph/Source/AnimNextAnimGraph/Private/TraitCore/TraitEventRaising.cpp
+++ b/Engine/Plugins/Experimental/AnimNextAnimGraph/Source/AnimNextAnimGraph/Private/TraitCore/TraitEventRaising.cpp
@@ -5,7 +5,7 @@
 #include "TraitCore/TraitEventList.h"
 #include "TraitCore/TraitBinding.h"
 #include "TraitCore/TraitStackBinding.h"
-#include "Graph/GraphInstanceComponent.h"
+#include "Graph/UAFGraphInstanceComponent.h"
 
 namespace UE::AnimNext
 {
@@ -46,7 +46,7 @@ namespace UE::AnimNext
 		}
 	}
 
-	void RaiseTraitEvents(FExecutionContext& Context, FGraphInstanceComponent& Component, const FTraitEventList& EventList)
+	void RaiseTraitEvents(FExecutionContext& Context, FUAFGraphInstanceComponent& Component, const FTraitEventList& EventList)
 	{
 		// Event handlers can raise events and as such the list may change while we iterate
 		// However, if an event is added while we iterate, we will not visit it
diff --git a/Engine/Plugins/Experimental/AnimNextAnimGraph/Source/AnimNextAnimGraph/Private/TraitCore/TraitReader.cpp b/Engine/Plugins/Experimental/AnimNextAnimGraph/Source/AnimNextAnimGraph/Private/TraitCore/TraitReader.cpp
--- a/Engine/Plugins/Experimental/AnimNextAnimGraph/Source/AnimNextAnimGraph/Private/TraitCore/TraitReader.cpp
+++ b/Engine/Plugins/Experimental/AnimNextAnimGraph/Source/AnimNextAnimGraph/Private/TraitCore/TraitReader.cpp
@@ -149,6 +149,7 @@ namespace UE::AnimNext
 				{
 					LatentHeader.bHasValidLatentProperties = false;
 					LatentHeader.bCanAllPropertiesFreeze = true;
+					LatentHeader.bAreAllPropertiesVariableAccesses = true;
 				}
 
 				if (NumLatentProperties == 0)
@@ -158,6 +159,7 @@ namespace UE::AnimNext
 
 				bool bHasValidLatentProperties = !!LatentHeader.bHasValidLatentProperties;
 				bool bCanAllPropertiesFreeze = !!LatentHeader.bCanAllPropertiesFreeze;
+				bool bAreAllPropertiesVariableAccesses = !!LatentHeader.bAreAllPropertiesVariableAccesses;
 
 				const FTraitRegistryHandle TraitHandle = TraitTemplate.GetRegistryHandle();
 				const FTrait* Trait = TraitRegistry.Find(TraitHandle);
@@ -189,6 +191,7 @@ namespace UE::AnimNext
 
 							bHasValidLatentProperties = true;
 							bCanAllPropertiesFreeze &= bCanFreeze;
+							bAreAllPropertiesVariableAccesses &= Metadata.bUsesVariableCopy;
 
 							// Consume the property size
 							LatentPropertyOffset += PropertyMemoryLayout.Size;
@@ -200,6 +203,7 @@ namespace UE::AnimNext
 
 				LatentHeader.bHasValidLatentProperties = bHasValidLatentProperties;
 				LatentHeader.bCanAllPropertiesFreeze = bCanAllPropertiesFreeze;
+				LatentHeader.bAreAllPropertiesVariableAccesses = bAreAllPropertiesVariableAccesses;
 			}
 
 			// Set our final node instance data size that factors in our latent properties
diff --git a/Engine/Plugins/Experimental/AnimNextAnimGraph/Source/AnimNextAnimGraph/Private/TraitCore/TraitStackBinding.cpp b/Engine/Plugins/Experimental/AnimNextAnimGraph/Source/AnimNextAnimGraph/Private/TraitCore/TraitStackBinding.cpp
--- a/Engine/Plugins/Experimental/AnimNextAnimGraph/Source/AnimNextAnimGraph/Private/TraitCore/TraitStackBinding.cpp
+++ b/Engine/Plugins/Experimental/AnimNextAnimGraph/Source/AnimNextAnimGraph/Private/TraitCore/TraitStackBinding.cpp
@@ -376,6 +376,14 @@ namespace UE::AnimNext
 		const uint32 NumLatentHandles = BaseTraitDesc->GetNumSubStackLatentPropreties();
 
 		FAnimNextGraphInstance& GraphInstance = NodeInstance->GetOwner();
-		GraphInstance.ExecuteLatentPins(TConstArrayView<FLatentPropertyHandle>(LatentHandles, NumLatentHandles), NodeInstance, bIsFrozen);
+
+		if (LatentHeader.bAreAllPropertiesVariableAccesses)
+		{
+			GraphInstance.CopyVariablesToLatentPins(TConstArrayView<FLatentPropertyHandle>(LatentHandles, NumLatentHandles), NodeInstance, bIsFrozen);
+		}
+		else
+		{
+			GraphInstance.ExecuteLatentPins(TConstArrayView<FLatentPropertyHandle>(LatentHandles, NumLatentHandles), NodeInstance, bIsFrozen);
+		}
 	}
 }
diff --git a/Engine/Plugins/Experimental/AnimNextAnimGraph/Source/AnimNextAnimGraph/Private/TraitCore/TraitWriter.cpp b/Engine/Plugins/Experimental/AnimNextAnimGraph/Source/AnimNextAnimGraph/Private/TraitCore/TraitWriter.cpp
--- a/Engine/Plugins/Experimental/AnimNextAnimGraph/Source/AnimNextAnimGraph/Private/TraitCore/TraitWriter.cpp
+++ b/Engine/Plugins/Experimental/AnimNextAnimGraph/Source/AnimNextAnimGraph/Private/TraitCore/TraitWriter.cpp
@@ -1,14 +1,13 @@
 // Copyright Epic Games, Inc. All Rights Reserved.
 
 #include "TraitCore/TraitWriter.h"
-
-#if WITH_EDITOR
 #include "TraitCore/TraitRegistry.h"
 #include "TraitCore/TraitTemplate.h"
 #include "TraitCore/NodeDescription.h"
 #include "TraitCore/NodeTemplate.h"
 #include "TraitCore/NodeTemplateRegistry.h"
 #include "Serialization/ArchiveUObject.h"
+#include "StructUtils/StructView.h"
 
 #include "TraitCore/TraitReader.h"
 
@@ -122,6 +121,7 @@ namespace UE::AnimNext
 		ensure(NumNodesWritten == NodeMappings.Num());
 	}
 
+#if WITH_EDITOR
 	void FTraitWriter::WriteNode(
 		const FNodeHandle NodeHandle,
 		const TFunction<FString (uint32 TraitIndex, FName PropertyName)>& GetTraitProperty,
@@ -216,6 +216,101 @@ namespace UE::AnimNext
 
 		NumNodesWritten++;
 	}
+#endif
+
+	void FTraitWriter::WriteNode(
+		const FNodeHandle NodeHandle,
+		TFunctionRef<uint16(uint32 TraitIndex, FName PropertyName)> GetTraitVariableMappingIndex,
+		TFunctionRef<TConstStructView<FAnimNextTraitSharedData>(uint32 TraitIndex)> GetTraitData)
+	{
+		ensure(bIsNodeWriting);
+
+		if (ErrorState != EErrorState::None)
+		{
+			// We encountered an error, do nothing
+			return;
+		}
+
+		const FTraitRegistry& TraitRegistry = FTraitRegistry::Get();
+		FNodeTemplateRegistry& NodeTemplateRegistry = FNodeTemplateRegistry::Get();
+
+		const FNodeMapping* NodeMapping = NodeMappings.FindByPredicate([NodeHandle](const FNodeMapping& It) { return It.NodeHandle == NodeHandle; });
+		if (NodeMapping == nullptr)
+		{
+			ErrorState = EErrorState::NodeHandleNotFound;
+			return;
+		}
+
+		const FNodeTemplate* NodeTemplate = NodeTemplateRegistry.Find(NodeMapping->NodeTemplateHandle);
+		if (NodeTemplate == nullptr)
+		{
+			ErrorState = EErrorState::NodeTemplateNotFound;
+			return;
+		}
+
+		// Populate our node description into a temporary buffer
+		alignas(16) uint8 Buffer[64 * 1024];	// Max node size
+
+		FNodeDescription* NodeDesc = new(Buffer) FNodeDescription(NodeMapping->NodeHandle.GetNodeID(), NodeMapping->NodeTemplateHandle);
+
+		// Populate our trait properties
+		const uint32 NumTraits = NodeTemplate->GetNumTraits();
+		const FTraitTemplate* TraitTemplates = NodeTemplate->GetTraits();
+		for (uint32 TraitIndex = 0; TraitIndex < NumTraits; ++TraitIndex)
+		{
+			const FTraitRegistryHandle TraitHandle = TraitTemplates[TraitIndex].GetRegistryHandle();
+			const FTrait* Trait = TraitRegistry.Find(TraitHandle);
+
+			FAnimNextTraitSharedData* SharedData = TraitTemplates[TraitIndex].GetTraitDescription(*NodeDesc);
+			const UScriptStruct* SharedDataStruct = Trait->GetTraitSharedDataStruct();
+
+			// Initialize our output struct
+			TConstStructView<FAnimNextTraitSharedData> SourceData = GetTraitData(TraitIndex);
+			check(SourceData.GetScriptStruct() == SharedDataStruct);
+			SharedDataStruct->InitializeDefaultValue((uint8*)SharedData);
+			SharedDataStruct->CopyScriptStruct(SharedData, SourceData.GetMemory());
+		}
+
+		// Append our node and trait shared data to our archive
+		NodeDesc->Serialize(*this);
+
+		// Append our trait latent property handles to our archive
+		// We only write out the properties that will be present at runtime
+		// This takes into account editor only latent properties which can be stripped in cooked builds
+		// Other forms of property stripping are not currently supported
+		// The latent property offsets will be computed at runtime on load to support property sizes/alignment
+		// changing between the editor and the runtime platform (e.g. 32 vs 64 bit pointers)
+		// To that end, we serialize the following property metadata:
+		//     * RigVM memory handle index
+		//     * Whether the property supports freezing or not
+		//     * The property name and index for us to look it up at runtime
+
+		for (uint32 TraitIndex = 0; TraitIndex < NumTraits; ++TraitIndex)
+		{
+			const FTraitRegistryHandle TraitHandle = TraitTemplates[TraitIndex].GetRegistryHandle();
+			const FTrait* Trait = TraitRegistry.Find(TraitHandle);
+
+			// Curry our lambda with the trait index
+			const auto GetTraitVariableMappingIndexAt = [&GetTraitVariableMappingIndex, TraitIndex](FName PropertyName)
+			{
+				return GetTraitVariableMappingIndex(TraitIndex, PropertyName);
+			};
+
+			TArray<FLatentPropertyMetadata> LatentProperties;
+			LatentProperties.Reserve(32);
+			FAnimNextTraitSharedData* SharedData = TraitTemplates[TraitIndex].GetTraitDescription(*NodeDesc);
+			uint32 NumLatentProperties = Trait->GetVariableMappedLatentPropertyHandles(SharedData, LatentProperties, IsFilterEditorOnly(), GetTraitVariableMappingIndexAt);
+
+			*this << NumLatentProperties;
+
+			for (FLatentPropertyMetadata& Metadata : LatentProperties)
+			{
+				*this << Metadata;
+			}
+		}
+
+		NumNodesWritten++;
+	}
 
 	FTraitWriter::EErrorState FTraitWriter::GetErrorState() const
 	{
@@ -280,4 +375,4 @@ namespace UE::AnimNext
 		return *this;
 	}
 }
-#endif
+
diff --git a/Engine/Plugins/Experimental/AnimNextAnimGraph/Source/AnimNextAnimGraph/Private/TraitInterfaces/IBlendStack.cpp b/Engine/Plugins/Experimental/AnimNextAnimGraph/Source/AnimNextAnimGraph/Private/TraitInterfaces/IBlendStack.cpp
--- a/Engine/Plugins/Experimental/AnimNextAnimGraph/Source/AnimNextAnimGraph/Private/TraitInterfaces/IBlendStack.cpp
+++ b/Engine/Plugins/Experimental/AnimNextAnimGraph/Source/AnimNextAnimGraph/Private/TraitInterfaces/IBlendStack.cpp
@@ -53,4 +53,27 @@ namespace UE::AnimNext
 
 		return nullptr;
 	}
+
+	FAnimNextGraphInstance* IBlendStack::GetGraphInstance(FExecutionContext& Context, const TTraitBinding<IBlendStack>& Binding, int32 InChildIndex) const
+	{
+		TTraitBinding<IBlendStack> SuperBinding;
+		if (Binding.GetStackInterfaceSuper(SuperBinding))
+		{
+			return SuperBinding.GetGraphInstance(Context, InChildIndex);
+		}
+
+		return nullptr;
+	}
+
+	int32 IBlendStack::GetActiveGraphInstance(FExecutionContext& Context, const TTraitBinding<IBlendStack>& Binding, FAnimNextGraphInstance*& OutGraphInstance) const
+	{
+		TTraitBinding<IBlendStack> SuperBinding;
+		if (Binding.GetStackInterfaceSuper(SuperBinding))
+		{
+			return SuperBinding.GetActiveGraphInstance(Context, OutGraphInstance);
+		}
+
+		OutGraphInstance = nullptr;
+		return INDEX_NONE;
+	}
 }
diff --git a/Engine/Plugins/Experimental/AnimNextAnimGraph/Source/AnimNextAnimGraph/Private/TraitInterfaces/IGraphFactory.cpp b/Engine/Plugins/Experimental/AnimNextAnimGraph/Source/AnimNextAnimGraph/Private/TraitInterfaces/IGraphFactory.cpp
--- a/Engine/Plugins/Experimental/AnimNextAnimGraph/Source/AnimNextAnimGraph/Private/TraitInterfaces/IGraphFactory.cpp
+++ b/Engine/Plugins/Experimental/AnimNextAnimGraph/Source/AnimNextAnimGraph/Private/TraitInterfaces/IGraphFactory.cpp
@@ -2,7 +2,9 @@
 
 #include "TraitInterfaces/IGraphFactory.h"
 
-#include "AnimNextAnimGraphSettings.h"
+#include "Factory/AnimGraphFactory.h"
+#include "Factory/AnimNextFactoryParams.h"
+#include "Graph/AnimNextAnimationGraph.h"
 #include "TraitCore/ExecutionContext.h"
 
 namespace UE::AnimNext
@@ -22,60 +24,42 @@ namespace UE::AnimNext
 	}
 #endif // WITH_EDITOR
 
-	const UAnimNextAnimationGraph* IGraphFactory::GetGraphFromObject(FExecutionContext& Context, const TTraitBinding<IGraphFactory>& Binding, const UObject* InObject, FAnimNextDataInterfacePayload& InOutPayload) const
+	void IGraphFactory::GetFactoryParams(FExecutionContext& Context, const TTraitBinding<IGraphFactory>& Binding, const UObject* InObject, FAnimNextFactoryParams& InOutParams) const
 	{
 		TTraitBinding<IGraphFactory> SuperBinding;
 		if (Binding.GetStackInterfaceSuper(SuperBinding))
 		{
-			return SuperBinding.GetGraphFromObject(Context, InObject, InOutPayload);
+			SuperBinding.GetFactoryParams(Context, InObject, InOutParams);
 		}
-
-		return nullptr;
 	}
 
-	void IGraphFactory::CreatePayloadForObject(FExecutionContext& Context, const TTraitBinding<IGraphFactory>& Binding, const UObject* InObject, FAnimNextDataInterfacePayload& InOutPayload) const
+	const UAnimNextAnimationGraph* IGraphFactory::GetOrBuildGraph(FExecutionContext& Context, const FTraitBinding& Binding, const UObject* InObject, FAnimNextFactoryParams& InOutParams)
 	{
-		TTraitBinding<IGraphFactory> SuperBinding;
-		if (Binding.GetStackInterfaceSuper(SuperBinding))
+		// If a graph was provided directly, just use it
+		if (const UAnimNextAnimationGraph* AnimationGraph = Cast<UAnimNextAnimationGraph>(InObject))
 		{
-			return SuperBinding.CreatePayloadForObject(Context, InObject, InOutPayload);
+			return AnimationGraph;
 		}
-	}
-
-	const UAnimNextAnimationGraph* IGraphFactory::GetGraphFromObjectWithFallback(FExecutionContext& Context, const FTraitBinding& InBinding, const UObject* InObject, FAnimNextDataInterfacePayload& InOutPayload)
-	{
-		// Graph creation works as follows:
-		//	- Caller needs to create and populate the native interface for the specified object (e.g. anim sequence)
-		//	  The caller knows the object it has and how to populate that native interface
-		//	- We then ask the trait stack to populate other native interfaces it knows about (e.g. sync group trait)
-		//	- Now we have everything we need to create the graph, ask the trait stack to do so
-		//		- If the trait stack fails to create a graph, create one using our user settings
-		//	- Now that we have a graph, create any remaining payloads that might be missing
-
-		const UAnimNextAnimationGraph* AnimationGraph = nullptr;
-
-		TTraitBinding<IGraphFactory> GraphFactoryBinding;
-		if (InBinding.GetStackInterface<IGraphFactory>(GraphFactoryBinding))
+		else if (InObject != nullptr) // Otherwise attempt to build a graph procedurally from the supplied params
 		{
-			// Ask trait stack to create payloads it cares about
-			GraphFactoryBinding.CreatePayloadForObject(Context, InObject, InOutPayload);
+			if (!InOutParams.IsValid())
+			{
+				InOutParams = FAnimGraphFactory::GetDefaultParamsForClass(InObject->GetClass());
+			}
 
-			// Ask trait stack to create a graph from our object/payload
-			AnimationGraph = GraphFactoryBinding.GetGraphFromObject(Context, InObject, InOutPayload);
-		}
+			// Apply default initializer (trait stacks can still alter this below)
+			FAnimGraphFactory::InitializeDefaultParamsForObject(InObject, InOutParams);
 
-		// If we failed, fallback to our user settings
-		if (AnimationGraph == nullptr)
-		{
-			AnimationGraph = GetDefault<UAnimNextAnimGraphSettings>()->GetGraphFromObject(InObject, InOutPayload);
-		}
+			// Defer to the trait stack to gather any additional traits to instantiate
+			TTraitBinding<IGraphFactory> GraphBuilderBinding;
+			if (Binding.GetStackInterface(GraphBuilderBinding))
+			{
+				GraphBuilderBinding.GetFactoryParams(Context, InObject, InOutParams);
+			}
 
-		// Create missing payloads, if any
-		if (AnimationGraph != nullptr)
-		{
-			GetDefault<UAnimNextAnimGraphSettings>()->GetNativePayloadFromGraph(InObject, AnimationGraph, InOutPayload);
+			return FAnimGraphFactory::BuildGraph(InOutParams.GetBuilder());
 		}
 
-		return AnimationGraph;
+		return nullptr;
 	}
 }
diff --git a/Engine/Plugins/Experimental/AnimNextAnimGraph/Source/AnimNextAnimGraph/Private/TraitInterfaces/IUpdate.cpp b/Engine/Plugins/Experimental/AnimNextAnimGraph/Source/AnimNextAnimGraph/Private/TraitInterfaces/IUpdate.cpp
--- a/Engine/Plugins/Experimental/AnimNextAnimGraph/Source/AnimNextAnimGraph/Private/TraitInterfaces/IUpdate.cpp
+++ b/Engine/Plugins/Experimental/AnimNextAnimGraph/Source/AnimNextAnimGraph/Private/TraitInterfaces/IUpdate.cpp
@@ -3,7 +3,7 @@
 #include "TraitInterfaces/IUpdate.h"
 
 #include "TraitInterfaces/IHierarchy.h"
-#include "Graph/GraphInstanceComponent.h"
+#include "Graph/UAFGraphInstanceComponent.h"
 #include "AnimNextAnimGraphStats.h"
 #include "Graph/AnimNextGraphInstance.h"
 #include "TraitCore/TraitEventRaising.h"
@@ -519,13 +519,14 @@ namespace UE::AnimNext
 
 		// Before we start the traversal, we give the graph instance components the chance to do some work
 		TraversalContext.BindTo(GraphInstance);
-		for (auto It = TraversalContext.GetComponentIterator(); It; ++It)
+		TraversalContext.ForEachComponent([&TraversalContext](FUAFGraphInstanceComponent& InComponent)
 		{
-			RaiseTraitEvents(TraversalContext, *It.Value(), *TraversalContext.InputEventList);
+			RaiseTraitEvents(TraversalContext, InComponent, *TraversalContext.InputEventList);
+
+			InComponent.PreUpdate(TraversalContext);
+			return true;
+		});
 
-			It.Value()->PreUpdate(TraversalContext);
-		}
-		
 		// Add the graph root to start the update process
 		Private::FUpdateEntry RootEntry(GraphInstance.GetGraphRootPtr(), RootState);
 		TraversalContext.PushUpdateEntry(&RootEntry);
@@ -648,12 +649,13 @@ namespace UE::AnimNext
 		TraversalContext.ExecuteBookkeepingActions(*TraversalContext.RootParentBookkeepingEntryList);
 
 		// After we finish the traversal, we give the graph instance components the chance to do some work
-		for (auto It = TraversalContext.GetComponentIterator(); It; ++It)
+		TraversalContext.ForEachComponent([&TraversalContext](FUAFGraphInstanceComponent& InComponent)
 		{
-			RaiseTraitEvents(TraversalContext, *It.Value(),* TraversalContext.OutputEventList);
+			RaiseTraitEvents(TraversalContext, InComponent, *TraversalContext.OutputEventList);
 
-			It.Value()->PostUpdate(TraversalContext);
-		}
+			InComponent.PostUpdate(TraversalContext);
+			return true;
+		});
 
 		// At this point, we shouldn't have any remaining scoped interfaces
 		// If this fails, it means we failed to pop them due to a push/pop mismatch
diff --git a/Engine/Plugins/Experimental/AnimNextAnimGraph/Source/AnimNextAnimGraph/Private/Traits/BlendSpacePlayer.cpp b/Engine/Plugins/Experimental/AnimNextAnimGraph/Source/AnimNextAnimGraph/Private/Traits/BlendSpacePlayer.cpp
--- a/Engine/Plugins/Experimental/AnimNextAnimGraph/Source/AnimNextAnimGraph/Private/Traits/BlendSpacePlayer.cpp
+++ b/Engine/Plugins/Experimental/AnimNextAnimGraph/Source/AnimNextAnimGraph/Private/Traits/BlendSpacePlayer.cpp
@@ -7,8 +7,7 @@
 #include "EvaluationVM/Tasks/NormalizeRotations.h"
 #include "Graph/AnimNextAnimationGraph.h"
 #include "AnimNextAnimGraphSettings.h"
-#include "AnimNextDataInterfacePayload.h"
-#include "DataInterface/DataInterfaceStructAdapter.h"
+#include "Factory/AnimGraphFactory.h"
 #include "GraphInterfaces/AnimNextNativeDataInterface_AnimSequencePlayer.h"
 #include "Graph/AnimNextGraphInstance.h"
 #include "TraitCore/NodeInstance.h"
@@ -120,31 +119,32 @@ namespace UE::AnimNext
 				{
 					const UAnimSequence* CurSequence = BlendSamples[CurSampleIndex].Animation;
 
-					FAnimNextDataInterfacePayload GraphPayload;
-
-					// Create a default 'native data interface' payload via which we expect to communicate with the graph 
-					FInstancedStruct NewPayload;
-					FAnimNextNativeDataInterface_AnimSequencePlayer& NewNativeInterface = NewPayload.InitializeAs<FAnimNextNativeDataInterface_AnimSequencePlayer>();
-					NewNativeInterface.AnimSequence = CurSequence;
-
-					// Assume start position is normalized.
-					NewNativeInterface.StartPosition = SharedData->GetStartPosition(Binding) * CurSequence->GetPlayLength();
-					NewNativeInterface.Loop = SharedData->GetbLoop(Binding);
-
-					GraphPayload.AddNative(MoveTemp(NewPayload));
+					FAnimNextFactoryParams FactoryParams = FAnimGraphFactory::GetDefaultParamsForClass(CurSequence->GetClass());
+					FactoryParams.PushPublicTrait(
+						FSequencePlayerData(
+						{
+							.AnimSequence = CurSequence,
+							// Assume start position is normalized.
+							.StartPosition = SharedData->GetStartPosition(Binding) * CurSequence->GetPlayLength(),
+							.bLoop = SharedData->GetbLoop(Binding)
+						}));
 
-					const UAnimNextAnimationGraph* AnimationGraph = IGraphFactory::GetGraphFromObjectWithFallback(Context, Binding, CurSequence, GraphPayload);
+					const UAnimNextAnimationGraph* AnimationGraph = IGraphFactory::GetOrBuildGraph(Context, Binding, CurSequence, FactoryParams);
 					if (AnimationGraph != nullptr)
 					{
 						FSampleGraphState& TargetGraphState = InstanceData->SampleGraphs[CurSampleIndex];
 
 						const FName EntryPoint = NAME_None;
 						FAnimNextGraphInstance& Owner = Binding.GetTraitPtr().GetNodeInstance()->GetOwner();
-						TargetGraphState.Instance = AnimationGraph->AllocateInstance(Owner.GetModuleInstance(), &Context, &Binding.GetTraitPtr().GetNodeInstance()->GetOwner(), EntryPoint);
+						TargetGraphState.Instance = AnimationGraph->AllocateInstance(
+							{
+								.ModuleInstance = Owner.GetModuleInstance(),
+								.ParentContext = &Context,
+								.ParentGraphInstance = &Binding.GetTraitPtr().GetNodeInstance()->GetOwner(),
+								.EntryPoint = EntryPoint
+							});
 						TargetGraphState.ChildPtr = TargetGraphState.Instance->GetGraphRootPtr();
-						TargetGraphState.GraphPayload = MoveTemp(GraphPayload);
-
-						TargetGraphState.Instance->BindPublicVariables(TargetGraphState.GraphPayload.Get());
+						FactoryParams.InitializeInstance(*TargetGraphState.Instance.Get());
 					}
 				}
 			}
diff --git a/Engine/Plugins/Experimental/AnimNextAnimGraph/Source/AnimNextAnimGraph/Private/Traits/BlendSpacePlayer.h b/Engine/Plugins/Experimental/AnimNextAnimGraph/Source/AnimNextAnimGraph/Private/Traits/BlendSpacePlayer.h
--- a/Engine/Plugins/Experimental/AnimNextAnimGraph/Source/AnimNextAnimGraph/Private/Traits/BlendSpacePlayer.h
+++ b/Engine/Plugins/Experimental/AnimNextAnimGraph/Source/AnimNextAnimGraph/Private/Traits/BlendSpacePlayer.h
@@ -12,52 +12,7 @@
 #include "TraitInterfaces/IContinuousBlend.h"
 #include "Animation/AnimSequence.h"
 #include "Animation/BlendSpace.h"
-#include "AnimNextDataInterfacePayload.h"
-
-#include "BlendSpacePlayer.generated.h"
-
-USTRUCT(meta = (DisplayName = "Blend Space Player"))
-struct FAnimNextBlendSpacePlayerTraitSharedData : public FAnimNextTraitSharedData
-{
-	GENERATED_BODY()
-
-	/** The blend space to play. */
-	UPROPERTY(EditAnywhere, Category = "Default", meta = (ExportAsReference = "true"))
-	TObjectPtr<UBlendSpace> BlendSpace;
-
-	/** The location on the x-axis to sample. */
-	UPROPERTY(EditAnywhere, Category = "Default")
-	float XAxisSamplePoint = 0.0f;
-
-	/** The location on the y-axis to sample. */
-	UPROPERTY(EditAnywhere, Category = "Default")
-	float YAxisSamplePoint = 0.0f;
-
-	/** The play rate multiplier at which this blend space plays. */
-	UPROPERTY(EditAnywhere, Category = "Default")
-	float PlayRate = 1.0f;
-
-	/** The time at which we should start playing this blend space. This is normalized in the [0,1] range. */
-	UPROPERTY(EditAnywhere, Category = "Default")
-	float StartPosition = 0.0f;
-
-	/** Whether to loop the animation */
-	UPROPERTY(EditAnywhere, Category = "Default")
-	bool bLoop = false;
-
-	// Latent pin support boilerplate
-	#define TRAIT_LATENT_PROPERTIES_ENUMERATOR(GeneratorMacro) \
-		GeneratorMacro(BlendSpace) \
-		GeneratorMacro(XAxisSamplePoint) \
-		GeneratorMacro(YAxisSamplePoint) \
-		GeneratorMacro(PlayRate) \
-		GeneratorMacro(StartPosition) \
-		GeneratorMacro(bLoop) \
-
-
-	GENERATE_TRAIT_LATENT_PROPERTIES(FAnimNextBlendSpacePlayerTraitSharedData, TRAIT_LATENT_PROPERTIES_ENUMERATOR)
-	#undef TRAIT_LATENT_PROPERTIES_ENUMERATOR
-};
+#include "Traits/BlendSpacePlayerTraitData.h"
 
 namespace UE::AnimNext
 {
@@ -74,8 +29,6 @@ namespace UE::AnimNext
 
 		struct FSampleGraphState
 		{
-			FAnimNextDataInterfacePayload GraphPayload;
-
 			// The graph instance
 			TSharedPtr<FAnimNextGraphInstance> Instance;
 
@@ -101,7 +54,7 @@ namespace UE::AnimNext
 			TArray<FSampleGraphState> SampleGraphs;
 
 			// Cached value of the blend space we are playing
-			TObjectPtr<UBlendSpace> BlendSpace;
+			TObjectPtr<const UBlendSpace> BlendSpace;
 
 			// Cached blend samples updated by task and used to interpolate between points over time
 			TArray<FBlendSampleData> BlendSamplesData;
diff --git a/Engine/Plugins/Experimental/AnimNextAnimGraph/Source/AnimNextAnimGraph/Private/Traits/BlendStackTrait.cpp b/Engine/Plugins/Experimental/AnimNextAnimGraph/Source/AnimNextAnimGraph/Private/Traits/BlendStackTrait.cpp
--- a/Engine/Plugins/Experimental/AnimNextAnimGraph/Source/AnimNextAnimGraph/Private/Traits/BlendStackTrait.cpp
+++ b/Engine/Plugins/Experimental/AnimNextAnimGraph/Source/AnimNextAnimGraph/Private/Traits/BlendStackTrait.cpp
@@ -4,11 +4,8 @@
 
 #include "Animation/AnimSequence.h"
 #include "AnimNextAnimGraphSettings.h"
-#include "AnimNextDataInterfacePayload.h"
-#include "DataInterface/DataInterfaceStructAdapter.h"
 #include "Graph/AnimNextGraphInstance.h"
 #include "GraphInterfaces/AnimNextNativeDataInterface_AnimSequencePlayer.h"
-#include "Injection/ModuleInjectionDataInterfaceAdapter.h"
 #include <limits>
 #include "Logging/StructuredLog.h"
 #include "TraitCore/NodeInstance.h"
@@ -16,6 +13,7 @@
 #include "TraitInterfaces/IGraphFactory.h"
 #include "TraitInterfaces/IInertializerBlend.h"
 #include "HierarchyTableBlendProfile.h"
+#include "Factory/AnimGraphFactory.h"
 
 namespace UE::AnimNext
 {
@@ -111,14 +109,6 @@ namespace UE::AnimNext
 		{
 			// Track lifetime for each child graph
 			ChildGraph.Lifetime += TraitState.GetDeltaTime();
-#if WITH_EDITOR
-			if(ChildGraph.State == EGraphState::Active && ChildGraph.Instance.IsValid() && ChildGraph.Instance->RequiresPublicVariableBinding())
-			{
-				FModuleInjectionDataInterfaceAdapter ModuleAdapter(Context.GetRootGraphInstance().GetModuleInstance(), ChildGraph.Request.BindingModuleHandle);
-				IDataInterfaceHost* DataInterfaceHost = &ModuleAdapter;
-				ChildGraph.Instance->BindPublicVariables(ChildGraph.Request.GraphPayload.Get(), TConstArrayView<IDataInterfaceHost*>(&DataInterfaceHost, 1));
-			}
-#endif
 		}
 	}
 
@@ -217,13 +207,14 @@ namespace UE::AnimNext
 			{
 			case EGraphRequestType::Owned:
 				{
-					//@TODO: Remove or implement entry points once we decide if we still need them.
-					const FName EntryPoint = NAME_None;
 					FAnimNextGraphInstance& Owner = Binding.GetTraitPtr().GetNodeInstance()->GetOwner();
-					Graph.Instance = Graph.Request.AnimationGraph->AllocateInstance(Owner.GetModuleInstance(), &Context, &Owner, EntryPoint);
-					FModuleInjectionDataInterfaceAdapter ModuleAdapter(Context.GetRootGraphInstance().GetModuleInstance(), Graph.Request.BindingModuleHandle);
-					IDataInterfaceHost* DataInterfaceHost = &ModuleAdapter;
-					Graph.Instance->BindPublicVariables(Graph.Request.GraphPayload.Get(), TConstArrayView<IDataInterfaceHost*>(&DataInterfaceHost, 1));
+					Graph.Instance = Graph.Request.AnimationGraph->AllocateInstance(
+						{
+							.ModuleInstance = Owner.GetModuleInstance(),
+							.ParentContext = &Context,
+							.ParentGraphInstance = &Owner,
+						});
+					Graph.Request.FactoryParams.InitializeInstance(*Graph.Instance.Get());
 					break;
 				}
 			case EGraphRequestType::Child:
@@ -332,30 +323,30 @@ namespace UE::AnimNext
 		TTraitBinding<IBlendStack> BlendStackTrait;
 		Binding.GetStackInterface(BlendStackTrait);
 
-		IBlendStack::FGraphRequestPtr ActiveGraphRequest;
-		BlendStackTrait.GetActiveGraph(Context, ActiveGraphRequest);
+		FAnimNextGraphInstance* ActiveGraphInstance = nullptr;
+		BlendStackTrait.GetActiveGraphInstance(Context, ActiveGraphInstance);
 
-		if (ActiveGraphRequest)
+		if (ActiveGraphInstance != nullptr)
 		{
-			for (const FInstancedStruct& Payload : ActiveGraphRequest->GraphPayload.GetNativePayloads())
+			// @TODO: Selecting the first active graph in a blend stack doesn't typically make sense, a wrapper trait is better.
+			// However this implementation is still useful for prototyping / debugging purposes. Still consider future removal.
+			const UAnimSequence* AnimSequence = nullptr;
+			ActiveGraphInstance->AccessVariablesStruct<FAnimNextNativeDataInterface_AnimSequencePlayer>([&AnimSequence](FAnimNextNativeDataInterface_AnimSequencePlayer& InAnimSequencePlayer)
 			{
-				if (const FAnimNextNativeDataInterface_AnimSequencePlayer* PlayAnimPayload = Payload.GetPtr<FAnimNextNativeDataInterface_AnimSequencePlayer>())
+				if (InAnimSequencePlayer.AnimSequence)
 				{
-					// @TODO: Selecting the first active payload in a blend stack doesn't typically make sense, a wrapper trait is better.
-					// However this implementation is still useful for prototyping / debugging purposes. Still consider future removal.
-					if (PlayAnimPayload->AnimSequence)
-					{
-						auto ExtractRootMotionAttribute = [AnimSequence = PlayAnimPayload->AnimSequence](float StartTime, float DeltaTime, bool bAllowLooping)
-						{
-							// We do not check for lifetimes, assume the sequence is alive during pose list execution.
-							check(AnimSequence->IsValidLowLevel());
-							return AnimSequence->ExtractRootMotion(FAnimExtractContext(static_cast<double>(StartTime), true, FDeltaTimeRecord(DeltaTime), bAllowLooping && AnimSequence->bLoop)); 
-						};
-
-						return FOnExtractRootMotionAttribute::CreateLambda(ExtractRootMotionAttribute);
-					}
+					AnimSequence = InAnimSequencePlayer.AnimSequence;
 				}
-			}
+			});
+
+			auto ExtractRootMotionAttribute = [AnimSequence](float StartTime, float DeltaTime, bool bAllowLooping)
+			{
+				// We do not check for lifetimes, assume the sequence is alive during pose list execution.
+				check(AnimSequence->IsValidLowLevel());
+				return AnimSequence->ExtractRootMotion(FAnimExtractContext(static_cast<double>(StartTime), true, FDeltaTimeRecord(DeltaTime), bAllowLooping && AnimSequence->bLoop)); 
+			};
+				
+			return FOnExtractRootMotionAttribute::CreateLambda(ExtractRootMotionAttribute);
 		}
 
 		return FOnExtractRootMotionAttribute();
@@ -485,6 +476,17 @@ namespace UE::AnimNext
 		return &InstanceData->ChildGraphs[InChildIndex].Request;
 	}
 
+	FAnimNextGraphInstance* FBlendStackCoreTrait::GetGraphInstance(FExecutionContext& Context, const TTraitBinding<IBlendStack>& Binding, int32 InChildIndex) const
+	{
+		FInstanceData* InstanceData = Binding.GetInstanceData<FInstanceData>();
+		if(!InstanceData->ChildGraphs.IsValidIndex(InChildIndex))
+		{
+			return nullptr;
+		}
+
+		return InstanceData->ChildGraphs[InChildIndex].Instance.Get();
+	}
+
 	TSharedPtr<const IBlendProfileInterface> FBlendStackCoreTrait::GetBlendProfile(FExecutionContext& Context, const TTraitBinding<ISmoothBlendPerBone>& Binding, int32 ChildIndex) const
 	{
 		const FInstanceData* InstanceData = Binding.GetInstanceData<FInstanceData>();
@@ -510,8 +512,8 @@ namespace UE::AnimNext
 		const UObject* DesiredObject = SharedData->GetObject(Binding);
 		if (DesiredObject != nullptr)
 		{
-			FAnimNextDataInterfacePayload GraphPayload;
-			const UAnimNextAnimationGraph* AnimationGraph = IGraphFactory::GetGraphFromObjectWithFallback(Context, Binding, DesiredObject, GraphPayload);
+			FAnimNextFactoryParams FactoryParams;
+			const UAnimNextAnimationGraph* AnimationGraph = IGraphFactory::GetOrBuildGraph(Context, Binding, DesiredObject, FactoryParams);
 			if (AnimationGraph != nullptr)
 			{
 				const int32 CurrentlyActiveGraphIndex = InstanceData->CurrentlyActiveGraphIndex;
@@ -526,7 +528,7 @@ namespace UE::AnimNext
 					GraphRequest.FactoryObject = DesiredObject;
 					GraphRequest.AnimationGraph = AnimationGraph;
 					GraphRequest.BlendArgs.BlendTime = SharedData->GetBlendTime(Binding);
-					GraphRequest.GraphPayload = MoveTemp(GraphPayload);
+					GraphRequest.FactoryParams = MoveTemp(FactoryParams);
 
 					BlendStackTrait.PushGraph(Context, MoveTemp(GraphRequest));
 				}
@@ -543,8 +545,8 @@ namespace UE::AnimNext
 		const UObject* DesiredObject = SharedData->GetObject(Binding);
 		if (DesiredObject != nullptr)
 		{
-			FAnimNextDataInterfacePayload GraphPayload;
-			const UAnimNextAnimationGraph* AnimationGraph = IGraphFactory::GetGraphFromObjectWithFallback(Context, Binding, DesiredObject, GraphPayload);
+			FAnimNextFactoryParams FactoryParams;
+			const UAnimNextAnimationGraph* AnimationGraph = IGraphFactory::GetOrBuildGraph(Context, Binding, DesiredObject, FactoryParams);
 			if (AnimationGraph != nullptr)
 			{
 				TTraitBinding<IBlendStack> BlendStackTrait;
@@ -559,7 +561,7 @@ namespace UE::AnimNext
 					GraphRequest.FactoryObject = DesiredObject;
 					GraphRequest.AnimationGraph = AnimationGraph;
 					GraphRequest.BlendArgs.BlendTime = SharedData->GetBlendTime(Binding);
-					GraphRequest.GraphPayload = MoveTemp(GraphPayload);
+					GraphRequest.FactoryParams = MoveTemp(FactoryParams);
 
 					BlendStackTrait.PushGraph(Context, MoveTemp(GraphRequest));
 				}
diff --git a/Engine/Plugins/Experimental/AnimNextAnimGraph/Source/AnimNextAnimGraph/Private/Traits/CallFunction.cpp b/Engine/Plugins/Experimental/AnimNextAnimGraph/Source/AnimNextAnimGraph/Private/Traits/CallFunction.cpp
--- a/Engine/Plugins/Experimental/AnimNextAnimGraph/Source/AnimNextAnimGraph/Private/Traits/CallFunction.cpp
+++ b/Engine/Plugins/Experimental/AnimNextAnimGraph/Source/AnimNextAnimGraph/Private/Traits/CallFunction.cpp
@@ -9,6 +9,7 @@
 #include "TraitCore/NodeInstance.h"
 #include "AnimNextAnimGraphStats.h"
 #include "AnimNextExecuteContext.h"
+#include "UAFRigVMComponent.h"
 
 DEFINE_STAT(STAT_AnimNext_CallFunction);
 
@@ -67,8 +68,9 @@ namespace UE::AnimNext
 			return;
 		}
 
-		FRigVMExtendedExecuteContext& ExtendedExecuteContext = GraphInstance.GetExtendedExecuteContext();
-		FAnimNextExecuteContext& AnimNextContext = GraphInstance.GetExtendedExecuteContext().GetPublicDataSafe<FAnimNextExecuteContext>();
+		FUAFRigVMComponent& RigVMComponent = GraphInstance.GetComponent<FUAFRigVMComponent>();
+		FRigVMExtendedExecuteContext& ExtendedExecuteContext = RigVMComponent.GetExtendedExecuteContext();
+		FAnimNextExecuteContext& AnimNextContext = ExtendedExecuteContext.GetPublicDataSafe<FAnimNextExecuteContext>();
 
 		// Propagate delta time
 		AnimNextContext.SetDeltaTime(InTraitState.GetDeltaTime());
diff --git a/Engine/Plugins/Experimental/AnimNextAnimGraph/Source/AnimNextAnimGraph/Private/Traits/NotifyDispatcher.cpp b/Engine/Plugins/Experimental/AnimNextAnimGraph/Source/AnimNextAnimGraph/Private/Traits/NotifyDispatcher.cpp
--- a/Engine/Plugins/Experimental/AnimNextAnimGraph/Source/AnimNextAnimGraph/Private/Traits/NotifyDispatcher.cpp
+++ b/Engine/Plugins/Experimental/AnimNextAnimGraph/Source/AnimNextAnimGraph/Private/Traits/NotifyDispatcher.cpp
@@ -133,11 +133,16 @@ namespace UE::AnimNext
 	}
 }
 
-void FAnimNextNotifyDispatcherComponent::OnInitialize()
+FAnimNextNotifyDispatcherComponent::FAnimNextNotifyDispatcherComponent()
 {
 	// Default to the current host's skeletal mesh component, if any
-	FAnimNextModuleInstance& ModuleInstance = GetModuleInstance();
-	const FAnimNextSkeletalMeshComponentReferenceComponent& ComponentReference = ModuleInstance.GetComponent<FAnimNextSkeletalMeshComponentReferenceComponent>();
+	FAnimNextModuleInstance* ModuleInstance = GetModuleInstancePtr();
+	if (ModuleInstance == nullptr)
+	{
+		return;
+	}
+
+	const FAnimNextSkeletalMeshComponentReferenceComponent& ComponentReference = ModuleInstance->GetComponent<FAnimNextSkeletalMeshComponentReferenceComponent>();
 	SkeletalMeshComponent = ComponentReference.GetComponent();
 	NotifyQueue.PredictedLODLevel = SkeletalMeshComponent ? SkeletalMeshComponent->GetPredictedLODLevel() : 0;
 }
@@ -307,12 +312,6 @@ void FAnimNextNotifyDispatcherComponent::TriggerSingleAnimNotify(float InDeltaTi
 		{
 			// Named notifies can be handled by a custom module event on our worker thread, no need to queue
 			ModuleInstance.RunRigVMEvent(AnimNotifyEvent->NotifyName, InDeltaTime);
-
-			// Also dispatch to any prerequisite modules
-			ModuleInstance.ForEachPrerequisite([AnimNotifyEvent, InDeltaTime](FAnimNextModuleInstance& InPrerequisiteInstance)
-			{
-				InPrerequisiteInstance.RunRigVMEvent(AnimNotifyEvent->NotifyName, InDeltaTime);
-			});
 		}
 	}
 }
diff --git a/Engine/Plugins/Experimental/AnimNextAnimGraph/Source/AnimNextAnimGraph/Private/Traits/NotifyDispatcher.h b/Engine/Plugins/Experimental/AnimNextAnimGraph/Source/AnimNextAnimGraph/Private/Traits/NotifyDispatcher.h
--- a/Engine/Plugins/Experimental/AnimNextAnimGraph/Source/AnimNextAnimGraph/Private/Traits/NotifyDispatcher.h
+++ b/Engine/Plugins/Experimental/AnimNextAnimGraph/Source/AnimNextAnimGraph/Private/Traits/NotifyDispatcher.h
@@ -8,39 +8,27 @@
 #include "TraitInterfaces/ITimelinePlayer.h"
 #include "TraitInterfaces/IUpdate.h"
 #include "Animation/AnimNotifyQueue.h"
-#include "Module/AnimNextModuleInstanceComponent.h"
-
+#include "Module/UAFModuleInstanceComponent.h"
+#include "Traits/NotifyDispatcherTraitData.h"
 #include "NotifyDispatcher.generated.h"
 
 struct FAnimNotifyEventReference;
 class USkeletalMeshComponent;
+struct FAnimNextModuleInstance;
 
 namespace UE::AnimNext
 {
 	struct FNotifyQueueDispatchEvent;
 }
 
-USTRUCT(meta = (DisplayName = "Notify Dispatcher"))
-struct FAnimNextNotifyDispatcherTraitSharedData : public FAnimNextTraitSharedData
-{
-	GENERATED_BODY()
-
-	// Latent pin support boilerplate
-	#define TRAIT_LATENT_PROPERTIES_ENUMERATOR(GeneratorMacro)
-
-	GENERATE_TRAIT_LATENT_PROPERTIES(FAnimNextNotifyDispatcherTraitSharedData, TRAIT_LATENT_PROPERTIES_ENUMERATOR)
-	#undef TRAIT_LATENT_PROPERTIES_ENUMERATOR
-};
-
 USTRUCT()
-struct FAnimNextNotifyDispatcherComponent : public FAnimNextModuleInstanceComponent
+struct FAnimNextNotifyDispatcherComponent : public FUAFModuleInstanceComponent
 {
 	GENERATED_BODY()
 
-	FAnimNextNotifyDispatcherComponent() = default;
+	FAnimNextNotifyDispatcherComponent();
 
-	// FAnimNextModuleInstanceComponent interface
-	virtual void OnInitialize() override;
+	// FUAFModuleInstanceComponent interface
 	virtual void OnTraitEvent(FAnimNextTraitEvent& Event) override;
 	virtual void OnEndExecution(float InDeltaTime) override;
 
diff --git a/Engine/Plugins/Experimental/AnimNextAnimGraph/Source/AnimNextAnimGraph/Private/Traits/SequencePlayer.h b/Engine/Plugins/Experimental/AnimNextAnimGraph/Source/AnimNextAnimGraph/Private/Traits/SequencePlayer.h
--- a/Engine/Plugins/Experimental/AnimNextAnimGraph/Source/AnimNextAnimGraph/Private/Traits/SequencePlayer.h
+++ b/Engine/Plugins/Experimental/AnimNextAnimGraph/Source/AnimNextAnimGraph/Private/Traits/SequencePlayer.h
@@ -12,40 +12,7 @@
 #include "TraitInterfaces/IUpdate.h"
 #include "Animation/AnimSequence.h"
 #include "TraitInterfaces/INotifySource.h"
-
-#include "SequencePlayer.generated.h"
-
-USTRUCT(meta = (DisplayName = "Sequence Player"))
-struct FAnimNextSequencePlayerTraitSharedData : public FAnimNextTraitSharedData
-{
-	GENERATED_BODY()
-
-	/** The sequence to play. */
-	UPROPERTY(EditAnywhere, Category = "Default", meta=(ExportAsReference="true"))
-	TObjectPtr<UAnimSequence> AnimSequence;
-
-	/** The play rate multiplier at which this sequence plays. */
-	UPROPERTY(EditAnywhere, Category = "Default")
-	float PlayRate = 1.0f;
-
-	/** The time at which we should start playing this sequence. */
-	UPROPERTY(EditAnywhere, Category = "Default")
-	float StartPosition = 0.0f;
-
-	/** Whether or not this sequence playback will loop. */
-	UPROPERTY(EditAnywhere, Category = "Default")
-	bool bLoop = false;
-
-	// Latent pin support boilerplate
-	#define TRAIT_LATENT_PROPERTIES_ENUMERATOR(GeneratorMacro) \
-		GeneratorMacro(AnimSequence) \
-		GeneratorMacro(PlayRate) \
-		GeneratorMacro(StartPosition) \
-		GeneratorMacro(bLoop) \
-
-	GENERATE_TRAIT_LATENT_PROPERTIES(FAnimNextSequencePlayerTraitSharedData, TRAIT_LATENT_PROPERTIES_ENUMERATOR)
-	#undef TRAIT_LATENT_PROPERTIES_ENUMERATOR
-};
+#include "Traits/SequencePlayerTraitData.h"
 
 namespace UE::AnimNext
 {
@@ -63,7 +30,7 @@ namespace UE::AnimNext
 		struct FInstanceData : FTrait::FInstanceData
 		{
 			// Cached value of the anim sequence we are playing
-			TObjectPtr<UAnimSequence> AnimSequence;
+			TObjectPtr<const UAnimSequence> AnimSequence;
 
 			/** Delta time range required for root motion extraction **/
 			FDeltaTimeRecord DeltaTimeRecord;
diff --git a/Engine/Plugins/Experimental/AnimNextAnimGraph/Source/AnimNextAnimGraph/Private/Traits/SubGraphHost.cpp b/Engine/Plugins/Experimental/AnimNextAnimGraph/Source/AnimNextAnimGraph/Private/Traits/SubGraphHost.cpp
--- a/Engine/Plugins/Experimental/AnimNextAnimGraph/Source/AnimNextAnimGraph/Private/Traits/SubGraphHost.cpp
+++ b/Engine/Plugins/Experimental/AnimNextAnimGraph/Source/AnimNextAnimGraph/Private/Traits/SubGraphHost.cpp
@@ -222,7 +222,13 @@ namespace UE::AnimNext
 			if (SubGraphEntry.State == ESlotState::ActiveWithGraph)
 			{
 				FAnimNextGraphInstance& Owner = Binding.GetTraitPtr().GetNodeInstance()->GetOwner();
-				SubGraphEntry.GraphInstance = SubGraphEntry.AnimationGraph->AllocateInstance(Owner.GetModuleInstance(), &Context, &Binding.GetTraitPtr().GetNodeInstance()->GetOwner(), SubGraphEntry.EntryPoint);
+				SubGraphEntry.GraphInstance = SubGraphEntry.AnimationGraph->AllocateInstance(
+					{
+						.ModuleInstance = Owner.GetModuleInstance(),
+						.ParentContext = &Context,
+						.ParentGraphInstance = &Binding.GetTraitPtr().GetNodeInstance()->GetOwner(),
+						.EntryPoint = SubGraphEntry.EntryPoint
+					});
 			}
 		}
 	}
diff --git a/Engine/Plugins/Experimental/AnimNextAnimGraph/Source/AnimNextAnimGraph/Private/Traits/SynchronizeUsingGroups.cpp b/Engine/Plugins/Experimental/AnimNextAnimGraph/Source/AnimNextAnimGraph/Private/Traits/SynchronizeUsingGroups.cpp
--- a/Engine/Plugins/Experimental/AnimNextAnimGraph/Source/AnimNextAnimGraph/Private/Traits/SynchronizeUsingGroups.cpp
+++ b/Engine/Plugins/Experimental/AnimNextAnimGraph/Source/AnimNextAnimGraph/Private/Traits/SynchronizeUsingGroups.cpp
@@ -2,7 +2,8 @@
 
 #include "Traits/SynchronizeUsingGroups.h"
 
-#include "AnimNextDataInterfacePayload.h"
+#include "InstanceTaskContext.h"
+#include "Factory/AnimNextFactoryParams.h"
 #include "Graph/SyncGroup_GraphInstanceComponent.h"
 #include "GraphInterfaces/AnimNextNativeDataInterface_SynchronizeUsingGroups.h"
 
@@ -160,26 +161,24 @@ namespace UE::AnimNext
 		ITimelinePlayer::AdvanceBy(Context, Binding, DeltaTime, bDispatchEvents);
 	}
 
-	void FSynchronizeUsingGroupsTrait::CreatePayloadForObject(FExecutionContext& Context, const TTraitBinding<IGraphFactory>& Binding, const UObject* InObject, FAnimNextDataInterfacePayload& InOutPayload) const
+	void FSynchronizeUsingGroupsTrait::GetFactoryParams(FExecutionContext& Context, const TTraitBinding<IGraphFactory>& Binding, const UObject* InObject, FAnimNextFactoryParams& InOutParams) const
 	{
-		IGraphFactory::CreatePayloadForObject(Context, Binding, InObject, InOutPayload);
+		IGraphFactory::GetFactoryParams(Context, Binding, InObject, InOutParams);
 
 		TTraitBinding<IGroupSynchronization> GroupSyncTrait;
 		Binding.GetStackInterface(GroupSyncTrait);
 
 		const FSyncGroupParameters GroupParameters = GroupSyncTrait.GetGroupParameters(Context);
 
-		// If synchronization is enabled, we populate the native interface
+		// If synchronization is enabled, request & populate the parameters
 		if (GroupParameters.SyncMode != EAnimGroupSynchronizationMode::NoSynchronization)
 		{
-			FInstancedStruct SyncGroupPayload;
-			auto& NativeInterface = SyncGroupPayload.InitializeAs<FAnimNextNativeDataInterface_SynchronizeUsingGroups>();
-			NativeInterface.GroupName = GroupParameters.GroupName;
-			NativeInterface.GroupRole = GroupParameters.GroupRole;
-			NativeInterface.SyncMode = GroupParameters.SyncMode;
-			NativeInterface.MatchSyncPoint = GroupParameters.bMatchSyncPoint;
-
-			InOutPayload.AddNative(MoveTemp(SyncGroupPayload));
+			FSynchronizeUsingGroupsData SynchronizeUsingGroups;
+			SynchronizeUsingGroups.GroupName = GroupParameters.GroupName;
+			SynchronizeUsingGroups.GroupRole = GroupParameters.GroupRole;
+			SynchronizeUsingGroups.SyncMode = GroupParameters.SyncMode;
+			SynchronizeUsingGroups.bMatchSyncPoint = GroupParameters.bMatchSyncPoint;
+			InOutParams.PushPublicTrait(SynchronizeUsingGroups);
 		}
 	}
 }
diff --git a/Engine/Plugins/Experimental/AnimNextAnimGraph/Source/AnimNextAnimGraph/Private/Traits/SynchronizeUsingGroups.h b/Engine/Plugins/Experimental/AnimNextAnimGraph/Source/AnimNextAnimGraph/Private/Traits/SynchronizeUsingGroups.h
--- a/Engine/Plugins/Experimental/AnimNextAnimGraph/Source/AnimNextAnimGraph/Private/Traits/SynchronizeUsingGroups.h
+++ b/Engine/Plugins/Experimental/AnimNextAnimGraph/Source/AnimNextAnimGraph/Private/Traits/SynchronizeUsingGroups.h
@@ -8,47 +8,12 @@
 #include "TraitInterfaces/IGraphFactory.h"
 #include "TraitInterfaces/IUpdate.h"
 #include "TraitInterfaces/ITimelinePlayer.h"
+#include "Traits/SynchronizeUsingGroupsTraitData.h"
 
-#include "SynchronizeUsingGroups.generated.h"
-
-USTRUCT(meta = (DisplayName = "Synchronize Using Groups"))
-struct FAnimNextSynchronizeUsingGroupsTraitSharedData : public FAnimNextTraitSharedData
-{
-	GENERATED_BODY()
-
-	// The group name
-	// If no name is provided, this trait is inactive
-	UPROPERTY(EditAnywhere, Category = "Default")
-	FName GroupName;
-
-	// The role this player can assume within the group
-	UPROPERTY(EditAnywhere, Category = "Default")
-	EAnimGroupSynchronizationRole GroupRole = EAnimGroupSynchronizationRole::CanBeLeader;
-
-	// The synchronization mode
-	UPROPERTY(EditAnywhere, Category = "Default")
-	EAnimGroupSynchronizationMode SyncMode = EAnimGroupSynchronizationMode::NoSynchronization;
-
-	// Whether or not to match the group sync point when joining as leader or follower with markers
-	// When disabled, the start position of synced sequences must be properly set to avoid pops
-	UPROPERTY(EditAnywhere, Category = "Default")
-	bool bMatchSyncPoint = true;
-
-	// Latent pin support boilerplate
-	#define TRAIT_LATENT_PROPERTIES_ENUMERATOR(GeneratorMacro) \
-		GeneratorMacro(GroupName) \
-		GeneratorMacro(GroupRole) \
-		GeneratorMacro(SyncMode) \
-		GeneratorMacro(bMatchSyncPoint) \
-
-	GENERATE_TRAIT_LATENT_PROPERTIES(FAnimNextSynchronizeUsingGroupsTraitSharedData, TRAIT_LATENT_PROPERTIES_ENUMERATOR)
-	#undef TRAIT_LATENT_PROPERTIES_ENUMERATOR
-};
+struct FSyncGroupGraphInstanceComponent;
 
 namespace UE::AnimNext
 {
-	struct FSyncGroupGraphInstanceComponent;
-
 	/**
 	 * FSynchronizeUsingGroupsTrait
 	 * 
@@ -93,6 +58,6 @@ namespace UE::AnimNext
 		virtual void AdvanceBy(FExecutionContext& Context, const TTraitBinding<ITimelinePlayer>& Binding, float DeltaTime, bool bDispatchEvents) const override;
 
 		// IGraphFactory impl
-		virtual void CreatePayloadForObject(FExecutionContext& Context, const TTraitBinding<IGraphFactory>& Binding, const UObject* InObject, FAnimNextDataInterfacePayload& InOutPayload) const override;
+		virtual void GetFactoryParams(FExecutionContext& Context, const TTraitBinding<IGraphFactory>& Binding, const UObject* InObject, FAnimNextFactoryParams& InOutParams) const override;
 	};
 }
diff --git a/Engine/Plugins/Experimental/AnimNextAnimGraph/Source/AnimNextAnimGraph/Public/EvaluationVM/Tasks/PushAnimSequenceKeyframe.h b/Engine/Plugins/Experimental/AnimNextAnimGraph/Source/AnimNextAnimGraph/Public/EvaluationVM/Tasks/PushAnimSequenceKeyframe.h
--- a/Engine/Plugins/Experimental/AnimNextAnimGraph/Source/AnimNextAnimGraph/Public/EvaluationVM/Tasks/PushAnimSequenceKeyframe.h
+++ b/Engine/Plugins/Experimental/AnimNextAnimGraph/Source/AnimNextAnimGraph/Public/EvaluationVM/Tasks/PushAnimSequenceKeyframe.h
@@ -24,15 +24,15 @@ struct FAnimNextAnimSequenceKeyframeTask : public FAnimNextEvaluationTask
 
 	DECLARE_ANIM_EVALUATION_TASK(FAnimNextAnimSequenceKeyframeTask)
 
-	static UE_API FAnimNextAnimSequenceKeyframeTask MakeFromSampleTime(TWeakObjectPtr<UAnimSequence> AnimSequence, double SampleTime, bool bInterpolate);
-	static UE_API FAnimNextAnimSequenceKeyframeTask MakeFromKeyframeIndex(TWeakObjectPtr<UAnimSequence> AnimSequence, uint32 KeyframeIndex);
+	static UE_API FAnimNextAnimSequenceKeyframeTask MakeFromSampleTime(TWeakObjectPtr<const UAnimSequence> AnimSequence, double SampleTime, bool bInterpolate);
+	static UE_API FAnimNextAnimSequenceKeyframeTask MakeFromKeyframeIndex(TWeakObjectPtr<const UAnimSequence> AnimSequence, uint32 KeyframeIndex);
 
 	// Task entry point
 	UE_API virtual void Execute(UE::AnimNext::FEvaluationVM& VM) const override;
 
 	// Anim Sequence to grab the keyframe from
 	UPROPERTY(VisibleAnywhere, Category = Properties)
-	TWeakObjectPtr<UAnimSequence> AnimSequence;
+	TWeakObjectPtr<const UAnimSequence> AnimSequence;
 
 	/** Delta time range required for root motion extraction **/
 	FDeltaTimeRecord DeltaTimeRecord;
diff --git a/Engine/Plugins/Experimental/AnimNextAnimGraph/Source/AnimNextAnimGraph/Public/Factory/AnimGraphFactory.h b/Engine/Plugins/Experimental/AnimNextAnimGraph/Source/AnimNextAnimGraph/Public/Factory/AnimGraphFactory.h
new file mode 100644
--- /dev/null
+++ b/Engine/Plugins/Experimental/AnimNextAnimGraph/Source/AnimNextAnimGraph/Public/Factory/AnimGraphFactory.h
@@ -0,0 +1,86 @@
+﻿// Copyright Epic Games, Inc. All Rights Reserved.
+
+#pragma once
+
+#include "CoreTypes.h"
+#include "Factory/AnimNextFactoryParams.h"
+
+class UAnimNextAnimationGraph;
+class UObject;
+
+namespace UE::AnimNext
+{
+	class FTraitWriter;
+	struct IAnimGraphBuilder;
+	struct FAnimGraphFactory;
+}
+
+namespace UE::AnimNext
+{
+
+struct FFactoryParamsInitializerContext
+{
+	// Get the object used to initialize the params (e.g. AnimSequence for a player)
+	const UObject* GetObject() const
+	{
+		return Object;
+	}
+
+	// Access a struct during initialization (for instance to inject an AnimSequence)
+	template<typename StructType>
+	bool AccessStruct(TFunctionRef<void(StructType&)> InFunction) const
+	{
+		for (TInstancedStruct<FAnimNextTraitSharedData>& StructData : StructDatas)
+		{
+			if (StructData.GetScriptStruct() == TBaseStructure<StructType>::Get())
+			{
+				InFunction(StructData.GetMutable<StructType>());
+				return true;
+			}
+		}
+		return false;
+	}
+
+private:
+	friend FAnimGraphFactory;
+
+	TArrayView<TInstancedStruct<FAnimNextTraitSharedData>> StructDatas;
+	const UObject* Object = nullptr;
+};
+
+// Creates or recycles programmatically-generated graphs
+// Uses the hash of the recipe to determine if the graph has been created already.
+// If the hash matches one that has already been created but graph has already been GCed it will be re-created as needed.
+struct FAnimGraphFactory
+{
+	// Make a graph from the specified recipe
+	// If the object is an UAnimNextAnimationGraph, it will be returned rather than build a new one
+	ANIMNEXTANIMGRAPH_API static const UAnimNextAnimationGraph* GetOrBuildGraph(const UObject* InObject, const IAnimGraphBuilder& InBuilder);
+
+	// Make a graph from the specified recipe
+	ANIMNEXTANIMGRAPH_API static const UAnimNextAnimationGraph* BuildGraph(const IAnimGraphBuilder& InBuilder);
+
+	// Make the default graph for running a UAnimNextAnimationGraph asset (allowing graphs to blend between each other)
+	ANIMNEXTANIMGRAPH_API static const UAnimNextAnimationGraph* GetDefaultGraphHost();
+
+	// Get default params used to build & initialize a graph according to registered object class
+	ANIMNEXTANIMGRAPH_API static const FAnimNextFactoryParams& GetDefaultParamsForClass(const UClass* InClass);
+
+	// Get default params used to build & initialize a graph according to registered object class, copying InObject via the registered initializer
+	// Effective equivalent of GetDefaultParamsForClass and InitializeDefaultParamsForObject.
+	ANIMNEXTANIMGRAPH_API static FAnimNextFactoryParams GetDefaultParamsForObject(const UObject* InObject);
+
+	// Applies default initializer to the params (usually injecting the object pointer) 
+	ANIMNEXTANIMGRAPH_API static void InitializeDefaultParamsForObject(const UObject* InObject, FAnimNextFactoryParams& InOutParams);
+
+	// Function used to initialize the object for the factory-generated graph
+	using FParamsInitializer = TFunction<void(const FFactoryParamsInitializerContext&)>;
+
+	// Register default params used to build & initialize a graph according to the supplied object class
+	ANIMNEXTANIMGRAPH_API static void RegisterDefaultParamsForClass(const UClass* InClass, FAnimNextFactoryParams&& InParams, FParamsInitializer&& InInitializer);
+
+	// Get all registered classes used to generate factory graphs
+	ANIMNEXTANIMGRAPH_API static TConstArrayView<UClass*> GetRegisteredClasses();
+};
+
+}
\ No newline at end of file
diff --git a/Engine/Plugins/Experimental/AnimNextAnimGraph/Source/AnimNextAnimGraph/Public/Factory/AnimNextFactoryParams.h b/Engine/Plugins/Experimental/AnimNextAnimGraph/Source/AnimNextAnimGraph/Public/Factory/AnimNextFactoryParams.h
new file mode 100644
--- /dev/null
+++ b/Engine/Plugins/Experimental/AnimNextAnimGraph/Source/AnimNextAnimGraph/Public/Factory/AnimNextFactoryParams.h
@@ -0,0 +1,99 @@
+﻿// Copyright Epic Games, Inc. All Rights Reserved.
+
+#pragma once
+
+#include "AnimNextAnimGraphBuilder.h"
+#include "InstanceTask.h"
+#include "InstanceTaskContext.h"
+#include "AnimNextSimpleAnimGraphBuilder.h"
+#include "TraitCore/TraitSharedData.h"
+#include "AnimNextFactoryParams.generated.h"
+
+struct FAnimNextAnimGraphFactoryMethod;
+struct FAnimNextGraphInstance;
+class UAnimNextAnimGraphSettings;
+class UAnimNextRigVMAsset;
+struct FAnimNextFactoryParams;
+struct FAnimNextGraphInstanceTaskInstanceData;
+
+namespace UE::AnimNext
+{
+	struct FAnimGraphFactory;
+}
+
+namespace UE::AnimNext::Editor
+{
+	class FAnimNextFactoryParamsDetails;
+};
+
+// Set of parameters that describe a way of generating & initializing an animation graph
+// TODO: To add an internal trait, use AddInternalTrait.
+// TODO: Interface to map public interface to internal traits 
+USTRUCT(BlueprintType)
+struct FAnimNextFactoryParams
+{
+	GENERATED_BODY()
+
+	FAnimNextFactoryParams() = default;
+
+	// Whether these params are valid for application
+	bool IsValid() const
+	{
+		return Builder.IsValid();
+	}
+
+	// Push a trait struct onto the specified stack in the graph builder, initialized to the struct's defaults.
+	// The trait struct forms part of the created instance's 'public interface'.
+	// If a trait already exists of this type then its values be overriden by virtue of the initializer running later
+	// Returns self for chaining.
+	template<typename StructType>
+	FAnimNextFactoryParams& PushPublicTraitStruct(int32 InStackIndex = 0)
+	{
+		TInstancedStruct<FAnimNextTraitSharedData> InstancedStruct;
+		InstancedStruct.InitializeAsScriptStruct(TBaseStructure<StructType>::Get());
+		Builder.PushTraitInstancedStructToStack(InStackIndex, MoveTemp(InstancedStruct));
+		return *this;
+	}
+
+	// Push a trait struct onto the specified stack in the graph builder, initialized to the supplied value
+	// The trait struct forms part of the created instance's 'public interface'.
+	// If a trait already exists of this type then its values be overriden by virtue of the initializer running later
+	// Returns self for chaining.
+	template<typename StructType>
+	FAnimNextFactoryParams& PushPublicTrait(const StructType& InValue, int32 InStackIndex = 0)
+	{
+		Builder.PushTraitStructViewToStack(InStackIndex, InValue);
+		return *this;
+	}
+
+	// Get the builder that these params hold
+	const UE::AnimNext::IAnimGraphBuilder& GetBuilder() const
+	{
+		return Builder;
+	}
+
+	// Add an initialize task used to set up an instance when these params are applied
+	FAnimNextFactoryParams& AddInitializeTask(UE::AnimNext::FInstanceTask&& InTask)
+	{
+		InitializeTasks.Add(MoveTemp(InTask));
+		return *this;
+	}
+
+	// Apply these params to the supplied instance.
+	ANIMNEXTANIMGRAPH_API void InitializeInstance(FUAFAssetInstance& InInstance) const;
+
+	// Clear these parameters back to empty
+	ANIMNEXTANIMGRAPH_API void Reset();
+
+private:
+	friend UE::AnimNext::FAnimGraphFactory;
+	friend UE::AnimNext::Editor::FAnimNextFactoryParamsDetails;
+	friend FAnimNextGraphInstanceTaskInstanceData;	// For old data upgrade
+
+	// Method used to create a graph
+	UPROPERTY(EditAnywhere, Category = Parameters)
+	FAnimNextSimpleAnimGraphBuilder Builder;
+
+	// Tasks to run on the factory-generated graph instance via InitializeInstance, used for initial setup of variables etc.
+	TArray<UE::AnimNext::FInstanceTask> InitializeTasks;
+};
diff --git a/Engine/Plugins/Experimental/AnimNextAnimGraph/Source/AnimNextAnimGraph/Public/Factory/AnimNextSimpleAnimGraphBuilder.h b/Engine/Plugins/Experimental/AnimNextAnimGraph/Source/AnimNextAnimGraph/Public/Factory/AnimNextSimpleAnimGraphBuilder.h
new file mode 100644
--- /dev/null
+++ b/Engine/Plugins/Experimental/AnimNextAnimGraph/Source/AnimNextAnimGraph/Public/Factory/AnimNextSimpleAnimGraphBuilder.h
@@ -0,0 +1,88 @@
+﻿// Copyright Epic Games, Inc. All Rights Reserved.
+
+#pragma once
+
+#include "AnimNextAnimGraphBuilder.h"
+#include "StructUtils/StructView.h"
+#include "AnimNextSimpleAnimGraphBuilder.generated.h"
+
+struct FAnimNextTraitSharedData;
+struct FAnimNextSimpleAnimGraphBuilder;
+class FAnimationAnimNextRuntimeTest_SimpleGraphBuilder;
+struct FAnimNextGraphInstanceTaskInstanceData;
+struct FAnimNextFactoryParams;
+
+namespace UE::AnimNext
+{
+	struct FAnimGraphFactory;
+};
+
+namespace UE::AnimNext::Editor
+{
+	class FAnimNextFactoryParamsDetails;
+};
+
+USTRUCT()
+struct FAnimNextSimpleAnimGraphBuilderTraitStackDesc
+{
+	GENERATED_BODY()
+
+private:
+	friend UE::AnimNext::FAnimGraphFactory;
+	friend FAnimNextSimpleAnimGraphBuilder;
+	friend FAnimNextGraphInstanceTaskInstanceData;	// For old data upgrade
+	friend UE::AnimNext::Editor::FAnimNextFactoryParamsDetails;
+	friend FAnimNextFactoryParams;
+
+	// Traits for this stack
+	UPROPERTY(EditAnywhere, Category = Parameters)
+	TArray<TInstancedStruct<FAnimNextTraitSharedData>> TraitStructs;
+};
+
+USTRUCT()
+struct FAnimNextSimpleAnimGraphBuilder : public FAnimNextAnimGraphBuilder
+{
+	GENERATED_BODY()
+
+	// Push a trait struct onto the supplied stack. Traits are ordered from bottom to top, stack-wise.
+	// Struct must be a sub-struct of FAnimNextTraitSharedData
+	ANIMNEXTANIMGRAPH_API void PushTraitStructViewToStack(int32 InStackIndex, TConstStructView<FAnimNextTraitSharedData> InStruct);
+	ANIMNEXTANIMGRAPH_API void PushTraitInstancedStructToStack(int32 InStackIndex, TInstancedStruct<FAnimNextTraitSharedData>&& InStruct);
+	
+	// Get the number of stacks from this graph
+	int32 GetNumStacks() const
+	{
+		return Stacks.Num();
+	}
+
+	// Check whether this method is empty or not
+	bool IsValid() const
+	{
+		ensure(Stacks.Num() <= 1);	// Currently only support a single stack
+		return Stacks.Num() > 0;
+	}
+
+	// Reset back to empty state
+	void Reset()
+	{
+		Stacks.Empty();
+	}
+	
+private:
+	friend UE::AnimNext::FAnimGraphFactory;
+	friend FAnimationAnimNextRuntimeTest_SimpleGraphBuilder;
+	friend FAnimNextGraphInstanceTaskInstanceData;	// For old data upgrade
+	friend UE::AnimNext::Editor::FAnimNextFactoryParamsDetails;
+	friend FAnimNextFactoryParams;
+
+	// Internal helper
+	void ValidateTraitStruct(int32 InStackIndex, TConstStructView<FAnimNextTraitSharedData> InStruct);
+
+	// FAnimNextAnimGraphBuilder interface
+	ANIMNEXTANIMGRAPH_API virtual bool Build(UE::AnimNext::FAnimGraphBuilderContext& InContext) const override;
+	ANIMNEXTANIMGRAPH_API virtual uint64 RecalculateKey() const override;
+
+	// The trait stacks we will use to generate our data
+	UPROPERTY(EditAnywhere, Category = Parameters)
+	TArray<FAnimNextSimpleAnimGraphBuilderTraitStackDesc> Stacks;
+};
\ No newline at end of file
diff --git a/Engine/Plugins/Experimental/AnimNextAnimGraph/Source/AnimNextAnimGraph/Public/Graph/AnimNextAnimGraph.h b/Engine/Plugins/Experimental/AnimNextAnimGraph/Source/AnimNextAnimGraph/Public/Graph/AnimNextAnimGraph.h
--- a/Engine/Plugins/Experimental/AnimNextAnimGraph/Source/AnimNextAnimGraph/Public/Graph/AnimNextAnimGraph.h
+++ b/Engine/Plugins/Experimental/AnimNextAnimGraph/Source/AnimNextAnimGraph/Public/Graph/AnimNextAnimGraph.h
@@ -42,6 +42,11 @@ namespace UE::AnimNext::Editor
 	class FAnimNextGraphDetails;
 }
 
+namespace UE::AnimNext::AnimGraph
+{
+	class FAnimNextAnimGraphModule;
+}
+
 // Injection data for an AnimNext Graph 
 USTRUCT(BlueprintType)
 struct FAnimNextGraphInjectionData
@@ -80,7 +85,7 @@ private:
 
 	// The asset to run as an animation graph
 	UPROPERTY(EditAnywhere, BlueprintReadWrite, Category=Graph, meta=(GetAllowedClasses="/Script/AnimNextAnimGraph.AnimNextAnimGraphSettings:GetAllowedAssetClasses", AllowPrivateAccess))
-	TObjectPtr<UObject> Asset;
+	TObjectPtr<const UObject> Asset;
 
 	// Injection data used to override the supplied graph
 	UPROPERTY(VisibleAnywhere, BlueprintReadOnly, Category=Graph, meta=(AllowPrivateAccess))
@@ -100,6 +105,7 @@ private:
 	friend FAnimNextModuleInjectionComponent;
 	friend UE::AnimNext::Editor::FAnimNextGraphDetails;
 	friend FRigUnit_GetPostProcessAnimation;
+	friend UE::AnimNext::AnimGraph::FAnimNextAnimGraphModule;
 };
 
 #undef UE_API
diff --git a/Engine/Plugins/Experimental/AnimNextAnimGraph/Source/AnimNextAnimGraph/Public/Graph/AnimNextAnimationGraph.h b/Engine/Plugins/Experimental/AnimNextAnimGraph/Source/AnimNextAnimGraph/Public/Graph/AnimNextAnimationGraph.h
--- a/Engine/Plugins/Experimental/AnimNextAnimGraph/Source/AnimNextAnimGraph/Public/Graph/AnimNextAnimationGraph.h
+++ b/Engine/Plugins/Experimental/AnimNextAnimGraph/Source/AnimNextAnimGraph/Public/Graph/AnimNextAnimationGraph.h
@@ -4,7 +4,7 @@
 
 #include "CoreMinimal.h"
 #include "Graph/AnimNextGraphEntryPoint.h"
-#include "DataInterface/AnimNextDataInterface.h"
+#include "Variables/AnimNextSharedVariables.h"
 #include "RigVMCore/RigVM.h"
 #include "TraitCore/TraitPtr.h"
 #include "TraitCore/TraitHandle.h"
@@ -37,6 +37,10 @@ namespace UE::AnimNext
 	struct FParametersProxy;
 	struct FPlayAnimSlotTrait;
 	struct FBlendStackCoreTrait;
+	struct FVariableOverrides;
+	struct FAnimGraphFactory;
+	struct FAnimGraphBuilderContext;
+	struct FVariableOverridesCollection;
 }
 
 namespace UE::AnimNext::UncookedOnly
@@ -61,9 +65,27 @@ namespace UE::AnimNext::Graph
 	extern ANIMNEXTANIMGRAPH_API const FName ResultName;
 }
 
+namespace UE::AnimNext
+{
+	// Parameters to AllocateInstance
+	struct FGraphAllocationParams
+	{
+		// The module instance to use as our outermost host. Usually non-null except in special cases (e.g. unit tests)
+		FAnimNextModuleInstance* ModuleInstance = nullptr;
+		// The current execution context we are running in, if any
+		FExecutionContext* ParentContext = nullptr;
+		// The hosting/parent graph instance
+		FAnimNextGraphInstance* ParentGraphInstance = nullptr;
+		// The entry point to use (deprecated)
+		FName EntryPoint = NAME_None;
+		// Any variable overrides to apply to this graph
+		TSharedPtr<const UE::AnimNext::FVariableOverridesCollection> Overrides;
+	};
+}
+
 // A user-created collection of animation logic & data
 UCLASS(MinimalAPI, BlueprintType)
-class UAnimNextAnimationGraph : public UAnimNextDataInterface
+class UAnimNextAnimationGraph : public UAnimNextSharedVariables
 {
 	GENERATED_BODY()
 
@@ -76,12 +98,10 @@ public:
 	UE_API virtual void GetPreloadDependencies(TArray<UObject*>& OutDeps) override;
 #endif
 
-	// Allocates an instance of the graph with an optional module and/or parent graph instance
-	// @param	InModuleInstance			The module instance to use
-	// @param	InParentGraphInstance		The parent graph instance to use
-	// @param	InEntryPoint				The entry point to use. If this is NAME_None then the default entry point for this graph is used
+	// Allocates an instance of the graph according to the supplied parameters
+	// @param	InParams		Parameters to allocate the graph with
 	// @return the allocated instance. This can be invalid if the entry point was not found or the graph's root trait is invalid.
-	UE_API TSharedPtr<FAnimNextGraphInstance> AllocateInstance(FAnimNextModuleInstance* InModuleInstance = nullptr, UE::AnimNext::FExecutionContext* InParentContext = nullptr, FAnimNextGraphInstance* InParentGraphInstance = nullptr, FName InEntryPoint = NAME_None) const;
+	UE_API TSharedPtr<FAnimNextGraphInstance> AllocateInstance(const UE::AnimNext::FGraphAllocationParams& InParams = UE::AnimNext::FGraphAllocationParams()) const;
 
 private:
 	// Loads the graph data from the provided archive buffer and returns true on success, false otherwise
@@ -112,6 +132,7 @@ private:
 	friend struct UE::AnimNext::FPlayAnimSlotTrait;
 	friend struct UE::AnimNext::FBlendStackCoreTrait;
 	friend struct FAnimNextModuleAnimGraphComponent;
+	friend UE::AnimNext::FAnimGraphBuilderContext;
 
 protected:
 #if WITH_EDITORONLY_DATA
diff --git a/Engine/Plugins/Experimental/AnimNextAnimGraph/Source/AnimNextAnimGraph/Public/Graph/AnimNextGraphEntryPoint.h b/Engine/Plugins/Experimental/AnimNextAnimGraph/Source/AnimNextAnimGraph/Public/Graph/AnimNextGraphEntryPoint.h
--- a/Engine/Plugins/Experimental/AnimNextAnimGraph/Source/AnimNextAnimGraph/Public/Graph/AnimNextGraphEntryPoint.h
+++ b/Engine/Plugins/Experimental/AnimNextAnimGraph/Source/AnimNextAnimGraph/Public/Graph/AnimNextGraphEntryPoint.h
@@ -19,9 +19,5 @@ struct FAnimNextGraphEntryPoint
 	// This is a handle to the root trait for a graph
 	UPROPERTY()
 	FAnimNextEntryPointHandle RootTraitHandle;
-
-	// Default state for this entry point
-	UPROPERTY()
-	FAnimNextGraphState DefaultState;
 };
 
