AnimNext: Pass-by-ref runtime refactor

diff --git a/Engine/Plugins/Experimental/AnimNextAnimGraph/Source/AnimNextAnimGraph/Public/Graph/AnimNextGraphInstance.h b/Engine/Plugins/Experimental/AnimNextAnimGraph/Source/AnimNextAnimGraph/Public/Graph/AnimNextGraphInstance.h
--- a/Engine/Plugins/Experimental/AnimNextAnimGraph/Source/AnimNextAnimGraph/Public/Graph/AnimNextGraphInstance.h
+++ b/Engine/Plugins/Experimental/AnimNextAnimGraph/Source/AnimNextAnimGraph/Public/Graph/AnimNextGraphInstance.h
@@ -3,7 +3,7 @@
 #pragma once
 
 #include "CoreMinimal.h"
-#include "DataInterface/AnimNextDataInterfaceInstance.h"
+#include "UAFAssetInstance.h"
 #include "HAL/CriticalSection.h"
 #include "TraitCore/TraitEvent.h"
 #include "TraitCore/TraitEventList.h"
@@ -46,13 +46,11 @@ namespace UE::AnimNext
 	struct FBlendSpacePlayerTrait;
 }
 
-using GraphInstanceComponentMapType = TMap<FName, TSharedPtr<UE::AnimNext::FGraphInstanceComponent>>;
-
 // Represents an instance of an AnimNext graph
 // This struct uses UE reflection because we wish for the GC to keep the graph
 // alive while we own a reference to it. It is not intended to be serialized on disk with a live instance.
 USTRUCT()
-struct FAnimNextGraphInstance : public FAnimNextDataInterfaceInstance
+struct FAnimNextGraphInstance : public FUAFAssetInstance
 {
 	GENERATED_BODY()
 
@@ -98,45 +96,22 @@ struct FAnimNextGraphInstance : public FAnimNextDataInterfaceInstance
 	
 	// Returns whether or not this graph instance is the root graph instance or false otherwise
 	UE_API bool IsRoot() const;
-
-	// Adds strong/hard object references during GC
-	UE_API void AddStructReferencedObjects(class FReferenceCollector& Collector);
-
-	// Returns a typed graph instance component, creating it lazily the first time it is queried
-	template<class ComponentType>
-	ComponentType& GetComponent();
-
-	// Returns a typed graph instance component pointer if found or nullptr otherwise
-	template<class ComponentType>
-	ComponentType* TryGetComponent();
-
-	// Returns a typed graph instance component pointer if found or nullptr otherwise
-	template<class ComponentType>
-	const ComponentType* TryGetComponent() const;
-
-	// Returns const iterators to the graph instance component container
-	UE_API GraphInstanceComponentMapType::TConstIterator GetComponentIterator() const;
-
+	
 	// Returns whether or not this graph instance has updated at least once
 	UE_API bool HasUpdated() const;
 
 	// Called each time the graph updates to mark the instance as updated
 	UE_API void MarkAsUpdated();
 
-	// Get a reference to the internal public variables
-	const FInstancedPropertyBag& GetVariables() const { return Variables; }
-
 private:
-	// Returns a pointer to the specified component, or nullptr if not found
-	UE_API UE::AnimNext::FGraphInstanceComponent* TryGetComponent(int32 ComponentNameHash, FName ComponentName) const;
-
-	// Adds the specified component and returns a reference to it
-	UE_API UE::AnimNext::FGraphInstanceComponent& AddComponent(int32 ComponentNameHash, FName ComponentName, TSharedPtr<UE::AnimNext::FGraphInstanceComponent>&& Component);
-
 	// Executes a list of latent RigVM pins and writes the result into the destination pointer (latent handle offsets are using the destination as base)
 	// When frozen, latent handles that can freeze are skipped, all others will execute
 	UE_API void ExecuteLatentPins(const TConstArrayView<UE::AnimNext::FLatentPropertyHandle>& LatentHandles, void* DestinationBasePtr, bool bIsFrozen);
 
+	// When a group of latent pins do not require RigVM to operate, this copies the values directly from variables rather than calling ExecuteLatentPins.
+	// When frozen, latent handles that can freeze are skipped, all others will execute
+	void CopyVariablesToLatentPins(const TConstArrayView<UE::AnimNext::FLatentPropertyHandle>& LatentHandles, void* DestinationBasePtr, bool bIsFrozen);
+
 #if WITH_EDITORONLY_DATA
 	// During graph compilation, if we have existing graph instances, we freeze them by releasing their memory before thawing them
 	// Freezing is a partial release of resources that retains the necessary information to re-create things safely
@@ -149,28 +124,6 @@ private:
 	UE_API void OnModuleCompiled(UAnimNextModule* InModule);
 #endif
 
-	// Check if variable binding is required
-	bool RequiresPublicVariableBinding() const
-	{
-		return PublicVariablesState == EPublicVariablesState::Unbound;
-	}
-
-	// Helper function used to bind to cached bindings after they are built
-	UE_API bool BindToCachedBindings();
-
-	// Bind the variables in the supplied traits in scope to their respective public variables, represented by structs held on a host
-	// Note: struct views are mutable here as bindings allow writes!
-	UE_API void BindPublicVariables(TArrayView<FStructView> InHostStructs, TConstArrayView<UE::AnimNext::IDataInterfaceHost*> InHosts = TConstArrayView<UE::AnimNext::IDataInterfaceHost*>());
-
-	// Bind the variables in the supplied traits in scope to their respective public variables, so they point at host memory
-	UE_API void BindPublicVariables(TConstArrayView<UE::AnimNext::IDataInterfaceHost*> InHosts);
-
-	// Unbind any public variables that were pointing at host memory and re-point them at the internal defaults
-	UE_API void UnbindPublicVariables();
-
-	// Get a reference to the internal public variables
-	FInstancedPropertyBag& GetMutableVariables() { return Variables; }
-
 	// The entry point in Graph that this instance corresponds to 
 	FName EntryPoint;
 
@@ -182,63 +135,7 @@ private:
 
 	// The root graph instance that owns us and the components
 	FAnimNextGraphInstance* RootGraphInstance = nullptr;
-
-	struct FCachedDataInterfaceBinding
-	{
-		struct FVariable
-		{
-			FVariable() = default;
-
-			FVariable(const FName& InVariableName, int32 InInterfaceVariableIndex, const FProperty* InProperty, uint8* InMemory)
-				: VariableName(InVariableName)
-				, InterfaceVariableIndex(InInterfaceVariableIndex)
-				, Property(InProperty)
-				, Memory(InMemory)
-			{
-			}
-
-			// Name of the variable
-			FName VariableName;
-
-			// Index of the variable in the public variables property bag for DataInterface
-			int32 InterfaceVariableIndex = INDEX_NONE;
-
-			// Cached property of the variable, used for type validation
-			const FProperty* Property = nullptr;
-
-			// Memory of the variable, bound to the 'outermost' container's property bag that implements DataInterface
-			uint8* Memory = nullptr;
-		};
-
-		const UAnimNextDataInterface* DataInterface = nullptr;
-		TArray<FVariable> CachedBindings;
-	};
-
-	UE_API void UpdateCachedBindingsForHost();
-	UE_API void UpdateCachedBindingsForHost(const FAnimNextModuleInstance& InHost);
-	UE_API void UpdateCachedBindingsForHost(const UE::AnimNext::IDataInterfaceHost& InHost);
-
-	// Helper function used to cache a host's bindings, overriding any of the already cached bindings that are valid.
-	// Template to allow binding to both FAnimNextDataInterfaceInstance and IDataInterfaceHost without an abstraction between the two.
-	template<typename HostType>
-	void UpdateCachedBindingsForHostHelper(const HostType& InHost);
-
-	// Bindings we have to our outer host(s), per data interface
-	TArray<FCachedDataInterfaceBinding> CachedBindings;
-
-	// Graph instance components that persist from update to update
-	GraphInstanceComponentMapType Components;
-
-	// The current state of public variable bindings to the host
-	enum class EPublicVariablesState : uint8
-	{
-		None,		// No public variables present
-		Unbound,	// Present, but currently unbound
-		Bound		// Present and bound
-	};
-
-	EPublicVariablesState PublicVariablesState;
-
+	
 	// Whether or not this graph has updated once
 	bool bHasUpdatedOnce : 1 = false;
 
@@ -265,43 +162,8 @@ struct TStructOpsTypeTraits<FAnimNextGraphInstance> : public TStructOpsTypeTrait
 {
 	enum
 	{
-		WithAddStructReferencedObjects = true,
 		WithCopy = false,
 	};
 };
 
-//////////////////////////////////////////////////////////////////////////
-
-template<class ComponentType>
-ComponentType& FAnimNextGraphInstance::GetComponent()
-{
-	const FName ComponentName = ComponentType::StaticComponentName();
-	const int32 ComponentNameHash = GetTypeHash(ComponentName);
-
-	if (UE::AnimNext::FGraphInstanceComponent* Component = TryGetComponent(ComponentNameHash, ComponentName))
-	{
-		return *static_cast<ComponentType*>(Component);
-	}
-
-	return static_cast<ComponentType&>(AddComponent(ComponentNameHash, ComponentName, MakeShared<ComponentType>(*this)));
-}
-
-template<class ComponentType>
-ComponentType* FAnimNextGraphInstance::TryGetComponent()
-{
-	const FName ComponentName = ComponentType::StaticComponentName();
-	const int32 ComponentNameHash = GetTypeHash(ComponentName);
-
-	return static_cast<ComponentType*>(TryGetComponent(ComponentNameHash, ComponentName));
-}
-
-template<class ComponentType>
-const ComponentType* FAnimNextGraphInstance::TryGetComponent() const
-{
-	const FName ComponentName = ComponentType::StaticComponentName();
-	const int32 ComponentNameHash = GetTypeHash(ComponentName);
-
-	return static_cast<ComponentType*>(TryGetComponent(ComponentNameHash, ComponentName));
-}
-
 #undef UE_API
diff --git a/Engine/Plugins/Experimental/AnimNextAnimGraph/Source/AnimNextAnimGraph/Public/Injection/InjectionRequest.h b/Engine/Plugins/Experimental/AnimNextAnimGraph/Source/AnimNextAnimGraph/Public/Injection/InjectionRequest.h
--- a/Engine/Plugins/Experimental/AnimNextAnimGraph/Source/AnimNextAnimGraph/Public/Injection/InjectionRequest.h
+++ b/Engine/Plugins/Experimental/AnimNextAnimGraph/Source/AnimNextAnimGraph/Public/Injection/InjectionRequest.h
@@ -5,12 +5,15 @@
 #include "CoreMinimal.h"
 
 #include "AlphaBlend.h"
-#include "AnimNextDataInterfacePayload.h"
+#include "Factory/AnimNextFactoryParams.h"
+#include "UAFAssetInstance.h"
+#include "InstanceTask.h"
 #include "TraitCore/TraitEvent.h"
 #include "TraitInterfaces/ITimeline.h"
 #include "Injection/InjectionStatus.h"
 #include "Injection/InjectionSite.h"
 #include "Module/ModuleHandle.h"
+#include "Containers/SpscQueue.h"
 
 #include "InjectionRequest.generated.h"
 
@@ -22,6 +25,7 @@ class UWorld;
 
 namespace UE::AnimNext
 {
+	struct FInstanceTaskContext;
 	struct IEvaluationModifier;
 	struct FInjection_InjectEvent;
 	struct FInjection_StatusUpdateEvent;
@@ -126,9 +130,6 @@ private:
 	UPROPERTY()
 	TObjectPtr<UObject> Object;
 
-	// Module handle to use when injecting. This will supply the source binding for data interfaces when the graph is instantiated.
-	UE::AnimNext::FModuleHandle BindingModuleHandle;
-
 	// Evaluation modifier to apply at the injection site. Shared ownership, references can persist in worker threads until un-injected.
 	TSharedPtr<UE::AnimNext::IEvaluationModifier> EvaluationModifier;
 
@@ -143,9 +144,9 @@ private:
 	UPROPERTY()
 	bool bTrackTimelineProgress = false;
 
-	// Payload that will be applied to the animation graph's variables via its data interfaces.
+	// Set of assets & structs that describe the required variables for factory-generating injected graphs
 	UPROPERTY()
-	FAnimNextDataInterfacePayload Payload;
+	FAnimNextFactoryParams FactoryParams;
 };
 
 namespace UE::AnimNext
@@ -219,6 +220,44 @@ namespace UE::AnimNext
 
 		ANIMNEXTANIMGRAPH_API void ExternalAddReferencedObjects(FReferenceCollector& Collector);
 
+		// Set a variable on the running injected instance.
+		// @param    InVariable         The variable to set
+		// @param    InValue            The value to set the variable to
+		template<typename ValueType>
+		void SetVariable(const FAnimNextVariableReference& InVariable, const ValueType& InValue)
+		{
+			QueueTask([InVariable, InValue](const FInstanceTaskContext& InContext)
+			{
+				InContext.SetVariable(InVariable, InValue);
+			});
+		}
+
+		// Access a variable on the running injected instance for modification. Can avoid excess copies for larger structures/arrays.
+		// @param    InVariable         The variable to set
+		// @param    InFunction         The function used to modify the variable. Note: called back at some later time on a worker thread
+		template<typename ValueType>
+		void AccessVariable(const FAnimNextVariableReference& InVariable, TUniqueFunction<void(ValueType&)> InFunction)
+		{
+			QueueTask([InVariable, Function = MoveTemp(InFunction)](const FInstanceTaskContext& InContext)
+			{
+				InContext.AccessVariable(InVariable, Function);
+			});
+		}
+
+		// Access the memory of the shared variable struct directly on the running injected instance.
+		// @param	InFunction			Function called with a reference to the variable's struct. Note: called back at some later time on a worker thread
+		template<typename StructType>
+		void AccessVariablesStruct(TUniqueFunction<void(StructType&)> InFunction) const
+		{
+			QueueTask([Function = MoveTemp(InFunction)](const FInstanceTaskContext& InContext)
+			{
+				InContext.AccessVariablesStruct(Function);
+			});
+		}
+
+		// Queues a task to run on the injected instance prior to execution. Use this to set variables etc.
+		ANIMNEXTANIMGRAPH_API void QueueTask(FUniqueInstanceTask&& InTask);
+
 	private:
 		// Sends this request to the specified host and it will attempt to play with the requested arguments
 		bool Inject(FInjectionRequestArgs&& InRequestArgs, FInjectionLifetimeEvents&& InLifetimeEvents, UObject* InHost, FModuleHandle InHandle);
@@ -242,11 +281,8 @@ namespace UE::AnimNext
 		// Validate this set of args is set up correctly for injection
 		static bool ValidateArgs(const FInjectionRequestArgs& InRequestArgs);
 
-		// Setup dependencies between the module being injected into and any module used for data interface bindings
-		static void SetupBindingModuleDependencies(UWorld* InWorld, FModuleHandle InHandle, const FInjectionRequestArgs& InRequestArgs);
-
-		// Remove dependencies between the module being injected into and any module used for data interface bindings
-		static void RemoveBindingModuleDependencies(UWorld* InWorld, FModuleHandle InHandle, const FInjectionRequestArgs& InRequestArgs);
+		// Flush any pending tasks, with the initialize task on the request args getting run once first
+		void FlushTasks(const FInstanceTaskContext& InContext);
 
 		// The request arguments
 		FInjectionRequestArgs RequestArgs;
@@ -266,6 +302,9 @@ namespace UE::AnimNext
 		// The injection event if we have injected already
 		TWeakPtr<FAnimNextTraitEvent> InjectionEvent;
 
+		// Queue of tasks to be executed next time this injection updates
+		TSpscQueue<FUniqueInstanceTask> TaskQueue;
+
 		// The current request status
 		EInjectionStatus Status = EInjectionStatus::None;
 
diff --git a/Engine/Plugins/Experimental/AnimNextAnimGraph/Source/AnimNextAnimGraph/Public/Injection/InjectionSite.h b/Engine/Plugins/Experimental/AnimNextAnimGraph/Source/AnimNextAnimGraph/Public/Injection/InjectionSite.h
--- a/Engine/Plugins/Experimental/AnimNextAnimGraph/Source/AnimNextAnimGraph/Public/Injection/InjectionSite.h
+++ b/Engine/Plugins/Experimental/AnimNextAnimGraph/Source/AnimNextAnimGraph/Public/Injection/InjectionSite.h
@@ -4,39 +4,60 @@
 
 #include "CoreTypes.h"
 #include "UObject/NameTypes.h"
+#include "Variables/AnimNextVariableReference.h"
 
 #include "InjectionSite.generated.h"
 
 // Specifies the desired injection site for an injection request.
 // Optionally allows for fallback to the Default Injection Site specified by the module
-//     if DesiredSiteName is None, or does not exist in the graph
+//     if DesiredSiteName is none, or does not exist in the graph
 USTRUCT()
 struct FAnimNextInjectionSite
 {
 	GENERATED_BODY()
 
 	FAnimNextInjectionSite()
-		: DesiredSiteName(NAME_None)
-		, bUseModuleFallback(true)
+		: bUseModuleFallback(true)
 	{
 	}
 
-	FAnimNextInjectionSite(const FName& InSiteName)
-		: DesiredSiteName(InSiteName)
-		, bUseModuleFallback(DesiredSiteName == NAME_None)
+	explicit FAnimNextInjectionSite(const FAnimNextVariableReference& InDesiredSite)
+		: DesiredSite(InDesiredSite)
+		, bUseModuleFallback(!InDesiredSite.IsNone())
 	{
 	}
 
+#if WITH_EDITORONLY_DATA
 	// The name of the site to inject into
 	UPROPERTY()
-	FName DesiredSiteName;
+	FName DesiredSiteName_DEPRECATED;
+#endif
+
+	UPROPERTY()
+	FAnimNextVariableReference DesiredSite;
 
 	// Flag specifying whether the request can (or should) fallback to the Default Injection Site
 	// specified in the module if DesiredSiteName was not found
 	UPROPERTY()
 	uint8 bUseModuleFallback : 1;
+
+#if WITH_EDITORONLY_DATA
+	bool Serialize(FArchive& Ar);
+	void PostSerialize(const FArchive& Ar);
+#endif
 };
 
+template<>
+struct TStructOpsTypeTraits<FAnimNextInjectionSite> : public TStructOpsTypeTraitsBase2<FAnimNextInjectionSite>
+{
+#if WITH_EDITORONLY_DATA
+	enum
+	{
+		WithSerializer = true,
+		WithPostSerialize = true,
+	};
+#endif
+};
 
 namespace UE::AnimNext
 {
diff --git a/Engine/Plugins/Experimental/AnimNextAnimGraph/Source/AnimNextAnimGraph/Public/Injection/InjectionUtils.h b/Engine/Plugins/Experimental/AnimNextAnimGraph/Source/AnimNextAnimGraph/Public/Injection/InjectionUtils.h
--- a/Engine/Plugins/Experimental/AnimNextAnimGraph/Source/AnimNextAnimGraph/Public/Injection/InjectionUtils.h
+++ b/Engine/Plugins/Experimental/AnimNextAnimGraph/Source/AnimNextAnimGraph/Public/Injection/InjectionUtils.h
@@ -4,9 +4,10 @@
 
 #include "CoreMinimal.h"
 #include "InjectionRequest.h"
+#include "Factory/AnimNextFactoryParams.h"
 #include "Injection/InjectionSite.h"
 
-struct FAnimNextDataInterfaceInstance;
+struct FUAFAssetInstance;
 class UAnimNextComponent;
 
 namespace UE::AnimNext
@@ -34,6 +35,7 @@ namespace UE::AnimNext
 	// Optional arguments to PlayAnim
 	struct FPlayAnimArgs
 	{
+		FAnimNextFactoryParams FactoryParams;
 		float PlayRate = 1.0f;
 		float StartPosition = 0.0f;
 		ELoopMode LoopMode = ELoopMode::Auto;
@@ -69,7 +71,7 @@ namespace UE::AnimNext
 			UAnimNextComponent* InComponent,
 			const FInjectionSite& InSite,
 			UAnimSequence* InAnimSequence,
-			const FPlayAnimArgs& InArgs = FPlayAnimArgs(),
+			FPlayAnimArgs&& InArgs = FPlayAnimArgs(),
 			const FInjectionBlendSettings& InBlendInSettings = FInjectionBlendSettings(),
 			const FInjectionBlendSettings& InBlendOutSettings = FInjectionBlendSettings(),
 			FInjectionLifetimeEvents&& InLifetimeEvents = FInjectionLifetimeEvents());
@@ -80,7 +82,7 @@ namespace UE::AnimNext
 			FModuleHandle InModuleHandle,
 			const FInjectionSite& InSite,
 			UAnimSequence* InAnimSequence,
-			const FPlayAnimArgs& InArgs = FPlayAnimArgs(),
+			FPlayAnimArgs&& InArgs = FPlayAnimArgs(),
 			const FInjectionBlendSettings& InBlendInSettings = FInjectionBlendSettings(),
 			const FInjectionBlendSettings& InBlendOutSettings = FInjectionBlendSettings(),
 			FInjectionLifetimeEvents&& InLifetimeEvents = FInjectionLifetimeEvents());
@@ -91,7 +93,7 @@ namespace UE::AnimNext
 		// @param InHandle             The module handle being injected into
 		// @param InSite               The injection site to 'play' the asset on
 		// @param InAsset              The asset to 'play'
-		// @param InPayload            The payload for the asset to use
+		// @param InFactoryParams      The params used to manufacture a graph to run the injected object 
 		// @param InBindingComponent   A component to use to supply additional data interface bindings at the injection site
 		// @param InBindingModuleHandle A module to use to supply additional data interface bindings at the injection site
 		// @param InBlendInSettings    The blend in settings
@@ -102,8 +104,7 @@ namespace UE::AnimNext
 			UAnimNextComponent* InComponent,
 			const FInjectionSite& InSite,
 			UObject* InAsset,
-			FAnimNextDataInterfacePayload&& InPayload = FAnimNextDataInterfacePayload(),
-			UAnimNextComponent* InBindingComponent = nullptr,
+			FAnimNextFactoryParams&& InFactoryParams = FAnimNextFactoryParams(),
 			const FInjectionBlendSettings& InBlendInSettings = FInjectionBlendSettings(),
 			const FInjectionBlendSettings& InBlendOutSettings = FInjectionBlendSettings(),
 			FInjectionLifetimeEvents&& InLifetimeEvents = FInjectionLifetimeEvents());
@@ -114,8 +115,7 @@ namespace UE::AnimNext
 			FModuleHandle InModuleHandle,
 			const FInjectionSite& InSite,
 			UObject* InAsset,
-			FAnimNextDataInterfacePayload&& InPayload = FAnimNextDataInterfacePayload(),
-			FModuleHandle InBindingModuleHandle = FModuleHandle(),
+			FAnimNextFactoryParams&& InFactoryParams = FAnimNextFactoryParams(),
 			const FInjectionBlendSettings& InBlendInSettings = FInjectionBlendSettings(),
 			const FInjectionBlendSettings& InBlendOutSettings = FInjectionBlendSettings(),
 			FInjectionLifetimeEvents&& InLifetimeEvents = FInjectionLifetimeEvents());
diff --git a/Engine/Plugins/Experimental/AnimNextAnimGraph/Source/AnimNextAnimGraph/Public/TraitCore/ExecutionContext.h b/Engine/Plugins/Experimental/AnimNextAnimGraph/Source/AnimNextAnimGraph/Public/TraitCore/ExecutionContext.h
--- a/Engine/Plugins/Experimental/AnimNextAnimGraph/Source/AnimNextAnimGraph/Public/TraitCore/ExecutionContext.h
+++ b/Engine/Plugins/Experimental/AnimNextAnimGraph/Source/AnimNextAnimGraph/Public/TraitCore/ExecutionContext.h
@@ -15,7 +15,7 @@
 #include "TraitCore/LatentPropertyHandle.h"
 #include "TraitCore/NodeHandle.h"
 #include "Graph/AnimNextGraphInstance.h"
-#include "Graph/GraphInstanceComponent.h"
+#include "Graph/UAFGraphInstanceComponent.h"
 #include "EngineDefines.h"
 
 #include <type_traits>
@@ -176,8 +176,8 @@ namespace UE::AnimNext
 		template<class ComponentType>
 		ComponentType* TryGetComponent() const;
 
-		// Returns const iterators to the graph instance component container
-		UE_API GraphInstanceComponentMapType::TConstIterator GetComponentIterator() const;
+		// Iterates each of the graph instance components
+		UE_API void ForEachComponent(TFunctionRef<bool(FUAFGraphInstanceComponent&)> InFunction) const;
 
 
 
@@ -222,8 +222,6 @@ namespace UE::AnimNext
 		UE_API bool PopScopedInterfaceImpl(FTraitInterfaceUID InterfaceUID, const FTraitBinding& Binding);
 		UE_API bool GetScopedInterfaceImpl(FTraitInterfaceUID InterfaceUID, FTraitBinding& OutBinding) const;
 		UE_API void ForEachScopedInterfaceImpl(FTraitInterfaceUID InterfaceUID, TFunctionRef<bool(FTraitBinding& Binding)> InFunction) const;
-		UE_API FGraphInstanceComponent* TryGetComponent(int32 ComponentNameHash, FName ComponentName) const;
-		UE_API FGraphInstanceComponent& AddComponent(int32 ComponentNameHash, FName ComponentName, TSharedPtr<FGraphInstanceComponent>&& Component) const;
 
 		UE_API const FNodeDescription& GetNodeDescription(const FAnimNextGraphInstance& GraphInstance, FNodeHandle NodeHandle) const;
 		UE_API const FNodeDescription& GetNodeDescription(const FNodeInstance& NodeInstance) const;
@@ -334,28 +332,17 @@ namespace UE::AnimNext
 	template<class ComponentType>
 	ComponentType& FExecutionContext::GetComponent() const
 	{
-		static_assert(std::is_base_of<FGraphInstanceComponent, ComponentType>::value, "ComponentType type must derive from FGraphInstanceComponent");
+		static_assert(std::is_base_of<FUAFGraphInstanceComponent, ComponentType>::value, "ComponentType type must derive from FGraphInstanceComponent");
 
-		const FName ComponentName = ComponentType::StaticComponentName();
-		const int32 ComponentNameHash = GetTypeHash(ComponentName);
-
-		if (FGraphInstanceComponent* Component = TryGetComponent(ComponentNameHash, ComponentName))
-		{
-			return *static_cast<ComponentType*>(Component);
-		}
-
-		return static_cast<ComponentType&>(AddComponent(ComponentNameHash, ComponentName, MakeShared<ComponentType>(*RootGraphInstance)));
+		return RootGraphInstance->GetComponent<ComponentType>();
 	}
 
 	template<class ComponentType>
 	ComponentType* FExecutionContext::TryGetComponent() const
 	{
-		static_assert(std::is_base_of<FGraphInstanceComponent, ComponentType>::value, "ComponentType type must derive from FGraphInstanceComponent");
-
-		const FName ComponentName = ComponentType::StaticComponentName();
-		const int32 ComponentNameHash = GetTypeHash(ComponentName);
+		static_assert(std::is_base_of<FUAFGraphInstanceComponent, ComponentType>::value, "ComponentType type must derive from FGraphInstanceComponent");
 
-		return static_cast<ComponentType*>(TryGetComponent(ComponentNameHash, ComponentName));
+		return RootGraphInstance->TryGetComponent<ComponentType>();
 	}
 }
 
diff --git a/Engine/Plugins/Experimental/AnimNextAnimGraph/Source/AnimNextAnimGraph/Public/TraitCore/LatentPropertyHandle.h b/Engine/Plugins/Experimental/AnimNextAnimGraph/Source/AnimNextAnimGraph/Public/TraitCore/LatentPropertyHandle.h
--- a/Engine/Plugins/Experimental/AnimNextAnimGraph/Source/AnimNextAnimGraph/Public/TraitCore/LatentPropertyHandle.h
+++ b/Engine/Plugins/Experimental/AnimNextAnimGraph/Source/AnimNextAnimGraph/Public/TraitCore/LatentPropertyHandle.h
@@ -25,6 +25,9 @@ namespace UE::AnimNext
 
 		// Whether or not the property supports freezing
 		bool bCanFreeze = false;
+
+		// Whether or not the property is a simple variable copy
+		bool bUsesVariableCopy = false;
 	};
 
 	// Serializes latent property metadata
@@ -33,6 +36,7 @@ namespace UE::AnimNext
 		Ar << Metadata.Name;
 		Ar << Metadata.RigVMIndex;
 		Ar << Metadata.bCanFreeze;
+		Ar << Metadata.bUsesVariableCopy;
 		return Ar;
 	}
 
@@ -45,10 +49,13 @@ namespace UE::AnimNext
 	struct alignas(uint16) FLatentPropertiesHeader final
 	{
 		// Whether or not this base trait has any bound latent properties that can evaluate
-		uint8 bHasValidLatentProperties = false;
+		uint8 bHasValidLatentProperties : 1 = false;
 
 		// Whether or not this base trait has all its latent properties support freezing
-		uint8 bCanAllPropertiesFreeze = false;
+		uint8 bCanAllPropertiesFreeze : 1 = false;
+
+		// Whether the latent properties for this trait are all simple variable accesses
+		uint8 bAreAllPropertiesVariableAccesses : 1 = false;
 	};
 
 	/**
diff --git a/Engine/Plugins/Experimental/AnimNextAnimGraph/Source/AnimNextAnimGraph/Public/TraitCore/Trait.h b/Engine/Plugins/Experimental/AnimNextAnimGraph/Source/AnimNextAnimGraph/Public/TraitCore/Trait.h
--- a/Engine/Plugins/Experimental/AnimNextAnimGraph/Source/AnimNextAnimGraph/Public/TraitCore/Trait.h
+++ b/Engine/Plugins/Experimental/AnimNextAnimGraph/Source/AnimNextAnimGraph/Public/TraitCore/Trait.h
@@ -46,16 +46,10 @@ struct FRigVMPinInfoArray;
 	virtual UE::AnimNext::ETraitStackPropagation OnTraitEvent(UE::AnimNext::FExecutionContext& Context, UE::AnimNext::FTraitBinding& Binding, FAnimNextTraitEvent& Event) const override; \
 	virtual TConstArrayView<UE::AnimNext::FTraitEventUID> GetTraitEvents() const override;
 
-#if WITH_EDITOR
 #define ANIM_NEXT_IMPL_DECLARE_ANIM_TRAIT_LATENT_PROPERTY_SUPPORT() \
 	virtual uint32 GetNumLatentTraitProperties() const override { return -FSharedData::GetLatentPropertyIndex(~(size_t)0); } \
 	virtual UE::AnimNext::FTraitLatentPropertyMemoryLayout GetLatentPropertyMemoryLayout(const FAnimNextTraitSharedData& SharedData, FName PropertyName, uint32 PropertyIndex) const override; \
 	virtual bool IsPropertyLatent(const FAnimNextTraitSharedData& InSharedData, FName PropertyName) const override;
-#else
-#define ANIM_NEXT_IMPL_DECLARE_ANIM_TRAIT_LATENT_PROPERTY_SUPPORT() \
-	virtual uint32 GetNumLatentTraitProperties() const override { return -FSharedData::GetLatentPropertyIndex(~(size_t)0); } \
-	virtual UE::AnimNext::FTraitLatentPropertyMemoryLayout GetLatentPropertyMemoryLayout(const FAnimNextTraitSharedData& SharedData, FName PropertyName, uint32 PropertyIndex) const override;
-#endif
 
 // Helper macros
 // In the trait class declaration, this macro declares the Super alias and base functions we override
@@ -118,7 +112,6 @@ struct FRigVMPinInfoArray;
 		return GetLatentPropertyMemoryLayoutImpl(PropertyName, PropertyIndex, CachedLatentPropertyMemoryLayouts); \
 	}
 
-#if WITH_EDITOR
 #define ANIM_NEXT_IMPL_DEFINE_ANIM_TRAIT_IS_PROPERTY_LATENT(TraitName) \
 	bool TraitName::IsPropertyLatent(const FAnimNextTraitSharedData& InSharedData, FName PropertyName) const \
 	{ \
@@ -127,9 +120,6 @@ struct FRigVMPinInfoArray;
 		check(Property != nullptr); \
 		return FSharedData::GetLatentPropertyIndex(Property->GetOffset_ForInternal()) > 0; \
 	}
-#else
-#define ANIM_NEXT_IMPL_DEFINE_ANIM_TRAIT_IS_PROPERTY_LATENT(TraitName)
-#endif
 
 // Helper that handles the GetTraitInterface() details for each interface specified by the generator macro
 #define ANIM_NEXT_IMPL_GET_INTERFACE_IMPL_FOR_INTERFACE(InterfaceName) \
@@ -454,9 +444,6 @@ namespace UE::AnimNext
 			bool bFilterEditorOnly,
 			const TFunction<uint16(FName PropertyName)>& GetTraitLatentPropertyIndex) const;
 
-		// Returns whether or not the specified property has been marked as latent
-		virtual bool IsPropertyLatent(const FAnimNextTraitSharedData& InSharedData, FName PropertyName) const { return false; }
-
 		// Makes the Trait Editor only display this Trait in Advanced view
 		virtual bool IsHidden() const { return false; }
 
@@ -467,6 +454,17 @@ namespace UE::AnimNext
 		virtual void GetProgrammaticPins(FAnimNextTraitSharedData* InSharedData, URigVMController* InController, int32 InParentPinIndex, const URigVMPin* InTraitPin, const FString& InDefaultValue, struct FRigVMPinInfoArray& OutPinArray) const {}
 #endif
 
+		// Takes the trait shared data and the editor properties as authored in the graph and generates the latent property metadata using UE reflection.
+		// Returns the number of latent property metadatas added to OutLatentPropertyHandles
+		uint32 GetVariableMappedLatentPropertyHandles(
+			const FAnimNextTraitSharedData* InSharedData,
+			TArray<FLatentPropertyMetadata>& OutLatentPropertyHandles,
+			bool bFilterEditorOnly,
+			const TFunction<uint16(FName PropertyName)>& GetTraitLatentPropertyIndex) const;
+
+		// Returns whether or not the specified property has been marked as latent
+		virtual bool IsPropertyLatent(const FAnimNextTraitSharedData& InSharedData, FName PropertyName) const { return false; }
+		
 	protected:
 		// Implements GetLatentPropertyMemoryLayout() by allowing a map to be provided for caching purposes to speed up repeated queries
 		UE_API FTraitLatentPropertyMemoryLayout GetLatentPropertyMemoryLayoutImpl(
diff --git a/Engine/Plugins/Experimental/AnimNextAnimGraph/Source/AnimNextAnimGraph/Public/TraitCore/TraitEventRaising.h b/Engine/Plugins/Experimental/AnimNextAnimGraph/Source/AnimNextAnimGraph/Public/TraitCore/TraitEventRaising.h
--- a/Engine/Plugins/Experimental/AnimNextAnimGraph/Source/AnimNextAnimGraph/Public/TraitCore/TraitEventRaising.h
+++ b/Engine/Plugins/Experimental/AnimNextAnimGraph/Source/AnimNextAnimGraph/Public/TraitCore/TraitEventRaising.h
@@ -5,13 +5,13 @@
 #include "CoreMinimal.h"
 
 struct FAnimNextTraitEvent;
+struct FUAFGraphInstanceComponent;
 
 namespace UE::AnimNext
 {
 	struct FExecutionContext;
 	struct FTraitStackBinding;
 	struct FTraitEventList;
-	struct FGraphInstanceComponent;
 
 	// Raises a trait event on the specified trait stack
 	// This calls OnTraitEvent on each trait of the stack following propagation rules
@@ -22,5 +22,5 @@ namespace UE::AnimNext
 	ANIMNEXTANIMGRAPH_API void RaiseTraitEvents(FExecutionContext& Context, const FTraitStackBinding& StackBinding, const FTraitEventList& EventList);
 
 	// Raises every trait event from the input list on the specified component
-	ANIMNEXTANIMGRAPH_API void RaiseTraitEvents(FExecutionContext& Context, FGraphInstanceComponent& Component, const FTraitEventList& EventList);
+	ANIMNEXTANIMGRAPH_API void RaiseTraitEvents(FExecutionContext& Context, FUAFGraphInstanceComponent& Component, const FTraitEventList& EventList);
 }
diff --git a/Engine/Plugins/Experimental/AnimNextAnimGraph/Source/AnimNextAnimGraph/Public/TraitInterfaces/IBlendStack.h b/Engine/Plugins/Experimental/AnimNextAnimGraph/Source/AnimNextAnimGraph/Public/TraitInterfaces/IBlendStack.h
--- a/Engine/Plugins/Experimental/AnimNextAnimGraph/Source/AnimNextAnimGraph/Public/TraitInterfaces/IBlendStack.h
+++ b/Engine/Plugins/Experimental/AnimNextAnimGraph/Source/AnimNextAnimGraph/Public/TraitInterfaces/IBlendStack.h
@@ -6,7 +6,7 @@
 #include "TraitCore/ITraitInterface.h"
 #include "TraitCore/TraitBinding.h"
 #include "AlphaBlend.h"
-#include "AnimNextDataInterfacePayload.h"
+#include "Factory/AnimNextFactoryParams.h"
 #include "Graph/AnimNextAnimationGraph.h"
 #include "Module/ModuleHandle.h"
 #include "Animation/BlendProfile.h"
@@ -49,9 +49,6 @@ namespace UE::AnimNext
 			 * Consider making dynamic payload if we want to implement a special blend framework.
 			 */
 
-			// Payload applied to the AnimationGraph's variables via its data interfaces.
-			FAnimNextDataInterfacePayload GraphPayload;
-
 			// Blend-in duration for this graph request 
 			FAlphaBlendArgs BlendArgs;
 
@@ -62,15 +59,14 @@ namespace UE::AnimNext
 			TObjectPtr<const UAnimNextAnimationGraph> AnimationGraph;
 
 			// The object that was used for factory generation of the graph.
-			// Used as a key to allow identification of this request without inspecting payloads.
 			TObjectPtr<const UObject> FactoryObject;
 
+			// The additional params used for factory generation of the graph.
+			FAnimNextFactoryParams FactoryParams;
+
 			// Child ptr when Type == Child
 			FTraitPtr ChildPtr;
 
-			// Additional module handle to use for data interface bindings
-			FModuleHandle BindingModuleHandle;
-
 			// The Blend mode to use
 			EBlendMode BlendMode = EBlendMode::Standard;
 
@@ -91,6 +87,15 @@ namespace UE::AnimNext
 		// Gets the graph request for the specified child index
 		UE_API virtual IBlendStack::FGraphRequestPtr GetGraph(FExecutionContext& Context, const TTraitBinding<IBlendStack>& Binding, int32 InChildIndex) const;
 
+		// Gets the graph instance for the specified child index.
+		// Instance's lifetime may not persist after this call, so do not store the returned ptr 
+		virtual FAnimNextGraphInstance* GetGraphInstance(FExecutionContext& Context, const TTraitBinding<IBlendStack>& Binding, int32 InChildIndex) const;
+
+		// Gets the graph instance from the most recent PushGraph
+		// Note that if a child graph instance was passed, the instance returned will be the same as that of the hosting trait stack
+		// @return the child index of the active graph instance, or INDEX_NONE if no active graph instance was found
+		virtual int32 GetActiveGraphInstance(FExecutionContext& Context, const TTraitBinding<IBlendStack>& Binding, FAnimNextGraphInstance*& OutGraphInstance) const;
+
 #if WITH_EDITOR
 		UE_API virtual const FText& GetDisplayName() const override;
 		UE_API virtual const FText& GetDisplayShortName() const override;
@@ -122,6 +127,18 @@ namespace UE::AnimNext
 			return GetInterface()->GetGraph(Context, *this, InChildIndex);
 		}
 
+		// @see IBlendStack::GetGraphInstance
+		virtual int32 GetActiveGraphInstance(FExecutionContext& Context, FAnimNextGraphInstance*& OutGraphInstance) const
+		{
+			return GetInterface()->GetActiveGraphInstance(Context, *this, OutGraphInstance);
+		}
+		
+		// @see IBlendStack::GetGraphInstance
+		virtual FAnimNextGraphInstance* GetGraphInstance(FExecutionContext& Context, int32 InChildIndex) const
+		{
+			return GetInterface()->GetGraphInstance(Context, *this, InChildIndex);
+		}
+
 	protected:
 		const IBlendStack* GetInterface() const { return GetInterfaceTyped<IBlendStack>(); }
 	};
diff --git a/Engine/Plugins/Experimental/AnimNextAnimGraph/Source/AnimNextAnimGraph/Public/TraitInterfaces/IGraphFactory.h b/Engine/Plugins/Experimental/AnimNextAnimGraph/Source/AnimNextAnimGraph/Public/TraitInterfaces/IGraphFactory.h
--- a/Engine/Plugins/Experimental/AnimNextAnimGraph/Source/AnimNextAnimGraph/Public/TraitInterfaces/IGraphFactory.h
+++ b/Engine/Plugins/Experimental/AnimNextAnimGraph/Source/AnimNextAnimGraph/Public/TraitInterfaces/IGraphFactory.h
@@ -11,35 +11,28 @@
 
 class UAnimNextAnimationGraph;
 struct FAnimNextDataInterfacePayload;
+struct FAnimNextFactoryParams;
 
 namespace UE::AnimNext
 {
 	/**
 	 * IGraphFactory
 	 *
-	 * This interface translates an object and payload struct into a UAnimNextAnimationGraph and payload instance(s), ready for instantiation
+	 * This interface allows traits to participate in graph building via FAnimGraphFactory
 	 */
 	struct IGraphFactory : ITraitInterface
 	{
 		DECLARE_ANIM_TRAIT_INTERFACE(IGraphFactory)
 
-		// Gets a graph and a set of native interfaces from an object
-		// @param	InObject				The object to manufacture a graph for (e.g. AnimSequence). If this is a UAnimNextAnimationGraph, it should be returned verbatim
-		// @param	InOutGraphPayload			Payload that can be used to communicate with a graph instance
-		// @return	The graph to be used.
-		UE_API virtual const UAnimNextAnimationGraph* GetGraphFromObject(FExecutionContext& Context, const TTraitBinding<IGraphFactory>& Binding, const UObject* InObject, FAnimNextDataInterfacePayload& InOutPayload) const;
+		// Called to modify/extend factory parameters for a generated graph.
+		// Object used to generate the base graph supplied for context.
+		// Fills in InOutParams with appropriately-initialized additive traits.
+		UE_API virtual void GetFactoryParams(FExecutionContext& Context, const TTraitBinding<IGraphFactory>& Binding, const UObject* InObject, FAnimNextFactoryParams& InOutParams) const;
 
-		// Creates data interface payload to bind when creating a graph for an object
-		// @param	InObject					The object (e.g. AnimSequence). If this is a UAnimNextAnimationGraph, it should be the same object as InGraph
-		// @param	InOutGraphPayload			Payload that can be used to communicate with a graph instance
-		UE_API virtual void CreatePayloadForObject(FExecutionContext& Context, const TTraitBinding<IGraphFactory>& Binding, const UObject* InObject, FAnimNextDataInterfacePayload& InOutPayload) const;
-
-		// Gets a graph and a set of native interfaces from an object
-		// Falls back to using user settings if no interface is found in the current stack, or if an interface was found but returns nullptr
-		// @param	InObject				The object to manufacture a graph for (e.g. AnimSequence). If this is a UAnimNextAnimationGraph, it should be returned verbatim
-		// @param	InOutGraphPayload			Payload that can be used to communicate with a graph instance
-		// @return	The graph to be used.
-		static UE_API const UAnimNextAnimationGraph* GetGraphFromObjectWithFallback(FExecutionContext& Context, const FTraitBinding& InBinding, const UObject* InObject, FAnimNextDataInterfacePayload& InOutPayload);
+		// If the supplied object is a UAnimNextAnimationGraph, returns it
+		// If not then uses the parameters (od the default parameters if they are not valid) combined with the current trait stack's collective calls
+		// to GetFactoryParams to build the graph procedurally
+		UE_API static const UAnimNextAnimationGraph* GetOrBuildGraph(FExecutionContext& Context, const FTraitBinding& Binding, const UObject* InObject, FAnimNextFactoryParams& InOutParams);
 
 #if WITH_EDITOR
 		UE_API virtual const FText& GetDisplayName() const override;
@@ -53,16 +46,10 @@ namespace UE::AnimNext
 	template<>
 	struct TTraitBinding<IGraphFactory> : FTraitBinding
 	{
-		// @see IGraphFactory::GetGraphFromObject
-		const UAnimNextAnimationGraph* GetGraphFromObject(FExecutionContext& Context, const UObject* InObject, FAnimNextDataInterfacePayload& InOutPayload) const
-		{
-			return GetInterface()->GetGraphFromObject(Context, *this, InObject, InOutPayload);
-		}
-
-		// @see IGraphFactory::CreatePayloadForObject
-		void CreatePayloadForObject(FExecutionContext& Context, const UObject* InObject, FAnimNextDataInterfacePayload& InOutPayload) const
+		// @see IGraphFactory::GetFactoryParams
+		void GetFactoryParams(FExecutionContext& Context, const UObject* InObject, FAnimNextFactoryParams& InOutParams) const
 		{
-			GetInterface()->CreatePayloadForObject(Context, *this, InObject, InOutPayload);
+			GetInterface()->GetFactoryParams(Context, *this, InObject, InOutParams);
 		}
 
 	protected:
diff --git a/Engine/Plugins/Experimental/AnimNextAnimGraph/Source/AnimNextAnimGraph/Public/Traits/BlendSmoother.h b/Engine/Plugins/Experimental/AnimNextAnimGraph/Source/AnimNextAnimGraph/Public/Traits/BlendSmoother.h
--- a/Engine/Plugins/Experimental/AnimNextAnimGraph/Source/AnimNextAnimGraph/Public/Traits/BlendSmoother.h
+++ b/Engine/Plugins/Experimental/AnimNextAnimGraph/Source/AnimNextAnimGraph/Public/Traits/BlendSmoother.h
@@ -23,6 +23,12 @@ struct FAnimNextBlendSmootherCoreTraitSharedData : public FAnimNextTraitSharedDa
 	// This struct is empty but required so that we can give a nice display name to the trait
 };
 
+namespace UE::AnimNext
+{
+	// Namespaced alias
+	using FBlendSmootherCoreData = FAnimNextBlendSmootherCoreTraitSharedData; 
+}
+
 USTRUCT(meta = (DisplayName = "Blend Smoother"))
 struct FAnimNextBlendSmootherTraitSharedData : public FAnimNextBlendSmootherCoreTraitSharedData
 {
diff --git a/Engine/Plugins/Experimental/AnimNextAnimGraph/Source/AnimNextAnimGraphEditor/Private/AnimNextAnimGraphEditorModule.cpp b/Engine/Plugins/Experimental/AnimNextAnimGraph/Source/AnimNextAnimGraphEditor/Private/AnimNextAnimGraphEditorModule.cpp
--- a/Engine/Plugins/Experimental/AnimNextAnimGraph/Source/AnimNextAnimGraphEditor/Private/AnimNextAnimGraphEditorModule.cpp
+++ b/Engine/Plugins/Experimental/AnimNextAnimGraph/Source/AnimNextAnimGraphEditor/Private/AnimNextAnimGraphEditorModule.cpp
@@ -40,7 +40,9 @@
 #include "RewindDebugger/EvaluationProgramTrack.h"
 #include "RewindDebugger/SequenceInfoTrack.h"
 #include "PersonaModule.h"
-#include "Developer/TraceServices/Public/TraceServices/ModuleService.h"
+#include "Factory/AnimNextFactoryParams.h"
+#include "TraceServices/ModuleService.h"
+#include "Factory/AnimNextFactoryParamsDetails.h"
 
 #define LOCTEXT_NAMESPACE "FAnimNextAnimGraphEditorModule"
 
@@ -75,8 +77,12 @@ class FAnimNextAnimGraphEditorModule : public IModuleInterface
 					
 		PropertyModule.RegisterCustomPropertyTypeLayout(
 							FAnimNextSequenceTraceInfo::StaticStruct()->GetFName(),
-							FOnGetPropertyTypeCustomizationInstance::CreateLambda([] { return MakeShared<AnimNextEditor::FAnimNextSequenceTraceInfoCustomization>(); }));					
-					
+							FOnGetPropertyTypeCustomizationInstance::CreateLambda([] { return MakeShared<AnimNextEditor::FAnimNextSequenceTraceInfoCustomization>(); }));
+
+		PropertyModule.RegisterCustomPropertyTypeLayout(
+			FAnimNextFactoryParams::StaticStruct()->GetFName(),
+			FOnGetPropertyTypeCustomizationInstance::CreateLambda([] { return MakeShared<FAnimNextFactoryParamsDetails>(); }));
+
 		IModularFeatures::Get().RegisterModularFeature(TraceServices::ModuleFeatureName, &GAnimNextAnimGraphTraceModule);
 		IModularFeatures::Get().RegisterModularFeature(RewindDebugger::IRewindDebuggerTrackCreator::ModularFeatureName, &GAnimNextModulesTrackCreator);
 		IModularFeatures::Get().RegisterModularFeature(RewindDebugger::IRewindDebuggerTrackCreator::ModularFeatureName, &GSequenceInfoTrackCreator);
@@ -160,8 +166,10 @@ class FAnimNextAnimGraphEditorModule : public IModuleInterface
 		if(FModuleManager::Get().IsModuleLoaded("PropertyEditor"))
 		{
 			FPropertyEditorModule& PropertyModule = FModuleManager::LoadModuleChecked<FPropertyEditorModule>("PropertyEditor");
-			PropertyModule.UnregisterCustomClassLayout("AnimNextCallFunctionSharedData");
-			PropertyModule.UnregisterCustomClassLayout("AnimNextAnimGraph");
+			PropertyModule.UnregisterCustomPropertyTypeLayout("AnimNextCallFunctionSharedData");
+			PropertyModule.UnregisterCustomPropertyTypeLayout("AnimNextAnimGraph");
+			PropertyModule.UnregisterCustomPropertyTypeLayout("AnimNextFactoryParamsDetails");
+			PropertyModule.UnregisterCustomPropertyTypeLayout("AnimNextSequenceTraceInfo");
 		}
 
 		IModularFeatures::Get().UnregisterModularFeature(ITraitStackEditor::ModularFeatureName, TraitStackEditor.Get());
diff --git a/Engine/Plugins/Experimental/AnimNextAnimGraph/Source/AnimNextAnimGraphEditor/Private/Factory/AnimNextFactoryParamsDetails.cpp b/Engine/Plugins/Experimental/AnimNextAnimGraph/Source/AnimNextAnimGraphEditor/Private/Factory/AnimNextFactoryParamsDetails.cpp
new file mode 100644
--- /dev/null
+++ b/Engine/Plugins/Experimental/AnimNextAnimGraph/Source/AnimNextAnimGraphEditor/Private/Factory/AnimNextFactoryParamsDetails.cpp
@@ -0,0 +1,514 @@
+﻿// Copyright Epic Games, Inc. All Rights Reserved.
+
+#include "AnimNextFactoryParamsDetails.h"
+
+#include "AnimNextStyle.h"
+#include "DetailLayoutBuilder.h"
+#include "DetailWidgetRow.h"
+#include "IDetailChildrenBuilder.h"
+#include "IDetailPropertyRow.h"
+#include "IStructureDataProvider.h"
+#include "PropertyHandle.h"
+#include "StructViewerFilter.h"
+#include "StructViewerModule.h"
+#include "Algo/Compare.h"
+#include "Factory/AnimGraphFactory.h"
+#include "Factory/AnimNextFactoryParams.h"
+#include "Modules/ModuleManager.h"
+#include "Styling/StyleColors.h"
+#include "TraitCore/TraitRegistry.h"
+#include "Widgets/Images/SImage.h"
+#include "Widgets/Input/SCheckBox.h"
+#include "Widgets/Input/SComboButton.h"
+#include "Widgets/Layout/SBox.h"
+#include "Widgets/Text/STextBlock.h"
+#include "Graph/AnimNextAnimationGraph.h"
+
+#define LOCTEXT_NAMESPACE "AnimNextFactoryParamsDetails"
+
+namespace UE::AnimNext::Editor
+{
+
+class FInstancedStructProvider : public IStructureDataProvider
+{
+public:
+	FInstancedStructProvider() = default;
+	
+	explicit FInstancedStructProvider(const TSharedPtr<IPropertyHandle>& InStructProperty)
+		: StructProperty(InStructProperty)
+	{
+	}
+	
+	virtual ~FInstancedStructProvider() override {}
+
+	void Reset()
+	{
+		StructProperty = nullptr;
+	}
+	
+	virtual bool IsValid() const override
+	{
+		bool bHasValidData = false;
+		EnumerateInstances([&bHasValidData](const UScriptStruct* ScriptStruct, uint8* Memory, UPackage* Package)
+		{
+			if (ScriptStruct && Memory)
+			{
+				bHasValidData = true;
+				return false; // Stop
+			}
+			return true; // Continue
+		});
+
+		return bHasValidData;
+	}
+	
+	virtual const UStruct* GetBaseStructure() const override
+	{
+		// Taken from UClass::FindCommonBase
+		auto FindCommonBaseStruct = [](const UScriptStruct* StructA, const UScriptStruct* StructB)
+		{
+			const UScriptStruct* CommonBaseStruct = StructA;
+			while (CommonBaseStruct && StructB && !StructB->IsChildOf(CommonBaseStruct))
+			{
+				CommonBaseStruct = Cast<UScriptStruct>(CommonBaseStruct->GetSuperStruct());
+			}
+			return CommonBaseStruct;
+		};
+
+		const UScriptStruct* CommonStruct = nullptr;
+		EnumerateInstances([&CommonStruct, &FindCommonBaseStruct](const UScriptStruct* ScriptStruct, uint8* Memory, UPackage* Package)
+		{
+			if (ScriptStruct)
+			{
+				CommonStruct = FindCommonBaseStruct(ScriptStruct, CommonStruct);
+			}
+			return true; // Continue
+		});
+
+		return CommonStruct;
+	}
+
+	virtual void GetInstances(TArray<TSharedPtr<FStructOnScope>>& OutInstances, const UStruct* ExpectedBaseStructure) const override
+	{
+		// The returned instances need to be compatible with base structure.
+		// This function returns empty instances in case they are not compatible, with the idea that we have as many instances as we have outer objects.
+		EnumerateInstances([&OutInstances, ExpectedBaseStructure](const UScriptStruct* ScriptStruct, uint8* Memory, UPackage* Package)
+		{
+			TSharedPtr<FStructOnScope> Result;
+			
+			if (ExpectedBaseStructure && ScriptStruct && ScriptStruct->IsChildOf(ExpectedBaseStructure))
+			{
+				Result = MakeShared<FStructOnScope>(ScriptStruct, Memory);
+				Result->SetPackage(Package);
+			}
+
+			OutInstances.Add(Result);
+
+			return true; // Continue
+		});
+	}
+
+	virtual bool IsPropertyIndirection() const override
+	{
+		return true;
+	}
+
+	virtual uint8* GetValueBaseAddress(uint8* ParentValueAddress, const UStruct* ExpectedBaseStructure) const override
+	{
+		if (!ParentValueAddress)
+		{
+			return nullptr;
+		}
+
+		FInstancedStruct& InstancedStruct = *reinterpret_cast<FInstancedStruct*>(ParentValueAddress);
+		if (ExpectedBaseStructure && InstancedStruct.GetScriptStruct() && InstancedStruct.GetScriptStruct()->IsChildOf(ExpectedBaseStructure))
+		{
+			return InstancedStruct.GetMutableMemory();
+		}
+		
+		return nullptr;
+	}
+	
+protected:
+
+	void EnumerateInstances(TFunctionRef<bool(const UScriptStruct* ScriptStruct, uint8* Memory, UPackage* Package)> InFunc) const
+	{
+		if (!StructProperty.IsValid() || !StructProperty->IsValidHandle())
+		{
+			return;
+		}
+		
+		TArray<UPackage*> Packages;
+		StructProperty->GetOuterPackages(Packages);
+
+		StructProperty->EnumerateRawData([&InFunc, &Packages](void* RawData, const int32 DataIndex, const int32 /*NumDatas*/)
+		{
+			const UScriptStruct* ScriptStruct = nullptr;
+			uint8* Memory = nullptr;
+			UPackage* Package = nullptr;
+			if (FInstancedStruct* InstancedStruct = static_cast<FInstancedStruct*>(RawData))
+			{
+				ScriptStruct = InstancedStruct->GetScriptStruct();
+				Memory = InstancedStruct->GetMutableMemory();
+				if (ensureMsgf(Packages.IsValidIndex(DataIndex), TEXT("Expecting packges and raw data to match.")))
+				{
+					Package = Packages[DataIndex];
+				}
+			}
+
+			return InFunc(ScriptStruct, Memory, Package);
+		});
+	}
+	
+	TSharedPtr<IPropertyHandle> StructProperty;
+};
+
+void FAnimNextFactoryParamsDetails::CustomizeHeader(TSharedRef<IPropertyHandle> InPropertyHandle, FDetailWidgetRow& InHeaderRow, IPropertyTypeCustomizationUtils& InCustomizationUtils)
+{
+}
+
+void FAnimNextFactoryParamsDetails::CustomizeChildren(TSharedRef<IPropertyHandle> InPropertyHandle, IDetailChildrenBuilder& InChildBuilder, IPropertyTypeCustomizationUtils& InCustomizationUtils)
+{
+	TSharedPtr<IPropertyHandle> BuilderHandle = InPropertyHandle->GetChildHandle(GET_MEMBER_NAME_CHECKED(FAnimNextFactoryParams, Builder));
+	BuilderHandle->MarkHiddenByCustomization();
+
+	TSharedPtr<IPropertyHandle> StacksHandle = BuilderHandle->GetChildHandle(GET_MEMBER_NAME_CHECKED(FAnimNextSimpleAnimGraphBuilder, Stacks));
+	TSharedPtr<IPropertyHandle> StackHandle = StacksHandle->GetChildHandle(0);
+	if (StackHandle.IsValid())
+	{
+		TraitStructsHandle = StackHandle->GetChildHandle(GET_MEMBER_NAME_CHECKED(FAnimNextSimpleAnimGraphBuilderTraitStackDesc, TraitStructs));
+	}
+	
+	bool bIsInvalidOrAnimGraph = false; 
+
+	// Find property for the factory source & common default params set
+	TOptional<FAnimNextFactoryParams> CommonDefaults;
+	FString FactorySourcePropertyName = InPropertyHandle->GetMetaData("FactorySourceProperty");
+	TSharedPtr<IPropertyHandle> FactorySourceHandle = InPropertyHandle->GetParentHandle()->GetChildHandle(FName(*FactorySourcePropertyName));
+	FactorySourceHandle->SetOnPropertyValueChanged(FSimpleDelegate::CreateLambda([WeakPropertyHandle = TWeakPtr<IPropertyHandle>(InPropertyHandle)]()
+	{
+		if (TSharedPtr<IPropertyHandle> PinnedPropertyHandle = WeakPropertyHandle.Pin())
+		{
+			PinnedPropertyHandle->RequestRebuildChildren();
+		}
+	}));
+	
+	if (FactorySourceHandle.IsValid() && FactorySourceHandle->GetProperty()->IsA<FObjectProperty>())
+	{
+		FactorySourceHandle->EnumerateRawData([&CommonDefaults, &bIsInvalidOrAnimGraph](void* RawData, const int32 DataIndex, const int32 NumDatas)
+		{
+			UObject* Object = *(static_cast<UObject**>(RawData));
+		
+			if (!CommonDefaults.IsSet())
+			{
+				if (Object == nullptr || Object->IsA<UAnimNextAnimationGraph>())
+				{
+					// TODO: At the moment we dont have 'details panel pin overrides' for assets, so we reset the params to not be confusing
+					// In the end we will want the params to also provide overrides like we supply to RunGraph nodes
+					CommonDefaults = FAnimNextFactoryParams();
+					bIsInvalidOrAnimGraph = true;
+				}
+				else
+				{
+					CommonDefaults = FAnimGraphFactory::GetDefaultParamsForObject(Object);
+				}
+			}
+			else
+			{
+				const FAnimNextFactoryParams& DefaultParams = FAnimGraphFactory::GetDefaultParamsForObject(Object);
+				if (!Algo::Compare(DefaultParams.Builder.Stacks[0].TraitStructs, CommonDefaults->Builder.Stacks[0].TraitStructs))
+				{
+					CommonDefaults.Reset();
+					return false;
+				}
+			}
+			return true;
+		});
+	}
+
+	TSet<const UScriptStruct*> AddedStructs;
+
+	if (TraitStructsHandle.IsValid())
+	{
+		// Find common 'current' set
+		TSet<const UScriptStruct*> CommonStructs;
+		TraitStructsHandle->EnumerateRawData([&CommonStructs](void* RawData, const int32 DataIndex, const int32 NumDatas)
+		{
+			TArray<FInstancedStruct>* InitializePayloads = static_cast<TArray<FInstancedStruct>*>(RawData);
+			if (CommonStructs.IsEmpty())
+			{
+				for (const FInstancedStruct& InitializePayload : *InitializePayloads)
+				{
+					CommonStructs.Add(InitializePayload.GetScriptStruct());
+				}
+			}
+			else
+			{
+				for (auto It = CommonStructs.CreateIterator(); It; ++It)
+				{
+					const UScriptStruct* CommonStruct = *It;
+					if (!InitializePayloads->ContainsByPredicate([CommonStruct](const FInstancedStruct& InInitializePayload){ return CommonStruct == InInitializePayload.GetScriptStruct(); }))
+					{
+						It.RemoveCurrent();
+					}
+				}
+			}
+			return true;
+		});
+
+		// No common structs, cannot edit (all struct types for edited objects are mutually exclusive)
+		if (CommonStructs.IsEmpty())
+		{
+			return;
+		}
+
+		// Add entries for instanced structs that are common to all objects and have a common type at this index
+		uint32 NumChildren = 0;
+		if (TraitStructsHandle->GetNumChildren(NumChildren) == FPropertyAccess::Success)
+		{
+			for (uint32 ChildIndex = 0; ChildIndex < NumChildren; ++ChildIndex)
+			{
+				TSharedRef<IPropertyHandle> InstancedStructHandle = TraitStructsHandle->GetChildHandle(ChildIndex).ToSharedRef();
+				TOptional<const UScriptStruct*> CommonType;
+				InstancedStructHandle->EnumerateRawData([&CommonType](void* RawData, const int32 DataIndex, const int32 NumDatas)
+				{
+					FInstancedStruct* InstancedStruct = static_cast<FInstancedStruct*>(RawData);
+					if (!CommonType.IsSet())
+					{
+						CommonType = InstancedStruct->GetScriptStruct();
+					}
+					else if (CommonType.GetValue() != InstancedStruct->GetScriptStruct())
+					{
+						CommonType.Reset();
+						return false;
+					}
+					return true;
+				});
+
+				if (CommonType.IsSet() && CommonStructs.Contains(CommonType.GetValue()))
+				{
+					AddedStructs.Add(CommonType.GetValue());
+
+					const FName UniqueName(InstancedStructHandle->GetPropertyPath());
+					IDetailPropertyRow* Row = InChildBuilder.AddChildStructure(InstancedStructHandle, MakeShared<FInstancedStructProvider>(InstancedStructHandle), UniqueName);
+					FDetailWidgetRow& CustomWidget = Row->CustomWidget(true);
+
+					CustomWidget
+					.NameContent()
+					[
+						SNew(SBorder)
+						.Padding(FMargin(6.0f, 2.0f))
+						.BorderImage(FAnimNextStyle::Get().GetBrush("AnimNext.TraitDetails.Background"))
+						[
+							SNew(STextBlock)
+							.TextStyle(FAnimNextStyle::Get(), "AnimNext.TraitDetails.Label")
+							.ColorAndOpacity(FStyleColors::Foreground)
+							.Text(CommonType.GetValue()->GetDisplayNameText())
+						]
+					];
+
+					const FTraitRegistry& TraitRegistry = FTraitRegistry::Get();
+					const FTrait* Trait = TraitRegistry.Find(CommonType.GetValue());
+					if (Trait)
+					{
+						// If this is an additive trait, allow it to be removed
+						if (Trait->GetTraitMode() == ETraitMode::Additive)
+						{
+							CustomWidget
+							.ValueContent()
+							[
+								SNew(SCheckBox)
+								.ToolTipText(LOCTEXT("EnableOptionalTrait", "Enable/disable this optional trait"))
+								.IsChecked(true)
+								.OnCheckStateChanged_Lambda([ChildIndex, WeakTraitStructsHandle = TWeakPtr<IPropertyHandle>(TraitStructsHandle), WeakPropertyHandle = TWeakPtr<IPropertyHandle>(InPropertyHandle)](ECheckBoxState InState)
+								{
+									TSharedPtr<IPropertyHandle> PinnedTraitStructsHandle = WeakTraitStructsHandle.Pin();
+									TSharedPtr<IPropertyHandle> PinnedPropertyHandle = WeakPropertyHandle.Pin();
+									if (PinnedTraitStructsHandle.IsValid() && PinnedPropertyHandle.IsValid())
+									{
+										PinnedTraitStructsHandle->AsArray()->DeleteItem(ChildIndex);
+										PinnedPropertyHandle->RequestRebuildChildren();
+									}
+								})
+							];
+						}
+						else if (FactorySourceHandle.IsValid())
+						{
+							// If its a base trait, hide any child property that has the appropriate metadata (if we have a factory source)
+							uint32 NumStructProperties = 0;
+							TSharedPtr<IPropertyHandle> StructHandle = InstancedStructHandle->GetChildHandle(0);
+							if (StructHandle->GetNumChildren(NumStructProperties) == FPropertyAccess::Success)
+							{
+								for (uint32 StructPropertyIndex = 0; StructPropertyIndex < NumStructProperties; ++StructPropertyIndex)
+								{
+									TSharedPtr<IPropertyHandle> StructPropertyHandle = StructHandle->GetChildHandle(StructPropertyIndex);
+									if (StructPropertyHandle->HasMetaData("FactorySource"))
+									{
+										StructPropertyHandle->MarkHiddenByCustomization();
+									}
+								}
+							}
+						}
+					}
+				}
+			}
+		}
+	}
+
+	if (CommonDefaults.IsSet() && CommonDefaults.GetValue().Builder.GetNumStacks() > 0)
+	{
+		// Now add any entries for defaults we didnt find above to allow them to be added
+		for (const TInstancedStruct<FAnimNextTraitSharedData>& DefaultStruct : CommonDefaults.GetValue().Builder.Stacks[0].TraitStructs)
+		{
+			if (!AddedStructs.Contains(DefaultStruct.GetScriptStruct()))
+			{
+				InChildBuilder.AddCustomRow(DefaultStruct.GetScriptStruct()->GetDisplayNameText())
+				.NameContent()
+				[
+					SNew(SBorder)
+					.Padding(FMargin(6.0f, 2.0f))
+					.BorderImage(FAnimNextStyle::Get().GetBrush("AnimNext.TraitDetails.Background"))
+					[
+						SNew(STextBlock)
+						.TextStyle(FAnimNextStyle::Get(), "AnimNext.TraitDetails.Label")
+						.ColorAndOpacity(FStyleColors::Foreground)
+						.Text(DefaultStruct.GetScriptStruct()->GetDisplayNameText())
+					]
+				]
+				.ValueContent()
+				[
+					SNew(SCheckBox)
+					.ToolTipText(LOCTEXT("EnableOptionalTrait", "Enable/disable this optional trait"))
+					.IsChecked(false)
+					.OnCheckStateChanged_Lambda([DefaultStruct, WeakTraitStructsHandle = TWeakPtr<IPropertyHandle>(TraitStructsHandle), WeakPropertyHandle = TWeakPtr<IPropertyHandle>(InPropertyHandle)](ECheckBoxState InState)
+					{
+						TSharedPtr<IPropertyHandle> PinnedTraitStructsHandle = WeakTraitStructsHandle.Pin();
+						TSharedPtr<IPropertyHandle> PinnedPropertyHandle = WeakPropertyHandle.Pin();
+						if (PinnedTraitStructsHandle.IsValid() && PinnedPropertyHandle.IsValid())
+						{
+							TSharedPtr<IPropertyHandleArray> ArrayHandle = PinnedTraitStructsHandle->AsArray();
+							FPropertyHandleItemAddResult Result = ArrayHandle->AddItem();
+							TSharedPtr<IPropertyHandle> NewEntryHandle = PinnedTraitStructsHandle->GetChildHandle(Result.GetIndex());
+							NewEntryHandle->EnumerateRawData([&DefaultStruct](void* RawData, const int32 DataIndex, const int32 NumDatas)
+							{
+								TInstancedStruct<FAnimNextTraitSharedData>* InstancedStruct = static_cast<TInstancedStruct<FAnimNextTraitSharedData>*>(RawData);
+								*InstancedStruct = DefaultStruct;
+								return true;
+							});
+							PinnedPropertyHandle->RequestRebuildChildren();
+						}
+					})
+				];
+			}
+		}
+	}
+
+	if (!bIsInvalidOrAnimGraph)
+	{
+		// Add an 'add' widget for any custom behaviors users want to opt into
+		InChildBuilder.AddCustomRow(LOCTEXT("AddTraitWidget", "Add Trait"))
+		.NameContent()
+		[
+			SNew(SComboButton)
+			.ToolTipText(LOCTEXT("AddTraitTooltip", "Add a trait to modify how this asset is played"))
+			.HasDownArrow(false)
+			.ComboButtonStyle(FAnimNextStyle::Get(), "AnimNext.TraitDetails.AddButton")
+			.ButtonContent()
+			[
+				SNew(SHorizontalBox)
+				+SHorizontalBox::Slot()
+				.VAlign(VAlign_Center)
+				.AutoWidth()
+				.Padding(0.0f, 0.0f, 3.0f, 0.0f)
+				[
+					SNew(SImage)
+					.Image(FAppStyle::Get().GetBrush("Plus"))
+				]
+				+SHorizontalBox::Slot()
+				.VAlign(VAlign_Center)
+				.AutoWidth()
+				.Padding(0.0f, 2.0f)
+				[
+					SNew(STextBlock)
+					.TextStyle(FAnimNextStyle::Get(), "AnimNext.TraitDetails.Label")
+					.ColorAndOpacity(FStyleColors::Foreground)
+					.Text(LOCTEXT("AddTraitButton", "Add"))
+				]
+			]
+			.OnGetMenuContent_Lambda([WeakTraitStructsHandle = TWeakPtr<IPropertyHandle>(TraitStructsHandle), WeakPropertyHandle = TWeakPtr<IPropertyHandle>(InPropertyHandle)]()
+			{
+				FStructViewerModule& StructViewerModule = FModuleManager::LoadModuleChecked<FStructViewerModule>("StructViewer");
+
+				class FStructFilter : public IStructViewerFilter
+				{
+				public:
+					virtual bool IsStructAllowed(const FStructViewerInitializationOptions& InInitOptions, const UScriptStruct* InStruct, TSharedRef<FStructViewerFilterFuncs> InFilterFuncs) override
+					{
+						if (InStruct->HasMetaData(TEXT("Hidden")))
+						{
+							return false;
+						}
+
+						if (!InStruct->IsChildOf(FAnimNextTraitSharedData::StaticStruct()))
+						{
+							return false;
+						}
+						
+						const FTraitRegistry& TraitRegistry = FTraitRegistry::Get();
+						const FTrait* Trait = TraitRegistry.Find(InStruct);
+						return Trait && Trait->GetTraitMode() == ETraitMode::Additive;
+					}
+
+					virtual bool IsUnloadedStructAllowed(const FStructViewerInitializationOptions& InInitOptions, const FSoftObjectPath& InStructPath, TSharedRef<class FStructViewerFilterFuncs> InFilterFuncs)
+					{
+						return false;
+					};
+				};
+
+				FStructViewerInitializationOptions Options;
+				Options.StructFilter = MakeShared<FStructFilter>();
+				Options.Mode = EStructViewerMode::StructPicker;
+				Options.NameTypeToDisplay = EStructViewerNameTypeToDisplay::DisplayName;
+				Options.bShowNoneOption = false;
+
+				return
+					SNew(SBox)
+					.WidthOverride(400.0f)
+					.HeightOverride(400.0f)
+					[
+						StructViewerModule.CreateStructViewer(Options, FOnStructPicked::CreateLambda([WeakPropertyHandle](const UScriptStruct* InStruct)
+						{
+							TSharedPtr<IPropertyHandle> PinnedPropertyHandle = WeakPropertyHandle.Pin();
+							if (PinnedPropertyHandle.IsValid())
+							{
+								TSharedPtr<IPropertyHandle> BuilderHandle = PinnedPropertyHandle->GetChildHandle(GET_MEMBER_NAME_CHECKED(FAnimNextFactoryParams, Builder));
+								TSharedPtr<IPropertyHandle> StacksHandle = BuilderHandle->GetChildHandle(GET_MEMBER_NAME_CHECKED(FAnimNextSimpleAnimGraphBuilder, Stacks));
+								TSharedPtr<IPropertyHandle> StackHandle = StacksHandle->GetChildHandle(0);
+								if (!StackHandle.IsValid())
+								{
+									StacksHandle->AsArray()->AddItem();
+									StackHandle = StacksHandle->GetChildHandle(0);
+								}
+
+								TSharedPtr<IPropertyHandle> TraitStructsHandle = StackHandle->GetChildHandle(GET_MEMBER_NAME_CHECKED(FAnimNextSimpleAnimGraphBuilderTraitStackDesc, TraitStructs));
+								TSharedPtr<IPropertyHandleArray> ArrayHandle = TraitStructsHandle->AsArray();
+								FPropertyHandleItemAddResult Result = ArrayHandle->AddItem();
+								TSharedPtr<IPropertyHandle> NewEntryHandle = TraitStructsHandle->GetChildHandle(Result.GetIndex());
+								NewEntryHandle->EnumerateRawData([InStruct](void* RawData, const int32 DataIndex, const int32 NumDatas)
+								{
+									FInstancedStruct* InstancedStruct = static_cast<FInstancedStruct*>(RawData);
+									InstancedStruct->InitializeAs(InStruct);
+									return true;
+								});
+								PinnedPropertyHandle->RequestRebuildChildren();
+							}
+						}))
+					];
+			})
+		];
+	}
+}
+
+}
+
+#undef LOCTEXT_NAMESPACE
\ No newline at end of file
diff --git a/Engine/Plugins/Experimental/AnimNextAnimGraph/Source/AnimNextAnimGraphEditor/Private/Graph/AnimNextAnimationGraphItemDetails.cpp b/Engine/Plugins/Experimental/AnimNextAnimGraph/Source/AnimNextAnimGraphEditor/Private/Graph/AnimNextAnimationGraphItemDetails.cpp
--- a/Engine/Plugins/Experimental/AnimNextAnimGraph/Source/AnimNextAnimGraphEditor/Private/Graph/AnimNextAnimationGraphItemDetails.cpp
+++ b/Engine/Plugins/Experimental/AnimNextAnimGraph/Source/AnimNextAnimGraphEditor/Private/Graph/AnimNextAnimationGraphItemDetails.cpp
@@ -12,7 +12,7 @@
 #include "ScopedTransaction.h"
 #include "RigVMModel/RigVMClient.h"
 #include "ToolMenus.h"
-#include "Entries/AnimNextDataInterfaceEntry.h"
+#include "Entries/AnimNextSharedVariablesEntry.h"
 #include "Graph/AnimNextAnimationGraph_EditorData.h"
 
 #define LOCTEXT_NAMESPACE "FAnimNextAnimationGraphItemDetails"
diff --git a/Engine/Plugins/Experimental/AnimNextAnimGraph/Source/AnimNextAnimGraphEditor/Private/Graph/AnimNextAnimationGraphMenuExtensions.cpp b/Engine/Plugins/Experimental/AnimNextAnimGraph/Source/AnimNextAnimGraphEditor/Private/Graph/AnimNextAnimationGraphMenuExtensions.cpp
--- a/Engine/Plugins/Experimental/AnimNextAnimGraph/Source/AnimNextAnimGraphEditor/Private/Graph/AnimNextAnimationGraphMenuExtensions.cpp
+++ b/Engine/Plugins/Experimental/AnimNextAnimGraph/Source/AnimNextAnimGraphEditor/Private/Graph/AnimNextAnimationGraphMenuExtensions.cpp
@@ -22,6 +22,7 @@
 #include "AnimNextTraitStackUnitNode.h"
 #include "AnimNextRigVMAssetEditorData.h"
 #include "RigVMModel/RigVMClient.h"
+#include "Variables/AnimNextSharedVariables.h"
 
 #define LOCTEXT_NAMESPACE "FAnimationGraphMenuExtensions"
 
@@ -166,7 +167,7 @@ void FAnimationGraphMenuExtensions::RegisterMenus()
 					FContentBrowserModule& ContentBrowserModule = FModuleManager::Get().LoadModuleChecked<FContentBrowserModule>(TEXT("ContentBrowser"));
 
 					FAssetPickerConfig AssetPickerConfig;
-					AssetPickerConfig.Filter.ClassPaths.Add(UAnimNextDataInterface::StaticClass()->GetClassPathName());
+					AssetPickerConfig.Filter.ClassPaths.Add(UAnimNextSharedVariables::StaticClass()->GetClassPathName());
 					AssetPickerConfig.Filter.bRecursiveClasses = true;
 					AssetPickerConfig.InitialAssetViewType = EAssetViewType::List;
 					AssetPickerConfig.AssetShowWarningText = LOCTEXT("NoAssetsWithPublicVariablesMessage", "No animation graphs with public variables found");
@@ -177,14 +178,17 @@ void FAnimationGraphMenuExtensions::RegisterMenus()
 						FString DefaultValue;
 						FRigVMTrait_AnimNextPublicVariables DefaultTrait;
 						FRigVMTrait_AnimNextPublicVariables NewTrait;
-						UAnimNextDataInterface* Asset = CastChecked<UAnimNextDataInterface>(InAssetData.GetAsset());
+						UAnimNextSharedVariables* Asset = CastChecked<UAnimNextSharedVariables>(InAssetData.GetAsset());
 						NewTrait.InternalAsset = Asset;
-						const FInstancedPropertyBag& PublicVariableDefaults = Asset->GetPublicVariableDefaults();
-						TConstArrayView<FPropertyBagPropertyDesc> Descs = PublicVariableDefaults.GetPropertyBagStruct()->GetPropertyDescs();
+						const FInstancedPropertyBag& VariableDefaults = Asset->GetVariableDefaults();
+						TConstArrayView<FPropertyBagPropertyDesc> Descs = VariableDefaults.GetPropertyBagStruct()->GetPropertyDescs();
 						NewTrait.InternalVariableNames.Reserve(Descs.Num());
 						for(const FPropertyBagPropertyDesc& Desc : Descs)
 						{
-							NewTrait.InternalVariableNames.Add(Desc.Name);
+							if (Desc.CachedProperty->HasAnyPropertyFlags(CPF_NativeAccessSpecifierPublic))
+							{
+								NewTrait.InternalVariableNames.Add(Desc.Name);
+							}
 						}
 						FRigVMTrait_AnimNextPublicVariables::StaticStruct()->ExportText(DefaultValue, &NewTrait, &DefaultTrait, nullptr, PPF_SerializedAsImportText, nullptr);
 
diff --git a/Engine/Plugins/Experimental/AnimNextAnimGraph/Source/AnimNextAnimGraphEditor/Private/RewindDebugger/SequenceInfoTrack.h b/Engine/Plugins/Experimental/AnimNextAnimGraph/Source/AnimNextAnimGraphEditor/Private/RewindDebugger/SequenceInfoTrack.h
--- a/Engine/Plugins/Experimental/AnimNextAnimGraph/Source/AnimNextAnimGraphEditor/Private/RewindDebugger/SequenceInfoTrack.h
+++ b/Engine/Plugins/Experimental/AnimNextAnimGraph/Source/AnimNextAnimGraphEditor/Private/RewindDebugger/SequenceInfoTrack.h
@@ -36,7 +36,7 @@ struct FAnimNextSequenceTraceInfo
 	}
 
 	UPROPERTY(VisibleAnywhere, Category = Properties)
-	TWeakObjectPtr<UAnimSequence> AnimSequence;
+	TWeakObjectPtr<const UAnimSequence> AnimSequence;
 
 	UPROPERTY(VisibleAnywhere, Category = Properties)
 	float DurationSeconds = 0.0f;
