AnimNext: Pass-by-ref runtime refactor

diff --git a/Engine/Plugins/Experimental/AnimNextAnimGraph/Source/AnimNextAnimGraphTestSuite/Private/AnimNextAnimGraphTraitGraphTest.cpp b/Engine/Plugins/Experimental/AnimNextAnimGraph/Source/AnimNextAnimGraphTestSuite/Private/AnimNextAnimGraphTraitGraphTest.cpp
--- a/Engine/Plugins/Experimental/AnimNextAnimGraph/Source/AnimNextAnimGraphTestSuite/Private/AnimNextAnimGraphTraitGraphTest.cpp
+++ b/Engine/Plugins/Experimental/AnimNextAnimGraph/Source/AnimNextAnimGraphTestSuite/Private/AnimNextAnimGraphTraitGraphTest.cpp
@@ -25,6 +25,7 @@
 #include "RigVMFunctions/Math/RigVMFunction_MathFloat.h"
 #include "AnimNextTraitStackUnitNode.h"
 #include "AnimGraphUncookedOnlyUtils.h"
+#include "UAFRigVMComponent.h"
 
 #if WITH_DEV_AUTOMATION_TESTS
 
@@ -54,7 +55,7 @@ namespace UE::AnimNext
 			const FSharedData* SharedData = Binding.GetSharedData<FSharedData>();
 			FInstanceData* InstanceData = Binding.GetInstanceData<FInstanceData>();
 
-			FRigVMExecuteContext& ExecuteContext = Context.GetRootGraphInstance().GetExtendedExecuteContext().GetPublicData();
+			FRigVMExecuteContext& ExecuteContext = Context.GetRootGraphInstance().GetComponent<FUAFRigVMComponent>().GetExtendedExecuteContext().GetPublicData();
 			ExecuteContext.Logf(EMessageSeverity::Info, TEXT("UpdateCount == %d"), ++InstanceData->UpdateCount);
 			ExecuteContext.Logf(EMessageSeverity::Info, TEXT("SomeInt32 == %d"), SharedData->SomeInt32);
 			ExecuteContext.Logf(EMessageSeverity::Info, TEXT("SomeFloat == %.02f"), SharedData->SomeFloat);
@@ -69,7 +70,7 @@ namespace UE::AnimNext
 			IEvaluate::PostEvaluate(Context, Binding);
 
 			FInstanceData* InstanceData = Binding.GetInstanceData<FInstanceData>();
-			FRigVMExecuteContext& ExecuteContext = Context.GetRootGraphInstance().GetExtendedExecuteContext().GetPublicData();
+			FRigVMExecuteContext& ExecuteContext = Context.GetRootGraphInstance().GetComponent<FUAFRigVMComponent>().GetExtendedExecuteContext().GetPublicData();
 			ExecuteContext.Logf(EMessageSeverity::Info, TEXT("EvaluateCount == %d"), ++InstanceData->EvaluateCount);
 		}
 	};
@@ -538,7 +539,7 @@ bool FAnimationAnimNextRuntimeTest_GraphExecute::RunTest(const FString& InParame
 		{
 			Messages.Add(Message);
 		});
-		GraphInstance->GetExtendedExecuteContext().SetRuntimeSettings(RuntimeSettings);
+		GraphInstance->GetComponent<FUAFRigVMComponent>().GetExtendedExecuteContext().SetRuntimeSettings(RuntimeSettings);
 
 		{
 			FUpdateGraphContext UpdateGraphContext(*GraphInstance.Get(), 1.0f / 30.0f);
@@ -676,7 +677,7 @@ bool FAnimationAnimNextRuntimeTest_GraphExecuteLatent::RunTest(const FString& In
 		{
 			Messages.Add(Message);
 		});
-		GraphInstance->GetExtendedExecuteContext().SetRuntimeSettings(RuntimeSettings);
+		GraphInstance->GetComponent<FUAFRigVMComponent>().GetExtendedExecuteContext().SetRuntimeSettings(RuntimeSettings);
 		
 		{
 			FUpdateGraphContext UpdateGraphContext(*GraphInstance.Get(), 1.0f / 30.0f);
@@ -737,10 +738,10 @@ bool FAnimationAnimNextRuntimeTest_Variables::RunTest(const FString& InParameter
 		AddPublicVariable(AnimationGraph, "Enum", EPropertyBagPropertyType::Double);
 		AddPublicVariable(AnimationGraph, "Struct", FVector::OneVector);
 		AddPublicVariable(AnimationGraph, "DerivedStruct", FTestDerivedVector());
-		AddPublicVariable(AnimationGraph, "Object", UAnimNextDataInterface::StaticClass()->GetDefaultObject<UAnimNextDataInterface>());
-		AddPublicVariable(AnimationGraph, "SoftObject", FSoftObjectPath(UAnimNextDataInterface::StaticClass()->GetDefaultObject<UAnimNextDataInterface>()));
-		AddPublicVariable(AnimationGraph, "Class", UAnimNextDataInterface::StaticClass());
-		AddPublicVariable(AnimationGraph, "SoftClass", FSoftClassPath(UAnimNextDataInterface::StaticClass()));
+		AddPublicVariable(AnimationGraph, "Object", UAnimNextSharedVariables::StaticClass()->GetDefaultObject<UAnimNextSharedVariables>());
+		AddPublicVariable(AnimationGraph, "SoftObject", FSoftObjectPath(UAnimNextSharedVariables::StaticClass()->GetDefaultObject<UAnimNextSharedVariables>()));
+		AddPublicVariable(AnimationGraph, "Class", UAnimNextSharedVariables::StaticClass());
+		AddPublicVariable(AnimationGraph, "SoftClass", FSoftClassPath(UAnimNextSharedVariables::StaticClass()));
 
 		TSharedPtr<FAnimNextGraphInstance> GraphInstance = AnimationGraph->AllocateInstance();
 
@@ -748,134 +749,134 @@ bool FAnimationAnimNextRuntimeTest_Variables::RunTest(const FString& InParameter
 		{
 			// Gets
 			bool TestBool = false;
-			UE_RETURN_ON_ERROR(GraphInstance->GetVariable("Bool", TestBool) == EPropertyBagResult::Success, "FAnimationAnimNextRuntimeTest_Variables -> GetVariable failed");
+			UE_RETURN_ON_ERROR(GraphInstance->GetVariable(FAnimNextVariableReference("Bool", AnimationGraph), TestBool) == EPropertyBagResult::Success, "FAnimationAnimNextRuntimeTest_Variables -> GetVariable failed");
 			UE_RETURN_ON_ERROR(TestBool == true, "FAnimationAnimNextRuntimeTest_Variables -> Variable value did not match");
 
 			uint8 TestByte = 0;
-			UE_RETURN_ON_ERROR(GraphInstance->GetVariable("Byte", TestByte) == EPropertyBagResult::Success, "FAnimationAnimNextRuntimeTest_Variables -> GetVariable failed");
+			UE_RETURN_ON_ERROR(GraphInstance->GetVariable(FAnimNextVariableReference("Byte", AnimationGraph), TestByte) == EPropertyBagResult::Success, "FAnimationAnimNextRuntimeTest_Variables -> GetVariable failed");
 			UE_RETURN_ON_ERROR(TestByte == 42, "FAnimationAnimNextRuntimeTest_Variables -> Variable value did not match");
 
 			int32 TestInt32 = 0;
-			UE_RETURN_ON_ERROR(GraphInstance->GetVariable("Int32", TestInt32) == EPropertyBagResult::Success, "FAnimationAnimNextRuntimeTest_Variables -> GetVariable failed");
+			UE_RETURN_ON_ERROR(GraphInstance->GetVariable(FAnimNextVariableReference("Int32", AnimationGraph), TestInt32) == EPropertyBagResult::Success, "FAnimationAnimNextRuntimeTest_Variables -> GetVariable failed");
 			UE_RETURN_ON_ERROR(TestInt32 == -4679222, "FAnimationAnimNextRuntimeTest_Variables -> Variable value did not match");
 
 			uint32 TestUInt32 = 0;
-			UE_RETURN_ON_ERROR(GraphInstance->GetVariable("UInt32", TestUInt32) == EPropertyBagResult::Success, "FAnimationAnimNextRuntimeTest_Variables -> GetVariable failed");
+			UE_RETURN_ON_ERROR(GraphInstance->GetVariable(FAnimNextVariableReference("UInt32", AnimationGraph), TestUInt32) == EPropertyBagResult::Success, "FAnimationAnimNextRuntimeTest_Variables -> GetVariable failed");
 			UE_RETURN_ON_ERROR(TestUInt32 == 3415919103, "FAnimationAnimNextRuntimeTest_Variables -> Variable value did not match");
 
 			int64 TestInt64 = 0;
-			UE_RETURN_ON_ERROR(GraphInstance->GetVariable("Int64", TestInt64) == EPropertyBagResult::Success, "FAnimationAnimNextRuntimeTest_Variables -> GetVariable failed");
+			UE_RETURN_ON_ERROR(GraphInstance->GetVariable(FAnimNextVariableReference("Int64", AnimationGraph), TestInt64) == EPropertyBagResult::Success, "FAnimationAnimNextRuntimeTest_Variables -> GetVariable failed");
 			UE_RETURN_ON_ERROR(TestInt64 == -3415919105, "FAnimationAnimNextRuntimeTest_Variables -> Variable value did not match");
 
 			uint64 TestUInt64 = 0;
-			UE_RETURN_ON_ERROR(GraphInstance->GetVariable("UInt64", TestUInt64) == EPropertyBagResult::Success, "FAnimationAnimNextRuntimeTest_Variables -> GetVariable failed");
+			UE_RETURN_ON_ERROR(GraphInstance->GetVariable(FAnimNextVariableReference("UInt64", AnimationGraph), TestUInt64) == EPropertyBagResult::Success, "FAnimationAnimNextRuntimeTest_Variables -> GetVariable failed");
 			UE_RETURN_ON_ERROR(TestUInt64 == 34159191067, "FAnimationAnimNextRuntimeTest_Variables -> Variable value did not match");
 
 			// Conversions
 
 			// Bool
 			TestByte = 0;
-			UE_RETURN_ON_ERROR(GraphInstance->GetVariable("Bool", TestByte) == EPropertyBagResult::Success, "FAnimationAnimNextRuntimeTest_Variables -> GetVariable failed");
+			UE_RETURN_ON_ERROR(GraphInstance->GetVariable(FAnimNextVariableReference("Bool", AnimationGraph), TestByte) == EPropertyBagResult::Success, "FAnimationAnimNextRuntimeTest_Variables -> GetVariable failed");
 			UE_RETURN_ON_ERROR(TestByte == 1, "FAnimationAnimNextRuntimeTest_Variables -> Variable value did not match");
 
 			TestInt32 = 0;
-			UE_RETURN_ON_ERROR(GraphInstance->GetVariable("Bool", TestInt32) == EPropertyBagResult::Success, "FAnimationAnimNextRuntimeTest_Variables -> GetVariable failed");
+			UE_RETURN_ON_ERROR(GraphInstance->GetVariable(FAnimNextVariableReference("Bool", AnimationGraph), TestInt32) == EPropertyBagResult::Success, "FAnimationAnimNextRuntimeTest_Variables -> GetVariable failed");
 			UE_RETURN_ON_ERROR(TestInt32 == 1, "FAnimationAnimNextRuntimeTest_Variables -> Variable value did not match");
 
 			TestUInt32 = 0;
-			UE_RETURN_ON_ERROR(GraphInstance->GetVariable("Bool", TestUInt32) == EPropertyBagResult::Success, "FAnimationAnimNextRuntimeTest_Variables -> GetVariable failed");
+			UE_RETURN_ON_ERROR(GraphInstance->GetVariable(FAnimNextVariableReference("Bool", AnimationGraph), TestUInt32) == EPropertyBagResult::Success, "FAnimationAnimNextRuntimeTest_Variables -> GetVariable failed");
 			UE_RETURN_ON_ERROR(TestUInt32 == 1, "FAnimationAnimNextRuntimeTest_Variables -> Variable value did not match");
 
 			TestInt64 = 0;
-			UE_RETURN_ON_ERROR(GraphInstance->GetVariable("Bool", TestInt64) == EPropertyBagResult::Success, "FAnimationAnimNextRuntimeTest_Variables -> GetVariable failed");
+			UE_RETURN_ON_ERROR(GraphInstance->GetVariable(FAnimNextVariableReference("Bool", AnimationGraph), TestInt64) == EPropertyBagResult::Success, "FAnimationAnimNextRuntimeTest_Variables -> GetVariable failed");
 			UE_RETURN_ON_ERROR(TestInt64 == 1, "FAnimationAnimNextRuntimeTest_Variables -> Variable value did not match");
 
 			TestUInt64 = 0;
-			UE_RETURN_ON_ERROR(GraphInstance->GetVariable("Bool", TestUInt64) == EPropertyBagResult::Success, "FAnimationAnimNextRuntimeTest_Variables -> GetVariable failed");
+			UE_RETURN_ON_ERROR(GraphInstance->GetVariable(FAnimNextVariableReference("Bool", AnimationGraph), TestUInt64) == EPropertyBagResult::Success, "FAnimationAnimNextRuntimeTest_Variables -> GetVariable failed");
 			UE_RETURN_ON_ERROR(TestUInt64 == 1, "FAnimationAnimNextRuntimeTest_Variables -> Variable value did not match");
 
 			// Byte
 			TestBool = false;
-			UE_RETURN_ON_ERROR(GraphInstance->GetVariable("Byte", TestBool) == EPropertyBagResult::Success, "FAnimationAnimNextRuntimeTest_Variables -> GetVariable failed");
+			UE_RETURN_ON_ERROR(GraphInstance->GetVariable(FAnimNextVariableReference("Byte", AnimationGraph), TestBool) == EPropertyBagResult::Success, "FAnimationAnimNextRuntimeTest_Variables -> GetVariable failed");
 			UE_RETURN_ON_ERROR(TestBool == true, "FAnimationAnimNextRuntimeTest_Variables -> Variable value did not match");
 
 			TestInt32 = 0;
-			UE_RETURN_ON_ERROR(GraphInstance->GetVariable("Byte", TestInt32) == EPropertyBagResult::Success, "FAnimationAnimNextRuntimeTest_Variables -> GetVariable failed");
+			UE_RETURN_ON_ERROR(GraphInstance->GetVariable(FAnimNextVariableReference("Byte", AnimationGraph), TestInt32) == EPropertyBagResult::Success, "FAnimationAnimNextRuntimeTest_Variables -> GetVariable failed");
 			UE_RETURN_ON_ERROR(TestInt32 == 42, "FAnimationAnimNextRuntimeTest_Variables -> Variable value did not match");
 
 			TestUInt32 = 0;
-			UE_RETURN_ON_ERROR(GraphInstance->GetVariable("Byte", TestUInt32) == EPropertyBagResult::Success, "FAnimationAnimNextRuntimeTest_Variables -> GetVariable failed");
+			UE_RETURN_ON_ERROR(GraphInstance->GetVariable(FAnimNextVariableReference("Byte", AnimationGraph), TestUInt32) == EPropertyBagResult::Success, "FAnimationAnimNextRuntimeTest_Variables -> GetVariable failed");
 			UE_RETURN_ON_ERROR(TestUInt32 == 42, "FAnimationAnimNextRuntimeTest_Variables -> Variable value did not match");
 
 			TestInt64 = 0;
-			UE_RETURN_ON_ERROR(GraphInstance->GetVariable("Byte", TestInt64) == EPropertyBagResult::Success, "FAnimationAnimNextRuntimeTest_Variables -> GetVariable failed");
+			UE_RETURN_ON_ERROR(GraphInstance->GetVariable(FAnimNextVariableReference("Byte", AnimationGraph), TestInt64) == EPropertyBagResult::Success, "FAnimationAnimNextRuntimeTest_Variables -> GetVariable failed");
 			UE_RETURN_ON_ERROR(TestInt64 == 42, "FAnimationAnimNextRuntimeTest_Variables -> Variable value did not match");
 
 			TestUInt64 = 0;
-			UE_RETURN_ON_ERROR(GraphInstance->GetVariable("Byte", TestUInt64) == EPropertyBagResult::Success, "FAnimationAnimNextRuntimeTest_Variables -> GetVariable failed");
+			UE_RETURN_ON_ERROR(GraphInstance->GetVariable(FAnimNextVariableReference("Byte", AnimationGraph), TestUInt64) == EPropertyBagResult::Success, "FAnimationAnimNextRuntimeTest_Variables -> GetVariable failed");
 			UE_RETURN_ON_ERROR(TestUInt64 == 42, "FAnimationAnimNextRuntimeTest_Variables -> Variable value did not match");
 
 			// Int32
 			TestBool = false;
-			UE_RETURN_ON_ERROR(GraphInstance->GetVariable("Int32", TestBool) == EPropertyBagResult::Success, "FAnimationAnimNextRuntimeTest_Variables -> GetVariable failed");
+			UE_RETURN_ON_ERROR(GraphInstance->GetVariable(FAnimNextVariableReference("Int32", AnimationGraph), TestBool) == EPropertyBagResult::Success, "FAnimationAnimNextRuntimeTest_Variables -> GetVariable failed");
 			UE_RETURN_ON_ERROR(TestBool == true, "FAnimationAnimNextRuntimeTest_Variables -> Variable value did not match");
 
 			TestByte = 0;
-			UE_RETURN_ON_ERROR(GraphInstance->GetVariable("Int32", TestByte) == EPropertyBagResult::Success, "FAnimationAnimNextRuntimeTest_Variables -> GetVariable failed");
+			UE_RETURN_ON_ERROR(GraphInstance->GetVariable(FAnimNextVariableReference("Int32", AnimationGraph), TestByte) == EPropertyBagResult::Success, "FAnimationAnimNextRuntimeTest_Variables -> GetVariable failed");
 			UE_RETURN_ON_ERROR(TestByte == 202/*(uint8)-4679222*/, "FAnimationAnimNextRuntimeTest_Variables -> Variable value did not match");
 
 			TestUInt32 = 0;
-			UE_RETURN_ON_ERROR(GraphInstance->GetVariable("Int32", TestUInt32) == EPropertyBagResult::Success, "FAnimationAnimNextRuntimeTest_Variables -> GetVariable failed");
+			UE_RETURN_ON_ERROR(GraphInstance->GetVariable(FAnimNextVariableReference("Int32", AnimationGraph), TestUInt32) == EPropertyBagResult::Success, "FAnimationAnimNextRuntimeTest_Variables -> GetVariable failed");
 			UE_RETURN_ON_ERROR(TestUInt32 == 4290288074/*(uint32)-4679222*/, "FAnimationAnimNextRuntimeTest_Variables -> Variable value did not match");
 
 			TestInt64 = 0;
-			UE_RETURN_ON_ERROR(GraphInstance->GetVariable("Int32", TestInt64) == EPropertyBagResult::Success, "FAnimationAnimNextRuntimeTest_Variables -> GetVariable failed");
+			UE_RETURN_ON_ERROR(GraphInstance->GetVariable(FAnimNextVariableReference("Int32", AnimationGraph), TestInt64) == EPropertyBagResult::Success, "FAnimationAnimNextRuntimeTest_Variables -> GetVariable failed");
 			UE_RETURN_ON_ERROR(TestInt64 == -4679222, "FAnimationAnimNextRuntimeTest_Variables -> Variable value did not match");
 
 			TestUInt64 = 0;
-			UE_RETURN_ON_ERROR(GraphInstance->GetVariable("Int32", TestUInt64) == EPropertyBagResult::Success, "FAnimationAnimNextRuntimeTest_Variables -> GetVariable failed");
+			UE_RETURN_ON_ERROR(GraphInstance->GetVariable(FAnimNextVariableReference("Int32", AnimationGraph), TestUInt64) == EPropertyBagResult::Success, "FAnimationAnimNextRuntimeTest_Variables -> GetVariable failed");
 			UE_RETURN_ON_ERROR(TestUInt64 == 4290288074/*(uint64)-4679222*/, "FAnimationAnimNextRuntimeTest_Variables -> Variable value did not match");
 
 			// Uint32
 			TestBool = false;
-			UE_RETURN_ON_ERROR(GraphInstance->GetVariable("UInt32", TestBool) == EPropertyBagResult::Success, "FAnimationAnimNextRuntimeTest_Variables -> GetVariable failed");
+			UE_RETURN_ON_ERROR(GraphInstance->GetVariable(FAnimNextVariableReference("UInt32", AnimationGraph), TestBool) == EPropertyBagResult::Success, "FAnimationAnimNextRuntimeTest_Variables -> GetVariable failed");
 			UE_RETURN_ON_ERROR(TestBool == true, "FAnimationAnimNextRuntimeTest_Variables -> Variable value did not match");
 
 			TestByte = 0;
-			UE_RETURN_ON_ERROR(GraphInstance->GetVariable("UInt32", TestByte) == EPropertyBagResult::Success, "FAnimationAnimNextRuntimeTest_Variables -> GetVariable failed");
+			UE_RETURN_ON_ERROR(GraphInstance->GetVariable(FAnimNextVariableReference("UInt32", AnimationGraph), TestByte) == EPropertyBagResult::Success, "FAnimationAnimNextRuntimeTest_Variables -> GetVariable failed");
 			UE_RETURN_ON_ERROR(TestByte == 255/*(uint8)3415919103*/, "FAnimationAnimNextRuntimeTest_Variables -> Variable value did not match");
 
 			TestInt32 = 0;
-			UE_RETURN_ON_ERROR(GraphInstance->GetVariable("UInt32", TestInt32) == EPropertyBagResult::Success, "FAnimationAnimNextRuntimeTest_Variables -> GetVariable failed");
+			UE_RETURN_ON_ERROR(GraphInstance->GetVariable(FAnimNextVariableReference("UInt32", AnimationGraph), TestInt32) == EPropertyBagResult::Success, "FAnimationAnimNextRuntimeTest_Variables -> GetVariable failed");
 			UE_RETURN_ON_ERROR(TestInt32 == -879048193/*(int32)3415919103*/, "FAnimationAnimNextRuntimeTest_Variables -> Variable value did not match");
 
 			TestInt64 = 0;
-			UE_RETURN_ON_ERROR(GraphInstance->GetVariable("UInt32", TestInt64) == EPropertyBagResult::Success, "FAnimationAnimNextRuntimeTest_Variables -> GetVariable failed");
+			UE_RETURN_ON_ERROR(GraphInstance->GetVariable(FAnimNextVariableReference("UInt32", AnimationGraph), TestInt64) == EPropertyBagResult::Success, "FAnimationAnimNextRuntimeTest_Variables -> GetVariable failed");
 			UE_RETURN_ON_ERROR(TestInt64 == 3415919103, "FAnimationAnimNextRuntimeTest_Variables -> Variable value did not match");
 
 			TestUInt64 = 0;
-			UE_RETURN_ON_ERROR(GraphInstance->GetVariable("UInt32", TestUInt64) == EPropertyBagResult::Success, "FAnimationAnimNextRuntimeTest_Variables -> GetVariable failed");
+			UE_RETURN_ON_ERROR(GraphInstance->GetVariable(FAnimNextVariableReference("UInt32", AnimationGraph), TestUInt64) == EPropertyBagResult::Success, "FAnimationAnimNextRuntimeTest_Variables -> GetVariable failed");
 			UE_RETURN_ON_ERROR(TestUInt64 == 3415919103, "FAnimationAnimNextRuntimeTest_Variables -> Variable value did not match");
 
 			// Int64
 			TestBool = false;
-			UE_RETURN_ON_ERROR(GraphInstance->GetVariable("Int64", TestBool) == EPropertyBagResult::Success, "FAnimationAnimNextRuntimeTest_Variables -> GetVariable failed");
+			UE_RETURN_ON_ERROR(GraphInstance->GetVariable(FAnimNextVariableReference("Int64", AnimationGraph), TestBool) == EPropertyBagResult::Success, "FAnimationAnimNextRuntimeTest_Variables -> GetVariable failed");
 			UE_RETURN_ON_ERROR(TestBool == true, "FAnimationAnimNextRuntimeTest_Variables -> Variable value did not match");
 
 			TestByte = 0;
-			UE_RETURN_ON_ERROR(GraphInstance->GetVariable("Int64", TestByte) == EPropertyBagResult::Success, "FAnimationAnimNextRuntimeTest_Variables -> GetVariable failed");
+			UE_RETURN_ON_ERROR(GraphInstance->GetVariable(FAnimNextVariableReference("Int64", AnimationGraph), TestByte) == EPropertyBagResult::Success, "FAnimationAnimNextRuntimeTest_Variables -> GetVariable failed");
 			UE_RETURN_ON_ERROR(TestByte == 255/*(uint8)-3415919105*/, "FAnimationAnimNextRuntimeTest_Variables -> Variable value did not match");
 
 			TestInt32 = 0;
-			UE_RETURN_ON_ERROR(GraphInstance->GetVariable("Int64", TestInt32) == EPropertyBagResult::Success, "FAnimationAnimNextRuntimeTest_Variables -> GetVariable failed");
+			UE_RETURN_ON_ERROR(GraphInstance->GetVariable(FAnimNextVariableReference("Int64", AnimationGraph), TestInt32) == EPropertyBagResult::Success, "FAnimationAnimNextRuntimeTest_Variables -> GetVariable failed");
 			UE_RETURN_ON_ERROR(TestInt32 == 879048191/*(int32)-3415919105*/, "FAnimationAnimNextRuntimeTest_Variables -> Variable value did not match");
 
 			TestUInt32 = 0;
-			UE_RETURN_ON_ERROR(GraphInstance->GetVariable("Int64", TestUInt32) == EPropertyBagResult::Success, "FAnimationAnimNextRuntimeTest_Variables -> GetVariable failed");
+			UE_RETURN_ON_ERROR(GraphInstance->GetVariable(FAnimNextVariableReference("Int64", AnimationGraph), TestUInt32) == EPropertyBagResult::Success, "FAnimationAnimNextRuntimeTest_Variables -> GetVariable failed");
 			UE_RETURN_ON_ERROR(TestUInt32 == 879048191/*(uint32)-3415919105*/, "FAnimationAnimNextRuntimeTest_Variables -> Variable value did not match");
 
 			TestUInt64 = 0;
-			UE_RETURN_ON_ERROR(GraphInstance->GetVariable("Int64", TestUInt64) == EPropertyBagResult::Success, "FAnimationAnimNextRuntimeTest_Variables -> GetVariable failed");
+			UE_RETURN_ON_ERROR(GraphInstance->GetVariable(FAnimNextVariableReference("Int64", AnimationGraph), TestUInt64) == EPropertyBagResult::Success, "FAnimationAnimNextRuntimeTest_Variables -> GetVariable failed");
 			UE_RETURN_ON_ERROR(TestUInt64 == 18446744070293632511ull/*(uint64)-3415919105*/, "FAnimationAnimNextRuntimeTest_Variables -> Variable value did not match");
 		}
 
@@ -883,110 +884,110 @@ bool FAnimationAnimNextRuntimeTest_Variables::RunTest(const FString& InParameter
 		{
 			// Gets
 			float TestFloat = 0.0f;
-			UE_RETURN_ON_ERROR(GraphInstance->GetVariable("Float", TestFloat) == EPropertyBagResult::Success, "FAnimationAnimNextRuntimeTest_Variables -> GetVariable failed");
+			UE_RETURN_ON_ERROR(GraphInstance->GetVariable(FAnimNextVariableReference("Float", AnimationGraph), TestFloat) == EPropertyBagResult::Success, "FAnimationAnimNextRuntimeTest_Variables -> GetVariable failed");
 			UE_RETURN_ON_ERROR(TestFloat == 1.0f, "FAnimationAnimNextRuntimeTest_Variables -> Variable value did not match");
 
 			double TestDouble = 0.0;
-			UE_RETURN_ON_ERROR(GraphInstance->GetVariable("Double", TestDouble) == EPropertyBagResult::Success, "FAnimationAnimNextRuntimeTest_Variables -> GetVariable failed");
+			UE_RETURN_ON_ERROR(GraphInstance->GetVariable(FAnimNextVariableReference("Double", AnimationGraph), TestDouble) == EPropertyBagResult::Success, "FAnimationAnimNextRuntimeTest_Variables -> GetVariable failed");
 			UE_RETURN_ON_ERROR(TestDouble == 1.0, "FAnimationAnimNextRuntimeTest_Variables -> Variable value did not match");
 
 			// Conversions
 			TestDouble = 0.0f;
-			UE_RETURN_ON_ERROR(GraphInstance->GetVariable("Float", TestDouble) == EPropertyBagResult::Success, "FAnimationAnimNextRuntimeTest_Variables -> GetVariable failed");
+			UE_RETURN_ON_ERROR(GraphInstance->GetVariable(FAnimNextVariableReference("Float", AnimationGraph), TestDouble) == EPropertyBagResult::Success, "FAnimationAnimNextRuntimeTest_Variables -> GetVariable failed");
 			UE_RETURN_ON_ERROR(TestDouble == 1.0f, "FAnimationAnimNextRuntimeTest_Variables -> Variable value did not match");
 
 			TestFloat = 0.0;
-			UE_RETURN_ON_ERROR(GraphInstance->GetVariable("Double", TestFloat) == EPropertyBagResult::Success, "FAnimationAnimNextRuntimeTest_Variables -> GetVariable failed");
+			UE_RETURN_ON_ERROR(GraphInstance->GetVariable(FAnimNextVariableReference("Double", AnimationGraph), TestFloat) == EPropertyBagResult::Success, "FAnimationAnimNextRuntimeTest_Variables -> GetVariable failed");
 			UE_RETURN_ON_ERROR(TestFloat == 1.0, "FAnimationAnimNextRuntimeTest_Variables -> Variable value did not match");
 		}
 		
 		FName TestFName = NAME_None;
-		UE_RETURN_ON_ERROR(GraphInstance->GetVariable("Name", TestFName) == EPropertyBagResult::Success, "FAnimationAnimNextRuntimeTest_Variables -> GetVariable failed");
+		UE_RETURN_ON_ERROR(GraphInstance->GetVariable(FAnimNextVariableReference("Name", AnimationGraph), TestFName) == EPropertyBagResult::Success, "FAnimationAnimNextRuntimeTest_Variables -> GetVariable failed");
 		UE_RETURN_ON_ERROR(TestFName == "Test", "FAnimationAnimNextRuntimeTest_Variables -> Variable value did not match");
 
 		FString TestString;
-		UE_RETURN_ON_ERROR(GraphInstance->GetVariable("String", TestString) == EPropertyBagResult::Success, "FAnimationAnimNextRuntimeTest_Variables -> GetVariable failed");
+		UE_RETURN_ON_ERROR(GraphInstance->GetVariable(FAnimNextVariableReference("String", AnimationGraph), TestString) == EPropertyBagResult::Success, "FAnimationAnimNextRuntimeTest_Variables -> GetVariable failed");
 		UE_RETURN_ON_ERROR(TestString == "Test", "FAnimationAnimNextRuntimeTest_Variables -> Variable value did not match");
 
 		FText TestText;
-		UE_RETURN_ON_ERROR(GraphInstance->GetVariable("Text", TestText) == EPropertyBagResult::Success, "FAnimationAnimNextRuntimeTest_Variables -> GetVariable failed");
+		UE_RETURN_ON_ERROR(GraphInstance->GetVariable(FAnimNextVariableReference("Text", AnimationGraph), TestText) == EPropertyBagResult::Success, "FAnimationAnimNextRuntimeTest_Variables -> GetVariable failed");
 		UE_RETURN_ON_ERROR(TestText.EqualTo(NSLOCTEXT("Tests", "Test", "Test")), "FAnimationAnimNextRuntimeTest_Variables -> Variable value did not match");
 
 		// Enums
 		{
 			EPropertyBagPropertyType TestEnum = EPropertyBagPropertyType::None;
-			UE_RETURN_ON_ERROR(GraphInstance->GetVariable("Enum", TestEnum) == EPropertyBagResult::Success, "FAnimationAnimNextRuntimeTest_Variables -> GetVariable failed");
+			UE_RETURN_ON_ERROR(GraphInstance->GetVariable(FAnimNextVariableReference("Enum", AnimationGraph), TestEnum) == EPropertyBagResult::Success, "FAnimationAnimNextRuntimeTest_Variables -> GetVariable failed");
 			UE_RETURN_ON_ERROR(TestEnum == EPropertyBagPropertyType::Double, "FAnimationAnimNextRuntimeTest_Variables -> Variable value did not match");
 
 			// Test mismatched enum fails
 			EPropertyBagContainerType MismatchedEnum = EPropertyBagContainerType::None;
-			UE_RETURN_ON_ERROR(GraphInstance->GetVariable("Enum", MismatchedEnum) != EPropertyBagResult::Success, "FAnimationAnimNextRuntimeTest_Variables -> GetVariable succeeded");
+			UE_RETURN_ON_ERROR(GraphInstance->GetVariable(FAnimNextVariableReference("Enum", AnimationGraph), MismatchedEnum) != EPropertyBagResult::Success, "FAnimationAnimNextRuntimeTest_Variables -> GetVariable succeeded");
 		}
 
 		// Structs
 		{
 			FVector TestStruct = FVector::ZeroVector;
-			UE_RETURN_ON_ERROR(GraphInstance->GetVariable("Struct", TestStruct) == EPropertyBagResult::Success, "FAnimationAnimNextRuntimeTest_Variables -> GetVariable failed");
+			UE_RETURN_ON_ERROR(GraphInstance->GetVariable(FAnimNextVariableReference("Struct", AnimationGraph), TestStruct) == EPropertyBagResult::Success, "FAnimationAnimNextRuntimeTest_Variables -> GetVariable failed");
 			UE_RETURN_ON_ERROR(TestStruct == FVector::OneVector, "FAnimationAnimNextRuntimeTest_Variables -> Variable value did not match");
 
 			// Test conversion from base -> derived fails 
 			FTestDerivedVector TestDerived;
-			UE_RETURN_ON_ERROR(GraphInstance->GetVariable("Struct", TestDerived) != EPropertyBagResult::Success, "FAnimationAnimNextRuntimeTest_Variables -> GetVariable succeeded");
+			UE_RETURN_ON_ERROR(GraphInstance->GetVariable(FAnimNextVariableReference("Struct", AnimationGraph), TestDerived) != EPropertyBagResult::Success, "FAnimationAnimNextRuntimeTest_Variables -> GetVariable succeeded");
 
 			// Test conversion from derived -> base
 			TestStruct = FVector::ZeroVector;
-			UE_RETURN_ON_ERROR(GraphInstance->GetVariable("DerivedStruct", TestStruct) == EPropertyBagResult::Success, "FAnimationAnimNextRuntimeTest_Variables -> GetVariable failed");
+			UE_RETURN_ON_ERROR(GraphInstance->GetVariable(FAnimNextVariableReference("DerivedStruct", AnimationGraph), TestStruct) == EPropertyBagResult::Success, "FAnimationAnimNextRuntimeTest_Variables -> GetVariable failed");
 			UE_RETURN_ON_ERROR(TestStruct == FVector::OneVector, "FAnimationAnimNextRuntimeTest_Variables -> Variable value did not match");
 		}
 
 		// Objects
 		{
-			UAnimNextDataInterface* TestObject = nullptr;
-			UE_RETURN_ON_ERROR(GraphInstance->GetVariable("Object", TestObject) == EPropertyBagResult::Success, "FAnimationAnimNextRuntimeTest_Variables -> GetVariable failed");
-			UE_RETURN_ON_ERROR(TestObject == UAnimNextDataInterface::StaticClass()->GetDefaultObject<UAnimNextDataInterface>(), "FAnimationAnimNextRuntimeTest_Variables -> Variable value did not match");
+			UAnimNextSharedVariables* TestObject = nullptr;
+			UE_RETURN_ON_ERROR(GraphInstance->GetVariable(FAnimNextVariableReference("Object", AnimationGraph), TestObject) == EPropertyBagResult::Success, "FAnimationAnimNextRuntimeTest_Variables -> GetVariable failed");
+			UE_RETURN_ON_ERROR(TestObject == UAnimNextSharedVariables::StaticClass()->GetDefaultObject<UAnimNextSharedVariables>(), "FAnimationAnimNextRuntimeTest_Variables -> Variable value did not match");
 
 			// Test unrelated object fails
 			UAnimSequence* TestAnimSequence = nullptr;
-			UE_RETURN_ON_ERROR(GraphInstance->GetVariable("Object", TestAnimSequence) != EPropertyBagResult::Success, "FAnimationAnimNextRuntimeTest_Variables -> GetVariable succeeded");
+			UE_RETURN_ON_ERROR(GraphInstance->GetVariable(FAnimNextVariableReference("Object", AnimationGraph), TestAnimSequence) != EPropertyBagResult::Success, "FAnimationAnimNextRuntimeTest_Variables -> GetVariable succeeded");
 
 			// Test Derived -> Base succeeds
 			UObject* BaseObject = nullptr;
-			UE_RETURN_ON_ERROR(GraphInstance->GetVariable("Object", BaseObject) == EPropertyBagResult::Success, "FAnimationAnimNextRuntimeTest_Variables -> GetVariable failed");
-			UE_RETURN_ON_ERROR(BaseObject == UAnimNextDataInterface::StaticClass()->GetDefaultObject<UAnimNextDataInterface>(), "FAnimationAnimNextRuntimeTest_Variables -> Variable value did not match");
+			UE_RETURN_ON_ERROR(GraphInstance->GetVariable(FAnimNextVariableReference("Object", AnimationGraph), BaseObject) == EPropertyBagResult::Success, "FAnimationAnimNextRuntimeTest_Variables -> GetVariable failed");
+			UE_RETURN_ON_ERROR(BaseObject == UAnimNextSharedVariables::StaticClass()->GetDefaultObject<UAnimNextSharedVariables>(), "FAnimationAnimNextRuntimeTest_Variables -> Variable value did not match");
 
 			// Test Base -> Derived fails
 			UAnimNextAnimationGraph* DerivedObject = nullptr;
-			UE_RETURN_ON_ERROR(GraphInstance->GetVariable("Object", DerivedObject) != EPropertyBagResult::Success, "FAnimationAnimNextRuntimeTest_Variables -> GetVariable succeeded");
+			UE_RETURN_ON_ERROR(GraphInstance->GetVariable(FAnimNextVariableReference("Object", AnimationGraph), DerivedObject) != EPropertyBagResult::Success, "FAnimationAnimNextRuntimeTest_Variables -> GetVariable succeeded");
 		}
 
 		FSoftObjectPath TestSoftObjectPath;
-		UE_RETURN_ON_ERROR(GraphInstance->GetVariable("SoftObject", TestSoftObjectPath) == EPropertyBagResult::Success, "FAnimationAnimNextRuntimeTest_Variables -> GetVariable failed");
-		UE_RETURN_ON_ERROR(TestSoftObjectPath == FSoftObjectPath(UAnimNextDataInterface::StaticClass()->GetDefaultObject<UAnimNextDataInterface>()), "FAnimationAnimNextRuntimeTest_Variables -> Variable value did not match");
+		UE_RETURN_ON_ERROR(GraphInstance->GetVariable(FAnimNextVariableReference("SoftObject", AnimationGraph), TestSoftObjectPath) == EPropertyBagResult::Success, "FAnimationAnimNextRuntimeTest_Variables -> GetVariable failed");
+		UE_RETURN_ON_ERROR(TestSoftObjectPath == FSoftObjectPath(UAnimNextSharedVariables::StaticClass()->GetDefaultObject<UAnimNextSharedVariables>()), "FAnimationAnimNextRuntimeTest_Variables -> Variable value did not match");
 
 		// Classes
 		{
 			UClass* TestClass = nullptr;
-			UE_RETURN_ON_ERROR(GraphInstance->GetVariable("Class", TestClass) == EPropertyBagResult::Success, "FAnimationAnimNextRuntimeTest_Variables -> GetVariable failed");
-			UE_RETURN_ON_ERROR(TestClass == UAnimNextDataInterface::StaticClass(), "FAnimationAnimNextRuntimeTest_Variables -> Variable value did not match");
+			UE_RETURN_ON_ERROR(GraphInstance->GetVariable(FAnimNextVariableReference("Class", AnimationGraph), TestClass) == EPropertyBagResult::Success, "FAnimationAnimNextRuntimeTest_Variables -> GetVariable failed");
+			UE_RETURN_ON_ERROR(TestClass == UAnimNextSharedVariables::StaticClass(), "FAnimationAnimNextRuntimeTest_Variables -> Variable value did not match");
 
 
 			// Test unrelated class fails
 			TSubclassOf<UAnimSequence> TestUnrelatedSubclassOf;
-			UE_RETURN_ON_ERROR(GraphInstance->GetVariable("Class", TestUnrelatedSubclassOf) != EPropertyBagResult::Success, "FAnimationAnimNextRuntimeTest_Variables -> GetVariable succeeded");
+			UE_RETURN_ON_ERROR(GraphInstance->GetVariable(FAnimNextVariableReference("Class", AnimationGraph), TestUnrelatedSubclassOf) != EPropertyBagResult::Success, "FAnimationAnimNextRuntimeTest_Variables -> GetVariable succeeded");
 
 			// Test Derived -> Base succeeds
 			TSubclassOf<UObject> BaseSubclassOf;
-			UE_RETURN_ON_ERROR(GraphInstance->GetVariable("Class", BaseSubclassOf) == EPropertyBagResult::Success, "FAnimationAnimNextRuntimeTest_Variables -> GetVariable failed");
-			UE_RETURN_ON_ERROR(BaseSubclassOf == UAnimNextDataInterface::StaticClass(), "FAnimationAnimNextRuntimeTest_Variables -> Variable value did not match");
+			UE_RETURN_ON_ERROR(GraphInstance->GetVariable(FAnimNextVariableReference("Class", AnimationGraph), BaseSubclassOf) == EPropertyBagResult::Success, "FAnimationAnimNextRuntimeTest_Variables -> GetVariable failed");
+			UE_RETURN_ON_ERROR(BaseSubclassOf == UAnimNextSharedVariables::StaticClass(), "FAnimationAnimNextRuntimeTest_Variables -> Variable value did not match");
 
 			// Test Base -> Derived fails
 			TSubclassOf<UAnimNextAnimationGraph> DerivedSubclassOf;
-			UE_RETURN_ON_ERROR(GraphInstance->GetVariable("Class", DerivedSubclassOf) != EPropertyBagResult::Success, "FAnimationAnimNextRuntimeTest_Variables -> GetVariable succeeded");
+			UE_RETURN_ON_ERROR(GraphInstance->GetVariable(FAnimNextVariableReference("Class", AnimationGraph), DerivedSubclassOf) != EPropertyBagResult::Success, "FAnimationAnimNextRuntimeTest_Variables -> GetVariable succeeded");
 		}
 
 		FSoftClassPath TestSoftClassPath;
-		UE_RETURN_ON_ERROR(GraphInstance->GetVariable("SoftClass", TestSoftClassPath) == EPropertyBagResult::Success, "FAnimationAnimNextRuntimeTest_Variables -> GetVariable failed");
-		UE_RETURN_ON_ERROR(TestSoftClassPath == FSoftClassPath(UAnimNextDataInterface::StaticClass()), "FAnimationAnimNextRuntimeTest_Variables -> Variable value did not match");
+		UE_RETURN_ON_ERROR(GraphInstance->GetVariable(FAnimNextVariableReference("SoftClass", AnimationGraph), TestSoftClassPath) == EPropertyBagResult::Success, "FAnimationAnimNextRuntimeTest_Variables -> GetVariable failed");
+		UE_RETURN_ON_ERROR(TestSoftClassPath == FSoftClassPath(UAnimNextSharedVariables::StaticClass()), "FAnimationAnimNextRuntimeTest_Variables -> Variable value did not match");
 	}
 
 	Tests::FUtils::CleanupAfterTests();
diff --git a/Engine/Plugins/Experimental/AnimNextAnimGraph/Source/AnimNextAnimGraphTestSuite/Private/AnimNextTraitBaseTest.cpp b/Engine/Plugins/Experimental/AnimNextAnimGraph/Source/AnimNextAnimGraphTestSuite/Private/AnimNextTraitBaseTest.cpp
--- a/Engine/Plugins/Experimental/AnimNextAnimGraph/Source/AnimNextAnimGraphTestSuite/Private/AnimNextTraitBaseTest.cpp
+++ b/Engine/Plugins/Experimental/AnimNextAnimGraph/Source/AnimNextAnimGraphTestSuite/Private/AnimNextTraitBaseTest.cpp
@@ -3,6 +3,9 @@
 #include "AnimNextTraitBaseTest.h"
 #include "AnimNextRuntimeTest.h"
 #include "AnimNextTest.h"
+#include "Factory/AnimGraphBuilderContext.h"
+#include "Factory/AnimGraphFactory.h"
+#include "Factory/AnimNextSimpleAnimGraphBuilder.h"
 
 #include "Misc/AutomationTest.h"
 
@@ -1498,6 +1501,155 @@ bool FAnimationAnimNextRuntimeTest_TraitSerialization::RunTest(const FString& In
 	return true;
 }
 
+IMPLEMENT_SIMPLE_AUTOMATION_TEST(FAnimationAnimNextRuntimeTest_SimpleGraphBuilder, "Animation.AnimNext.Runtime.SimpleGraphBuilder", EAutomationTestFlags::EditorContext | EAutomationTestFlags::EngineFilter)
+
+bool FAnimationAnimNextRuntimeTest_SimpleGraphBuilder::RunTest(const FString& InParameters)
+{
+	using namespace UE::AnimNext;
+
+	{
+		FScopedClearNodeTemplateRegistry ScopedClearNodeTemplateRegistry;
+		FNodeTemplateRegistry& Registry = FNodeTemplateRegistry::Get();
+
+		AUTO_REGISTER_ANIM_TRAIT(FTraitSerialization_Base)
+		AUTO_REGISTER_ANIM_TRAIT(FTraitSerialization_Add)
+		AUTO_REGISTER_ANIM_TRAIT(FTraitNativeSerialization_Add)
+	
+		FTraitSerialization_Base::FSharedData TraitBaseRef0;
+		TraitBaseRef0.Integer = 1651;
+		TraitBaseRef0.IntegerArray[0] = 1071;
+		TraitBaseRef0.IntegerArray[1] = -158;
+		TraitBaseRef0.IntegerArray[2] = 88116;
+		TraitBaseRef0.IntegerArray[3] = 0x417;
+		TraitBaseRef0.IntegerTArray = { -8162, 88152, 0x8152f };
+		TraitBaseRef0.Vector = FVector(0.1917, 12435.1, -18200.1726);
+		TraitBaseRef0.VectorArray[0] = FVector(192.1716, -1927.115, 99176.12);
+		TraitBaseRef0.VectorArray[1] = FVector(961.811, -18956.117, 81673.44);
+		TraitBaseRef0.VectorTArray = { FVector(-1927.8771, 1826.9917, -123.1555), FVector(9177.011, -71.44, -917.88), FVector(123.91, 852.11, -81652.1) };
+		TraitBaseRef0.String = TEXT("sample string 123");
+		TraitBaseRef0.Name = FName(TEXT("sample name 999178"));
+
+		FTraitSerialization_Add::FSharedData TraitAddRef0;
+		TraitAddRef0.Integer = 16511;
+		TraitAddRef0.IntegerArray[0] = 10711;
+		TraitAddRef0.IntegerArray[1] = -1581;
+		TraitAddRef0.IntegerArray[2] = 881161;
+		TraitAddRef0.IntegerArray[3] = 0x4171;
+		TraitAddRef0.IntegerTArray = { -81621, 881521, 0x8152f1 };
+		TraitAddRef0.Vector = FVector(0.19171, 12435.11, -18200.17261);
+		TraitAddRef0.VectorArray[0] = FVector(192.17161, -1927.1151, 99176.121);
+		TraitAddRef0.VectorArray[1] = FVector(961.8111, -18956.1171, 81673.441);
+		TraitAddRef0.VectorTArray = { FVector(-1927.87711, 1826.99171, -123.15551), FVector(9177.0111, -71.441, -917.881), FVector(123.911, 852.111, -81652.11) };
+		TraitAddRef0.String = TEXT("sample string 1231");
+		TraitAddRef0.Name = FName(TEXT("sample name 9991781"));
+
+		FTraitNativeSerialization_Add::FSharedData TraitNativeRef0;
+		TraitNativeRef0.Integer = 16514;
+		TraitNativeRef0.IntegerArray[0] = 10714;
+		TraitNativeRef0.IntegerArray[1] = -1584;
+		TraitNativeRef0.IntegerArray[2] = 881164;
+		TraitNativeRef0.IntegerArray[3] = 0x4174;
+		TraitNativeRef0.IntegerTArray = { -81624, 881524, 0x8152f4 };
+		TraitNativeRef0.Vector = FVector(0.19174, 12435.14, -18200.17264);
+		TraitNativeRef0.VectorArray[0] = FVector(192.17164, -1927.1154, 99176.124);
+		TraitNativeRef0.VectorArray[1] = FVector(961.8114, -18956.1174, 81673.444);
+		TraitNativeRef0.VectorTArray = { FVector(-1927.87714, 1826.99174, -123.15554), FVector(9177.0114, -71.444, -917.884), FVector(123.914, 852.114, -81652.14) };
+		TraitNativeRef0.String = TEXT("sample string 1234");
+		TraitNativeRef0.Name = FName(TEXT("sample name 9991784"));
+		
+		{
+			FAnimGraphBuilderContext AnimGraphBuilderContext;
+			{
+				FAnimNextSimpleAnimGraphBuilder Builder;
+				Builder.PushTraitStructViewToStack(0, TraitBaseRef0);
+				Builder.PushTraitStructViewToStack(0, TraitAddRef0);
+				Builder.PushTraitStructViewToStack(0, TraitNativeRef0);
+
+				const bool bResult = Builder.Build(AnimGraphBuilderContext);
+				AddErrorIfFalse(bResult, "FAnimationAnimNextRuntimeTest_SimpleGraphBuilder -> Builder failed");
+			}
+
+			// Clear out the node template registry to test registration on load
+			{
+				FScopedClearNodeTemplateRegistry ScopedClearNodeTemplateRegistryForLoad;
+
+				AddErrorIfFalse(Registry.GetNum() == 0, "FAnimationAnimNextRuntimeTest_SimpleGraphBuilder -> Registry should contain 0 templates");
+
+				// Builder context loads from archive buffer
+				const UAnimNextAnimationGraph* AnimationGraph = AnimGraphBuilderContext.Build();
+				UE_RETURN_ON_ERROR(AnimationGraph != nullptr, "FAnimationAnimNextRuntimeTest_SimpleGraphBuilder -> Failed to create animation graph");
+
+				TSharedPtr<FAnimNextGraphInstance> GraphInstance = AnimationGraph->AllocateInstance();
+
+				FExecutionContext Context(*GraphInstance.Get());
+
+				// Validate trait serialization
+				{
+					FTraitStackBinding StackNode;
+					AddErrorIfFalse(Context.GetStack(GraphInstance->GetGraphRootPtr(), StackNode), "FAnimationAnimNextRuntimeTest_SimpleGraphBuilder -> Failed to bind to trait stack");
+
+					// Validate shared data for base trait on node 0
+					{
+						TTraitBinding<IInterfaceA> BindingA0;
+						AddErrorIfFalse(StackNode.GetInterface(BindingA0), "FAnimationAnimNextRuntimeTest_SimpleGraphBuilder -> InterfaceA not found");
+
+						const auto* SharedDataA0 = BindingA0.GetSharedData<FTraitSerialization_Base::FSharedData>();
+
+						AddErrorIfFalse(SharedDataA0->Integer == TraitBaseRef0.Integer, "FAnimationAnimNextRuntimeTest_SimpleGraphBuilder -> Unexpected serialized value");
+						AddErrorIfFalse(FMemory::Memcmp(SharedDataA0->IntegerArray, TraitBaseRef0.IntegerArray, sizeof(TraitBaseRef0.IntegerArray)) == 0, "FAnimationAnimNextRuntimeTest_SimpleGraphBuilder -> Unexpected serialized value");
+						AddErrorIfFalse(SharedDataA0->IntegerTArray == TraitBaseRef0.IntegerTArray, "FAnimationAnimNextRuntimeTest_SimpleGraphBuilder -> Unexpected serialized value");
+						AddErrorIfFalse(SharedDataA0->Vector == TraitBaseRef0.Vector, "FAnimationAnimNextRuntimeTest_SimpleGraphBuilder -> Unexpected serialized value");
+						AddErrorIfFalse(FMemory::Memcmp(SharedDataA0->VectorArray, TraitBaseRef0.VectorArray, sizeof(TraitBaseRef0.VectorArray)) == 0, "FAnimationAnimNextRuntimeTest_SimpleGraphBuilder -> Unexpected serialized value");
+						AddErrorIfFalse(SharedDataA0->VectorTArray == TraitBaseRef0.VectorTArray, "FAnimationAnimNextRuntimeTest_SimpleGraphBuilder -> Unexpected serialized value");
+						AddErrorIfFalse(SharedDataA0->String == TraitBaseRef0.String, "FAnimationAnimNextRuntimeTest_SimpleGraphBuilder -> Unexpected serialized value");
+						AddErrorIfFalse(SharedDataA0->Name == TraitBaseRef0.Name, "FAnimationAnimNextRuntimeTest_SimpleGraphBuilder -> Unexpected serialized value");
+					}
+
+					// Validate shared data for additive trait on node 0
+					{
+						TTraitBinding<IInterfaceB> BindingB0;
+						AddErrorIfFalse(StackNode.GetInterface(BindingB0), "FAnimationAnimNextRuntimeTest_SimpleGraphBuilder -> InterfaceB not found");
+
+						const auto* SharedDataB0 = BindingB0.GetSharedData<FTraitSerialization_Add::FSharedData>();
+
+						AddErrorIfFalse(SharedDataB0->Integer == TraitAddRef0.Integer, "FAnimationAnimNextRuntimeTest_SimpleGraphBuilder -> Unexpected serialized value");
+						AddErrorIfFalse(FMemory::Memcmp(SharedDataB0->IntegerArray, TraitAddRef0.IntegerArray, sizeof(TraitAddRef0.IntegerArray)) == 0, "FAnimationAnimNextRuntimeTest_SimpleGraphBuilder -> Unexpected serialized value");
+						AddErrorIfFalse(SharedDataB0->IntegerTArray == TraitAddRef0.IntegerTArray, "FAnimationAnimNextRuntimeTest_SimpleGraphBuilder -> Unexpected serialized value");
+						AddErrorIfFalse(SharedDataB0->Vector == TraitAddRef0.Vector, "FAnimationAnimNextRuntimeTest_SimpleGraphBuilder -> Unexpected serialized value");
+						AddErrorIfFalse(FMemory::Memcmp(SharedDataB0->VectorArray, TraitAddRef0.VectorArray, sizeof(TraitAddRef0.VectorArray)) == 0, "FAnimationAnimNextRuntimeTest_SimpleGraphBuilder -> Unexpected serialized value");
+						AddErrorIfFalse(SharedDataB0->VectorTArray == TraitAddRef0.VectorTArray, "FAnimationAnimNextRuntimeTest_SimpleGraphBuilder -> Unexpected serialized value");
+						AddErrorIfFalse(SharedDataB0->String == TraitAddRef0.String, "FAnimationAnimNextRuntimeTest_SimpleGraphBuilder -> Unexpected serialized value");
+						AddErrorIfFalse(SharedDataB0->Name == TraitAddRef0.Name, "FAnimationAnimNextRuntimeTest_SimpleGraphBuilder -> Unexpected serialized value");
+					}
+
+					// Validate shared data for native trait on node 0
+					{
+						TTraitBinding<IInterfaceC> BindingC0;
+						AddErrorIfFalse(StackNode.GetInterface(BindingC0), "FAnimationAnimNextRuntimeTest_SimpleGraphBuilder -> InterfaceC not found");
+
+						const auto* SharedDataC0 = BindingC0.GetSharedData<FTraitNativeSerialization_Add::FSharedData>();
+
+						AddErrorIfFalse(SharedDataC0->Integer == TraitNativeRef0.Integer, "FAnimationAnimNextRuntimeTest_SimpleGraphBuilder -> Unexpected serialized value");
+						AddErrorIfFalse(FMemory::Memcmp(SharedDataC0->IntegerArray, TraitNativeRef0.IntegerArray, sizeof(TraitNativeRef0.IntegerArray)) == 0, "FAnimationAnimNextRuntimeTest_SimpleGraphBuilder -> Unexpected serialized value");
+						AddErrorIfFalse(SharedDataC0->IntegerTArray == TraitNativeRef0.IntegerTArray, "FAnimationAnimNextRuntimeTest_SimpleGraphBuilder -> Unexpected serialized value");
+						AddErrorIfFalse(SharedDataC0->Vector == TraitNativeRef0.Vector, "FAnimationAnimNextRuntimeTest_SimpleGraphBuilder -> Unexpected serialized value");
+						AddErrorIfFalse(FMemory::Memcmp(SharedDataC0->VectorArray, TraitNativeRef0.VectorArray, sizeof(TraitNativeRef0.VectorArray)) == 0, "FAnimationAnimNextRuntimeTest_SimpleGraphBuilder -> Unexpected serialized value");
+						AddErrorIfFalse(SharedDataC0->VectorTArray == TraitNativeRef0.VectorTArray, "FAnimationAnimNextRuntimeTest_SimpleGraphBuilder -> Unexpected serialized value");
+						AddErrorIfFalse(SharedDataC0->String == TraitNativeRef0.String, "FAnimationAnimNextRuntimeTest_SimpleGraphBuilder -> Unexpected serialized value");
+						AddErrorIfFalse(SharedDataC0->Name == TraitNativeRef0.Name, "FAnimationAnimNextRuntimeTest_SimpleGraphBuilder -> Unexpected serialized value");
+						AddErrorIfFalse(SharedDataC0->bSerializeCalled, "FAnimationAnimNextRuntimeTest_SimpleGraphBuilder -> Unexpected serialized value");
+					}
+				}
+			}
+		}
+	}
+
+	Tests::FUtils::CleanupAfterTests();
+
+	return true;
+}
+
+
 #endif
 
 FTraitA_BaseSharedData::FTraitA_BaseSharedData()
diff --git a/Engine/Plugins/Experimental/AnimNextAnimGraph/Source/AnimNextAnimGraphUncookedOnly/Internal/AnimNextController.h b/Engine/Plugins/Experimental/AnimNextAnimGraph/Source/AnimNextAnimGraphUncookedOnly/Internal/AnimNextController.h
--- a/Engine/Plugins/Experimental/AnimNextAnimGraph/Source/AnimNextAnimGraphUncookedOnly/Internal/AnimNextController.h
+++ b/Engine/Plugins/Experimental/AnimNextAnimGraph/Source/AnimNextAnimGraphUncookedOnly/Internal/AnimNextController.h
@@ -3,7 +3,7 @@
 #pragma once
 
 #include "CoreMinimal.h"
-#include "RigVMModel/RigVMController.h"
+#include "AnimNextControllerBase.h"
 #include "AnimNextController.generated.h"
 
 namespace UE::AnimNext
@@ -16,7 +16,7 @@ namespace UE::AnimNext
   * Implements AnimNext RigVM controller extensions
   */
 UCLASS(MinimalAPI)
-class UAnimNextController : public URigVMController
+class UAnimNextController : public UAnimNextControllerBase
 {
 	GENERATED_BODY()
 
diff --git a/Engine/Plugins/Experimental/AnimNextAnimGraph/Source/AnimNextAnimGraphUncookedOnly/Internal/Graph/AnimNextAnimationGraph_EditorData.h b/Engine/Plugins/Experimental/AnimNextAnimGraph/Source/AnimNextAnimGraphUncookedOnly/Internal/Graph/AnimNextAnimationGraph_EditorData.h
--- a/Engine/Plugins/Experimental/AnimNextAnimGraph/Source/AnimNextAnimGraphUncookedOnly/Internal/Graph/AnimNextAnimationGraph_EditorData.h
+++ b/Engine/Plugins/Experimental/AnimNextAnimGraph/Source/AnimNextAnimGraphUncookedOnly/Internal/Graph/AnimNextAnimationGraph_EditorData.h
@@ -4,7 +4,7 @@
 
 #include "CoreMinimal.h"
 #include "AnimNextController.h"
-#include "DataInterface/AnimNextDataInterface_EditorData.h"
+#include "Variables/AnimNextSharedVariables_EditorData.h"
 #include "AnimNextExecuteContext.h"
 
 #include "AnimNextAnimationGraph_EditorData.generated.h"
@@ -42,7 +42,7 @@ public:
 
 /** Editor data for AnimNext animation graphs */
 UCLASS(MinimalAPI)
-class UAnimNextAnimationGraph_EditorData : public UAnimNextDataInterface_EditorData
+class UAnimNextAnimationGraph_EditorData : public UAnimNextSharedVariables_EditorData
 {
 	GENERATED_BODY()
 
@@ -75,8 +75,6 @@ protected:
 	UE_API virtual void InitializeAssetUserData() override;
 	UE_API virtual void OnPreCompileAsset(FRigVMCompileSettings& InSettings) override;
 	UE_API virtual void OnPreCompileGetProgrammaticFunctionHeaders(const FRigVMCompileSettings& InSettings, FAnimNextGetFunctionHeaderCompileContext& OutCompileContext) override;
-	UE_API virtual void OnPreCompileGetProgrammaticVariables(const FRigVMCompileSettings& InSettings, FAnimNextGetVariableCompileContext& OutCompileContext) override;
-	UE_API virtual void OnPreCompileGetProgrammaticGraphs(const FRigVMCompileSettings& InSettings, FAnimNextGetGraphCompileContext& OutCompileContext) override;
 	UE_API virtual void OnPreCompileProcessGraphs(const FRigVMCompileSettings& InSettings, FAnimNextProcessGraphCompileContext& OutCompileContext) override;
 	UE_API virtual void OnPostCompileCleanup(const FRigVMCompileSettings& InSettings) override;
 
diff --git a/Engine/Plugins/Experimental/AnimNextAnimGraph/Source/AnimNextAnimGraphUncookedOnly/Private/Graph/AnimNextAnimationGraph_EditorData.cpp b/Engine/Plugins/Experimental/AnimNextAnimGraph/Source/AnimNextAnimGraphUncookedOnly/Private/Graph/AnimNextAnimationGraph_EditorData.cpp
--- a/Engine/Plugins/Experimental/AnimNextAnimGraph/Source/AnimNextAnimGraphUncookedOnly/Private/Graph/AnimNextAnimationGraph_EditorData.cpp
+++ b/Engine/Plugins/Experimental/AnimNextAnimGraph/Source/AnimNextAnimGraphUncookedOnly/Private/Graph/AnimNextAnimationGraph_EditorData.cpp
@@ -14,7 +14,7 @@
 #include "Curves/CurveFloat.h"
 #include "Module/AnimNextModule.h"
 #include "Entries/AnimNextAnimationGraphEntry.h"
-#include "Entries/AnimNextDataInterfaceEntry.h"
+#include "Entries/AnimNextSharedVariablesEntry.h"
 #include "Graph/AnimNextAnimationGraph.h"
 #include "Entries/AnimNextEventGraphEntry.h"
 #include "Entries/AnimNextVariableEntry.h"
@@ -568,10 +568,7 @@ void UAnimNextAnimationGraph_EditorData::OnPreCompileGetProgrammaticFunctionHead
 				if(FunctionHeader.IsValid())
 				{
 					FAnimNextProgrammaticFunctionHeader AnimNextFunctionHeader = {};
-					AnimNextFunctionHeader.Wrapped = FunctionHeader;
-					// @TODO: Determine if param / return variables are needed based on 'FAnimNextCallFunctionSharedData'
-					//AnimNextFunctionHeader.bGenerateParamVariables = true;
-					//AnimNextFunctionHeader.bGenerateReturnVariables = true;
+					AnimNextFunctionHeader.FunctionHeader = FunctionHeader;
 					OutCompileContext.GetMutableFunctionHeaders().Add(AnimNextFunctionHeader);
 				}
 			}
@@ -579,35 +576,6 @@ void UAnimNextAnimationGraph_EditorData::OnPreCompileGetProgrammaticFunctionHead
 	}
 }
 
-void UAnimNextAnimationGraph_EditorData::OnPreCompileGetProgrammaticVariables(const FRigVMCompileSettings& InSettings, FAnimNextGetVariableCompileContext& OutCompileContext)
-{
-	using namespace UE::AnimNext::UncookedOnly;
-
-	Super::OnPreCompileGetProgrammaticVariables(InSettings, OutCompileContext);
-
-	for (const FAnimNextProgrammaticFunctionHeader& ProgrammaticFunctionHeader : OutCompileContext.GetFunctionHeaders())
-	{
-		if (!ProgrammaticFunctionHeader.bGenerateParamVariables && !ProgrammaticFunctionHeader.bGenerateReturnVariables)
-		{
-			continue;
-		}
-
-		const FRigVMGraphFunctionHeader& FunctionHeader = ProgrammaticFunctionHeader.Wrapped;
-		for (const FRigVMGraphFunctionArgument& Argument : FunctionHeader.Arguments)
-		{
-			bool bAddParam = ProgrammaticFunctionHeader.bGenerateParamVariables && Argument.Direction == ERigVMPinDirection::Input;
-			bool bAddReturn = ProgrammaticFunctionHeader.bGenerateReturnVariables && Argument.Direction == ERigVMPinDirection::Output;
-
-			if (bAddParam || bAddReturn)
-			{
-				FRigVMGraphFunctionArgument InternallyNamedArgument = Argument;
-				InternallyNamedArgument.Name = FName(FUtils::MakeFunctionWrapperVariableName(FunctionHeader.Name, Argument.Name));
-				OutCompileContext.GetMutableProgrammaticVariables().Add(FAnimNextProgrammaticVariable::FromRigVMGraphFunctionArgument(InternallyNamedArgument));
-			}
-		}
-	}
-}
-
 void UAnimNextAnimationGraph_EditorData::OnPreCompileProcessGraphs(const FRigVMCompileSettings& InSettings, FAnimNextProcessGraphCompileContext& OutCompileContext)
 {
 	using namespace UE::AnimNext::UncookedOnly;
@@ -760,7 +728,7 @@ TConstArrayView<TSubclassOf<UAnimNextRigVMAssetEntry>> UAnimNextAnimationGraph_E
 	{
 		UAnimNextAnimationGraphEntry::StaticClass(),
 		UAnimNextVariableEntry::StaticClass(),
-		UAnimNextDataInterfaceEntry::StaticClass(),
+		UAnimNextSharedVariablesEntry::StaticClass(),
 	};
 
 	return Classes;
@@ -897,167 +865,6 @@ void UAnimNextAnimationGraph_EditorData::InitializeAssetUserData()
 	Super::InitializeAssetUserData();
 }
 
-void UAnimNextAnimationGraph_EditorData::OnPreCompileGetProgrammaticGraphs(const FRigVMCompileSettings& InSettings, FAnimNextGetGraphCompileContext& OutCompileContext)
-{
-	using namespace UE::AnimNext::UncookedOnly;
-
-	Super::OnPreCompileGetProgrammaticGraphs(InSettings, OutCompileContext);
-
-	if(OutCompileContext.GetFunctionHeaders().Num() > 0)
-	{
-		FRigVMClient* VMClient = GetRigVMClient();
-
-		// Create all shim events for our traits to call
-		const bool bSetupUndoRedo = false;
-		URigVMGraph* WrapperGraph = NewObject<URigVMGraph>(this, NAME_None, RF_Transient);
-		UAnimNextController* Controller = CastChecked<UAnimNextController>(VMClient->GetOrCreateController(WrapperGraph));
-		FRigVMControllerNotifGuard NotifGuard(Controller);
-		bool bAddedWrapperEvent = true;
-
-		for(const FAnimNextProgrammaticFunctionHeader& AnimNextFunctionHeader : OutCompileContext.GetFunctionHeaders())
-		{
-			const FRigVMGraphFunctionHeader& FunctionHeader = AnimNextFunctionHeader.Wrapped;
-
-			URigVMLibraryNode* LibraryNode = Cast<URigVMLibraryNode>(FunctionHeader.LibraryPointer.GetNodeSoftPath().TryLoad());
-			if(LibraryNode == nullptr)
-			{
-				InSettings.ReportError(FString::Printf(TEXT("Could not find function '%s'"), *FunctionHeader.Name.ToString()));
-				continue;
-			}
-
-			// Create user-defined entry point
-			FString WrapperEventName = FUtils::MakeFunctionWrapperEventName(FunctionHeader.Name);
-			URigVMUnitNode* EventNode = Controller->AddUnitNode(FRigVMFunction_UserDefinedEvent::StaticStruct(), TEXT("Execute"), FVector2D::ZeroVector, FunctionHeader.Name.ToString(), bSetupUndoRedo);
-			if(EventNode == nullptr)
-			{
-				InSettings.ReportError(FString::Printf(TEXT("Could not spawn event node for function '%s'"), *FunctionHeader.Name.ToString()));
-				continue;
-			}
-			URigVMPin* EventNamePin = EventNode->FindPin(GET_MEMBER_NAME_STRING_CHECKED(FRigVMFunction_UserDefinedEvent, EventName));
-			if(EventNamePin == nullptr)
-			{
-				InSettings.ReportError(FString::Printf(TEXT("Could not find custom event name pin")));
-				continue;
-			}
-			Controller->SetPinDefaultValue(EventNamePin->GetPinPath(), WrapperEventName, true, bSetupUndoRedo);
-
-			// Call function
-			URigVMFunctionReferenceNode* FunctionNode = Controller->AddFunctionReferenceNode(LibraryNode, FVector2D::ZeroVector, FunctionHeader.Name.ToString(), bSetupUndoRedo);
-			if(FunctionNode == nullptr)
-			{
-				InSettings.ReportError(FString::Printf(TEXT("Could not spawn function node for function '%s'"), *FunctionHeader.Name.ToString()));
-				continue;
-			}
-
-			// Link up Execute nodes if needed, function may be pure & lack an input pin
-			URigVMPin* CurrentExecuteOutputPin = EventNode->FindPin(FRigVMStruct::ExecuteContextName.ToString());
-			URigVMPin* ExecuteInputPin = FunctionNode->FindPin(FRigVMStruct::ExecuteContextName.ToString());
-			if (ExecuteInputPin && !Controller->AddLink(CurrentExecuteOutputPin, ExecuteInputPin, bSetupUndoRedo))
-			{
-				InSettings.ReportError(FString::Printf(TEXT("Could not link execute pins for function '%s'"), *FunctionHeader.Name.ToString()));
-				continue;
-			}
-
-			// Update current execute pin, RigVM doesn't have a concept of input / output execute pins, just one execute content pin used for both
-			CurrentExecuteOutputPin = ExecuteInputPin ? ExecuteInputPin : CurrentExecuteOutputPin;
-
-			// Generate & link internal variables if desired
-			if (AnimNextFunctionHeader.bGenerateParamVariables || AnimNextFunctionHeader.bGenerateReturnVariables)
-			{
-				// Controller needs to notify the AST of variable changes to make new links
-				bool bSuspendNotificationForInternalVariables = false;
-				FRigVMControllerNotifGuard VarNotifGuard(Controller, bSuspendNotificationForInternalVariables);
-
-				// Generate & link input arguments, also generate result variable node but link later
-				for (const FRigVMGraphFunctionArgument& Argument : FunctionHeader.Arguments)
-				{
-					// Execution context is captured as arg pins, skip those for internval variable gen
-					if (Argument.Direction == ERigVMPinDirection::IO)
-					{
-						continue;
-					}
-
-					bool bIsGetter = Argument.Direction == ERigVMPinDirection::Input;
-
-					if (bIsGetter && AnimNextFunctionHeader.bGenerateParamVariables)
-					{
-						FName InternalVariableName = FName(FUtils::MakeFunctionWrapperVariableName(FunctionHeader.Name, Argument.Name));
-						URigVMVariableNode* FunctionParamVariableNode = Controller->AddVariableNode(InternalVariableName
-							, Argument.CPPType.ToString()
-							, Argument.CPPTypeObject.Get()
-							, bIsGetter
-							, Argument.DefaultValue
-							, FVector2D::ZeroVector
-							, InternalVariableName.ToString()
-							, bSetupUndoRedo);
-
-						if (!FunctionParamVariableNode)
-						{
-							InSettings.ReportError(FString::Printf(TEXT("Failed to add internal variable node for param: %s, var: %s"), *FunctionHeader.Name.ToString(), *InternalVariableName.ToString()));
-							return;
-						}
-
-						// Link Param Pins
-						URigVMPin* ParamValuePin = FunctionParamVariableNode->GetValuePin();
-						URigVMPin* FunctionArgumentPin = FunctionNode->FindPin(Argument.Name.ToString());
-						if (!Controller->AddLink(ParamValuePin, FunctionArgumentPin, bSetupUndoRedo))
-						{
-							InSettings.ReportError(FString::Printf(TEXT("Failed to link internal variable param node to function: %s -> %s"), *GetNameSafe(ParamValuePin), *GetNameSafe(FunctionArgumentPin)));
-							return;
-						}
-					}
-
-					if (!bIsGetter && AnimNextFunctionHeader.bGenerateReturnVariables)
-					{
-						FName InternalResultName = FName(FUtils::MakeFunctionWrapperVariableName(FunctionHeader.Name, Argument.Name));
-						URigVMVariableNode* FunctionResultVariableNode = Controller->AddVariableNode(InternalResultName
-							, Argument.CPPType.ToString()
-							, Argument.CPPTypeObject.Get()
-							, bIsGetter
-							, Argument.DefaultValue
-							, FVector2D::ZeroVector
-							, InternalResultName.ToString()
-							, bSetupUndoRedo);
-
-						if (!FunctionResultVariableNode)
-						{
-							InSettings.ReportError(FString::Printf(TEXT("Failed to add internal variable node for result: %s, var: %s"), *FunctionHeader.Name.ToString(), *InternalResultName.ToString()));
-							return;
-						}
-
-						// Link Result pins
-						URigVMPin* FunctionResultPin = FunctionNode->FindPin(Argument.Name.ToString());
-						URigVMPin* ResultValuePin = FunctionResultVariableNode->GetValuePin();
-						if (!Controller->AddLink(FunctionResultPin, ResultValuePin, bSetupUndoRedo))
-						{
-							InSettings.ReportError(FString::Printf(TEXT("Failed to link internal variable result node to function: %s -> %s"), *GetNameSafe(FunctionResultPin), *GetNameSafe(ResultValuePin)));
-							return;
-						}
-
-						// Link Result Execute pins
-						URigVMPin* ResultExecuteInputPin = FunctionResultVariableNode->FindPin(FRigVMStruct::ExecuteContextName.ToString());
-						if (!Controller->AddLink(CurrentExecuteOutputPin, ResultExecuteInputPin, bSetupUndoRedo))
-						{
-							InSettings.ReportError(FString::Printf(TEXT("Failed to link execute pins for variable result node: %s -> %s"), *GetNameSafe(CurrentExecuteOutputPin), *GetNameSafe(ResultExecuteInputPin)));
-							return;
-						}
-
-						// Update current execute pin, RigVM doesn't have a concept of input / output execute pins, just one execute content pin used for both
-						CurrentExecuteOutputPin = ResultExecuteInputPin;
-					}
-				}
-			}
-
-			bAddedWrapperEvent = true;
-		}
-
-		if(bAddedWrapperEvent)
-		{
-			OutCompileContext.GetMutableProgrammaticGraphs().Add(WrapperGraph);
-		}
-	}
-}
-
 void UAnimNextAnimationGraph_EditorData::GetAnimNextAssetRegistryTags(FAssetRegistryTagsContext& Context, FAnimNextAssetRegistryExports& OutExports) const
 {
 	UE::AnimNext::UncookedOnly::FAnimGraphUtils::GetAssetManifestNodesRegistryExports(this, OutExports);
diff --git a/Engine/Plugins/Experimental/AnimNextAnimGraph/Source/AnimNextAnimGraphUncookedOnly/Private/Graph/RigVMTrait_AnimNextPublicVariablesUncookedOnly.cpp b/Engine/Plugins/Experimental/AnimNextAnimGraph/Source/AnimNextAnimGraphUncookedOnly/Private/Graph/RigVMTrait_AnimNextPublicVariablesUncookedOnly.cpp
--- a/Engine/Plugins/Experimental/AnimNextAnimGraph/Source/AnimNextAnimGraphUncookedOnly/Private/Graph/RigVMTrait_AnimNextPublicVariablesUncookedOnly.cpp
+++ b/Engine/Plugins/Experimental/AnimNextAnimGraph/Source/AnimNextAnimGraphUncookedOnly/Private/Graph/RigVMTrait_AnimNextPublicVariablesUncookedOnly.cpp
@@ -7,6 +7,7 @@
 #include "Graph/RigVMTrait_AnimNextPublicVariables.h"
 #include "RigVMModel/RigVMController.h"
 #include "RigVMModel/RigVMPin.h"
+#include "Variables/AnimNextSharedVariables.h"
 
 namespace UE::AnimNext::UncookedOnly
 {
@@ -38,41 +39,30 @@ void FPublicVariablesImpl::GetProgrammaticPins(const FRigVMTrait_AnimNextPublicV
 		return;
 	}
 
-	TArray<UAnimNextVariableEntry*> PublicVariables;
-	EditorData->GetPublicVariables(PublicVariables);
-	if(PublicVariables.Num() == 0)
+	TArray<UAnimNextRigVMAssetEditorData::FVariableInfo> PublicVariablesInfoArray;
+	EditorData->GetAllVariables(PublicVariablesInfoArray, UAnimNextRigVMAssetEditorData::EVariableRecursion::SelfOnly, UAnimNextRigVMAssetEditorData::EVariableAccessFilter::PublicOnly);
+	if(PublicVariablesInfoArray.Num() == 0)
 	{
 		return;
 	}
 
 	TArray<FPropertyBagPropertyDesc> Descs;
-	Descs.Reserve(PublicVariables.Num());
+	Descs.Reserve(PublicVariablesInfoArray.Num());
 	TArray<TConstArrayView<uint8>> Values;
-	Values.Reserve(PublicVariables.Num());
+	Values.Reserve(PublicVariablesInfoArray.Num());
 
 	// Maintain trait sorted order
 	for (FName VariableName : InTrait.InternalVariableNames)
 	{
-		auto FindMatchingName = [&](const UAnimNextVariableEntry* InVariable)
+		auto FindMatchingName = [&VariableName](const UAnimNextRigVMAssetEditorData::FVariableInfo& InVariableInfo)
 		{
-			if (!InVariable)
-			{
-				return false;
-			}
-
-			return InVariable->GetEntryName() == VariableName;
+			return InVariableInfo.Name == VariableName;
 		};
 
-		if (UAnimNextVariableEntry** PublicVariableItr = PublicVariables.FindByPredicate(FindMatchingName))
+		if (UAnimNextRigVMAssetEditorData::FVariableInfo* PublicVariablePtr = PublicVariablesInfoArray.FindByPredicate(FindMatchingName))
 		{
-			if (UAnimNextVariableEntry* PublicVariable = *PublicVariableItr)
-			{
-				const FProperty* Property;
-				TConstArrayView<uint8> Value;
-				PublicVariable->GetDefaultValue(Property, Value);
-				Descs.Emplace(PublicVariable->GetEntryName(), Property);
-				Values.Emplace(Value);
-			}
+			Descs.Emplace(PublicVariablePtr->Name, PublicVariablePtr->Property);
+			Values.Emplace(PublicVariablePtr->DefaultValue);
 		}
 
 	}
@@ -98,4 +88,4 @@ bool FPublicVariablesImpl::ShouldCreatePinForProperty(const FRigVMTrait_AnimNext
 		InTrait.InternalVariableNames.Contains(InProperty->GetFName());
 }
 
-}
\ No newline at end of file
+}
diff --git a/Engine/Plugins/Experimental/AnimNextChooser/Source/AnimNextChooser/Private/ChooserParameters.cpp b/Engine/Plugins/Experimental/AnimNextChooser/Source/AnimNextChooser/Private/ChooserParameters.cpp
--- a/Engine/Plugins/Experimental/AnimNextChooser/Source/AnimNextChooser/Private/ChooserParameters.cpp
+++ b/Engine/Plugins/Experimental/AnimNextChooser/Source/AnimNextChooser/Private/ChooserParameters.cpp
@@ -1,20 +1,21 @@
 // Copyright Epic Games, Inc. All Rights Reserved.
 
 #include "ChooserParameters.h"
-#include "DataInterface/AnimNextDataInterfaceInstance.h"
+#include "UAFAssetInstance.h"
 #include "StructUtils/PropertyBag.h"
+#include "UObject/FortniteMainBranchObjectVersion.h"
 
 namespace UE::AnimNext::Private
 {
 
 // Helper function to find the first data interface instance in the context 
-FAnimNextDataInterfaceInstance* GetFirstDataInterfaceInstance(FChooserEvaluationContext& Context)
+FUAFAssetInstance* GetFirstDataInterfaceInstance(FChooserEvaluationContext& Context)
 {
 	for(const FStructView& Param : Context.Params)
 	{
-		if(Param.GetScriptStruct() == FAnimNextDataInterfaceInstance::StaticStruct())
+		if(Param.GetScriptStruct() == FUAFAssetInstance::StaticStruct())
 		{
-			return Param.GetPtr<FAnimNextDataInterfaceInstance>();
+			return Param.GetPtr<FUAFAssetInstance>();
 		}
 	}
 
@@ -25,25 +26,25 @@ FAnimNextDataInterfaceInstance* GetFirstDataInterfaceInstance(FChooserEvaluation
 
 bool FBoolAnimProperty::GetValue(FChooserEvaluationContext& Context, bool& OutResult) const
 {
-	if(FAnimNextDataInterfaceInstance* Instance = UE::AnimNext::Private::GetFirstDataInterfaceInstance(Context))
+	if(FUAFAssetInstance* Instance = UE::AnimNext::Private::GetFirstDataInterfaceInstance(Context))
 	{
-		return Instance->GetVariable(VariableName, OutResult) == EPropertyBagResult::Success;
+		return Instance->GetVariable(Variable, OutResult) == EPropertyBagResult::Success;
 	}
 	return false;
 }
 
 bool FBoolAnimProperty::SetValue(FChooserEvaluationContext& Context, bool InValue) const
 {
-	if(FAnimNextDataInterfaceInstance* Instance = UE::AnimNext::Private::GetFirstDataInterfaceInstance(Context))
+	if(FUAFAssetInstance* Instance = UE::AnimNext::Private::GetFirstDataInterfaceInstance(Context))
 	{
-		return Instance->SetVariable(VariableName, InValue) == EPropertyBagResult::Success;
+		return Instance->SetVariable(Variable, InValue) == EPropertyBagResult::Success;
 	}
 	return false;
 }
 
 void FBoolAnimProperty::GetDisplayName(FText& OutName) const
 {
-	FString DisplayName = VariableName.ToString();
+	FString DisplayName = Variable.GetName().ToString();
 	int Index = -1;
 	DisplayName.FindLastChar(':', Index);
 	if (Index>=0)
@@ -54,28 +55,45 @@ void FBoolAnimProperty::GetDisplayName(FText& OutName) const
 	OutName = FText::FromString(DisplayName);
 }
 
+#if WITH_EDITORONLY_DATA
+bool FBoolAnimProperty::Serialize(FArchive& Ar)
+{
+	Ar.UsingCustomVersion(FFortniteMainBranchObjectVersion::GUID);
+	return false;
+}
+
+void FBoolAnimProperty::PostSerialize(const FArchive& Ar)
+{
+	if (Ar.IsLoading() && Ar.CustomVer(FFortniteMainBranchObjectVersion::GUID) < FFortniteMainBranchObjectVersion::AnimNextVariableReferences)
+	{
+		PRAGMA_DISABLE_DEPRECATION_WARNINGS
+		Variable = FAnimNextVariableReference(VariableName_DEPRECATED);
+		PRAGMA_ENABLE_DEPRECATION_WARNINGS
+	}
+}
+#endif
 
 bool FFloatAnimProperty::GetValue(FChooserEvaluationContext& Context, double& OutResult) const
 {
-	if(FAnimNextDataInterfaceInstance* Instance = UE::AnimNext::Private::GetFirstDataInterfaceInstance(Context))
+	if(FUAFAssetInstance* Instance = UE::AnimNext::Private::GetFirstDataInterfaceInstance(Context))
 	{
-		return Instance->GetVariable(VariableName, OutResult) == EPropertyBagResult::Success;
+		return Instance->GetVariable(Variable, OutResult) == EPropertyBagResult::Success;
 	}
 	return false;
 }
 
 bool FFloatAnimProperty::SetValue(FChooserEvaluationContext& Context, double InValue) const
 {
-	if(FAnimNextDataInterfaceInstance* Instance = UE::AnimNext::Private::GetFirstDataInterfaceInstance(Context))
+	if(FUAFAssetInstance* Instance = UE::AnimNext::Private::GetFirstDataInterfaceInstance(Context))
  	{
- 		return Instance->SetVariable(VariableName, InValue) == EPropertyBagResult::Success;
+ 		return Instance->SetVariable(Variable, InValue) == EPropertyBagResult::Success;
  	}
  	return false;
 }
 
 void FFloatAnimProperty::GetDisplayName(FText& OutName) const
 {
-	FString DisplayName = VariableName.ToString();
+	FString DisplayName = Variable.GetName().ToString();
 	int Index = -1;
 	DisplayName.FindLastChar(':', Index);
 	if (Index>=0)
@@ -86,27 +104,45 @@ void FFloatAnimProperty::GetDisplayName(FText& OutName) const
 	OutName = FText::FromString(DisplayName);
 }
 
+#if WITH_EDITORONLY_DATA
+bool FFloatAnimProperty::Serialize(FArchive& Ar)
+{
+	Ar.UsingCustomVersion(FFortniteMainBranchObjectVersion::GUID);
+	return false;
+}
+
+void FFloatAnimProperty::PostSerialize(const FArchive& Ar)
+{
+	if (Ar.IsLoading() && Ar.CustomVer(FFortniteMainBranchObjectVersion::GUID) < FFortniteMainBranchObjectVersion::AnimNextVariableReferences)
+	{
+		PRAGMA_DISABLE_DEPRECATION_WARNINGS
+		Variable = FAnimNextVariableReference(VariableName_DEPRECATED);
+		PRAGMA_ENABLE_DEPRECATION_WARNINGS
+	}
+}
+#endif
+
 bool FEnumAnimProperty::GetValue(FChooserEvaluationContext& Context, uint8& OutResult) const
 {
-	if(FAnimNextDataInterfaceInstance* Instance = UE::AnimNext::Private::GetFirstDataInterfaceInstance(Context))
+	if(FUAFAssetInstance* Instance = UE::AnimNext::Private::GetFirstDataInterfaceInstance(Context))
 	{
-		return Instance->GetVariable(VariableName, OutResult) == EPropertyBagResult::Success;
+		return Instance->GetVariable(Variable, OutResult) == EPropertyBagResult::Success;
 	}
 	return false;
 }
 
 bool FEnumAnimProperty::SetValue(FChooserEvaluationContext& Context, uint8 InValue) const
 {
-	if(FAnimNextDataInterfaceInstance* Instance = UE::AnimNext::Private::GetFirstDataInterfaceInstance(Context))
+	if(FUAFAssetInstance* Instance = UE::AnimNext::Private::GetFirstDataInterfaceInstance(Context))
 	{
-		return Instance->SetVariable(VariableName, InValue) == EPropertyBagResult::Success;
+		return Instance->SetVariable(Variable, InValue) == EPropertyBagResult::Success;
 	}
 	return false;
 }
 
 void FEnumAnimProperty::GetDisplayName(FText& OutName) const
 {
-	FString DisplayName = VariableName.ToString();
+	FString DisplayName = Variable.GetName().ToString();
 	int Index = -1;
 	DisplayName.FindLastChar(':', Index);
 	if (Index>=0)
@@ -116,3 +152,21 @@ void FEnumAnimProperty::GetDisplayName(FText& OutName) const
 	
 	OutName = FText::FromString(DisplayName);
 }
+
+#if WITH_EDITORONLY_DATA
+bool FEnumAnimProperty::Serialize(FArchive& Ar)
+{
+	Ar.UsingCustomVersion(FFortniteMainBranchObjectVersion::GUID);
+	return false;
+}
+
+void FEnumAnimProperty::PostSerialize(const FArchive& Ar)
+{
+	if (Ar.IsLoading() && Ar.CustomVer(FFortniteMainBranchObjectVersion::GUID) < FFortniteMainBranchObjectVersion::AnimNextVariableReferences)
+	{
+		PRAGMA_DISABLE_DEPRECATION_WARNINGS
+		Variable = FAnimNextVariableReference(VariableName_DEPRECATED);
+		PRAGMA_ENABLE_DEPRECATION_WARNINGS
+	}
+}
+#endif
\ No newline at end of file
diff --git a/Engine/Plugins/Experimental/AnimNextChooser/Source/AnimNextChooser/Private/RigUnit_EvaluateChooser.cpp b/Engine/Plugins/Experimental/AnimNextChooser/Source/AnimNextChooser/Private/RigUnit_EvaluateChooser.cpp
--- a/Engine/Plugins/Experimental/AnimNextChooser/Source/AnimNextChooser/Private/RigUnit_EvaluateChooser.cpp
+++ b/Engine/Plugins/Experimental/AnimNextChooser/Source/AnimNextChooser/Private/RigUnit_EvaluateChooser.cpp
@@ -2,7 +2,7 @@
 
 
 #include "RigUnit_EvaluateChooser.h"
-#include "DataInterface/AnimNextDataInterfaceInstance.h"
+#include "UAFAssetInstance.h"
 #include "ControlRig.h"
 
 #include UE_INLINE_GENERATED_CPP_BY_NAME(RigUnit_EvaluateChooser)
@@ -43,8 +43,8 @@ FRigUnit_EvaluateChooser_AnimNext_Execute()
 {
 	Result = nullptr;
 
-	const FAnimNextDataInterfaceInstance& Instance = ExecuteContext.GetContextData<FAnimNextModuleContextData>().GetDataInterfaceInstance();
-	RunChooserHelper({ ContextObject }, FStructView::Make(const_cast<FAnimNextDataInterfaceInstance&>(Instance)), Chooser, Result);
+	const FUAFAssetInstance& Instance = ExecuteContext.GetContextData<FAnimNextModuleContextData>().GetInstance();
+	RunChooserHelper({ ContextObject }, FStructView::Make(const_cast<FUAFAssetInstance&>(Instance)), Chooser, Result);
 }
 
 const FName FRigVMDispatch_EvaluateChooser::ResultName = TEXT("Result");
@@ -101,8 +101,8 @@ void FRigVMDispatch_EvaluateChooser::Execute(FRigVMExtendedExecuteContext& InCon
 	if (InContext.GetContextPublicDataStruct()->IsChildOf(FAnimNextExecuteContext::StaticStruct()))
 	{
 		FAnimNextExecuteContext& VMExecuteContext = InContext.GetPublicData<FAnimNextExecuteContext>();
-		const FAnimNextDataInterfaceInstance& Instance = VMExecuteContext.GetContextData<FAnimNextModuleContextData>().GetDataInterfaceInstance();
-		ChooserContext.AddStructViewParam(	FStructView::Make(const_cast<FAnimNextDataInterfaceInstance&>(Instance)));
+		const FUAFAssetInstance& Instance = VMExecuteContext.GetContextData<FAnimNextModuleContextData>().GetInstance();
+		ChooserContext.AddStructViewParam(	FStructView::Make(const_cast<FUAFAssetInstance&>(Instance)));
 	}
 	else
 	{
diff --git a/Engine/Plugins/Experimental/AnimNextControlRig/Source/AnimNextControlRig/Private/ControlRigTrait.cpp b/Engine/Plugins/Experimental/AnimNextControlRig/Source/AnimNextControlRig/Private/ControlRigTrait.cpp
--- a/Engine/Plugins/Experimental/AnimNextControlRig/Source/AnimNextControlRig/Private/ControlRigTrait.cpp
+++ b/Engine/Plugins/Experimental/AnimNextControlRig/Source/AnimNextControlRig/Private/ControlRigTrait.cpp
@@ -6,7 +6,6 @@
 #include "Graph/AnimNextGraphInstance.h"
 #include "AnimNextAnimGraphSettings.h"
 #include "Animation/AnimSequence.h"
-#include "DataInterface/DataInterfaceStructAdapter.h"
 #include "TraitInterfaces/IUpdate.h"
 #include "TraitInterfaces/IEvaluate.h"
 #include "TraitInterfaces/ITimeline.h"
diff --git a/Engine/Plugins/Experimental/AnimNextPoseSearch/Source/AnimNextPoseSearch/Private/HistoryCollectorTrait.cpp b/Engine/Plugins/Experimental/AnimNextPoseSearch/Source/AnimNextPoseSearch/Private/HistoryCollectorTrait.cpp
--- a/Engine/Plugins/Experimental/AnimNextPoseSearch/Source/AnimNextPoseSearch/Private/HistoryCollectorTrait.cpp
+++ b/Engine/Plugins/Experimental/AnimNextPoseSearch/Source/AnimNextPoseSearch/Private/HistoryCollectorTrait.cpp
@@ -247,7 +247,9 @@ namespace UE::AnimNext
 			FAnimNextModuleInstance* ModuleInstance = Binding.GetTraitPtr().GetNodeInstance()->GetOwner().GetModuleInstance();
 			FPoseHistoryReference Reference;
 			Reference.PoseHistory = InstanceData->PoseHistoryPtr;
-			ModuleInstance->SetVariable(PoseHistoryReferenceVariableName, Reference);
+			PRAGMA_DISABLE_DEPRECATION_WARNINGS
+			ModuleInstance->SetVariable(FAnimNextVariableReference(PoseHistoryReferenceVariableName), Reference);
+			PRAGMA_ENABLE_DEPRECATION_WARNINGS
 		}
 		
 		check(InstanceData->PoseHistoryPtr.IsValid());
diff --git a/Engine/Plugins/Experimental/AnimNextPoseSearch/Source/AnimNextPoseSearch/Private/MotionMatchingTrait.cpp b/Engine/Plugins/Experimental/AnimNextPoseSearch/Source/AnimNextPoseSearch/Private/MotionMatchingTrait.cpp
--- a/Engine/Plugins/Experimental/AnimNextPoseSearch/Source/AnimNextPoseSearch/Private/MotionMatchingTrait.cpp
+++ b/Engine/Plugins/Experimental/AnimNextPoseSearch/Source/AnimNextPoseSearch/Private/MotionMatchingTrait.cpp
@@ -6,13 +6,13 @@
 #include "Animation/BlendSpace.h"
 #include "AnimNextAnimGraphSettings.h"
 #include "Component/AnimNextComponent.h"
-#include "DataInterface/DataInterfaceStructAdapter.h"
 #include "EvaluationVM/EvaluationVM.h"
 #include "Graph/AnimNextGraphInstance.h"
-#include "GraphInterfaces/AnimNextNativeDataInterface_AnimSequencePlayer.h"
-#include "GraphInterfaces/AnimNextNativeDataInterface_BlendSpacePlayer.h"
-#include "GraphInterfaces/AnimNextNativeDataInterface_SynchronizeUsingGroups.h"
+#include "Traits/SequencePlayerTraitData.h"
+#include "Traits/BlendSpacePlayerTraitData.h"
+#include "Traits/SynchronizeUsingGroupsTraitData.h"
 #include "IPoseHistory.h"
+#include "Factory/AnimGraphFactory.h"
 #include "Module/AnimNextModuleInstance.h"
 #include "Module/RigUnit_AnimNextModuleEvents.h"
 #include "PoseSearch/MultiAnimAsset.h"
@@ -271,8 +271,8 @@ namespace UE::AnimNext
 #endif // ENABLE_ANIM_DEBUG
 
 				FChooserEvaluationContext EvaluationContext(const_cast<UObject*>(AnimContext));
-				const FAnimNextDataInterfaceInstance& Instance = Binding.GetTraitPtr().GetNodeInstance()->GetOwner();
-                EvaluationContext.AddStructViewParam(	FStructView::Make(const_cast<FAnimNextDataInterfaceInstance&>(Instance)));
+				const FUAFAssetInstance& Instance = Binding.GetTraitPtr().GetNodeInstance()->GetOwner();
+                EvaluationContext.AddStructViewParam(	FStructView::Make(const_cast<FUAFAssetInstance&>(Instance)));
                 
 				UPoseSearchLibrary::UpdateMotionMatchingState(
 					&EvaluationContext,
@@ -315,72 +315,67 @@ namespace UE::AnimNext
 					if (ActiveGraphRequest && (ActiveGraphRequest->FactoryObject == AnimationAsset))
 					{
 						check(ActiveGraphRequest->AnimationGraph);
-						for (FStructView& Payload : ActiveGraphRequest->GraphPayload.Get())
+						FAnimNextGraphInstance* ActiveGraphInstance = nullptr;
+						BlendStackBinding.GetActiveGraphInstance(Context, ActiveGraphInstance);
+						check(ActiveGraphInstance);
+
+						ActiveGraphInstance->AccessVariablesStruct<FBlendSpacePlayerData>([&SharedData, &Binding, &bSkipBlend, &MotionMatchingState, DesiredPlayRate](FBlendSpacePlayerData& InBlendSpacePlayer)
 						{
-							if (FAnimNextNativeDataInterface_BlendSpacePlayer* BlendSpacePayload = Payload.GetPtr<FAnimNextNativeDataInterface_BlendSpacePlayer>())
+							const float MaxDeltaAssetTimeToTrySkipBlendsForBlendSpaces = SharedData->GetMaxDeltaAssetTimeToTrySkipBlendsForBlendSpaces(Binding);
+							if (FMath::Abs(InBlendSpacePlayer.StartPosition - MotionMatchingState.SearchResult.SelectedTime) < MaxDeltaAssetTimeToTrySkipBlendsForBlendSpaces)
 							{
-								const float MaxDeltaAssetTimeToTrySkipBlendsForBlendSpaces = SharedData->GetMaxDeltaAssetTimeToTrySkipBlendsForBlendSpaces(Binding);
-								if (FMath::Abs(BlendSpacePayload->StartPosition - MotionMatchingState.SearchResult.SelectedTime) < MaxDeltaAssetTimeToTrySkipBlendsForBlendSpaces)
-								{
-									// If we're still on the same blend space, and we're updating blend space inputs dynamically, then DON'T BLEND.
-									bSkipBlend = true;
+								// If we're still on the same blend space, and we're updating blend space inputs dynamically, then DON'T BLEND.
+								bSkipBlend = true;
 
-									BlendSpacePayload->PlayRate = DesiredPlayRate;
+								InBlendSpacePlayer.PlayRate = DesiredPlayRate;
 
-									BlendSpacePayload->XAxisSamplePoint = MotionMatchingState.SearchResult.BlendParameters.X;
-									BlendSpacePayload->YAxisSamplePoint = MotionMatchingState.SearchResult.BlendParameters.Y;
-								}
-								break;
+								InBlendSpacePlayer.XAxisSamplePoint = MotionMatchingState.SearchResult.BlendParameters.X;
+								InBlendSpacePlayer.YAxisSamplePoint = MotionMatchingState.SearchResult.BlendParameters.Y;
 							}
-						}
+						});
 					}
 				}
 
 				if (!bSkipBlend)
 				{
-					UE::AnimNext::IBlendStack::FGraphRequest NewGraphRequest;
-					NewGraphRequest.BlendArgs = SharedData->BlendArguments;
-
-					// Create a default 'native data interface' payload via which we expect to communicate with the graph 
-					FAnimNextDataInterfacePayload GraphPayload;
-					FInstancedStruct AssetPlayerPayload;
-					// @TODO: Support asset->playload mapping in a more data-driven way
+					FAnimNextFactoryParams FactoryParams = FAnimGraphFactory::GetDefaultParamsForClass(AnimationAsset->GetClass());
 					if (bIsBlendSpace)
 					{
+						FSynchronizeUsingGroupsData SynchronizeUsingGroups;
 						// Tell blendspace samples to synchronize among themselves, but not with anything else.
-						FInstancedStruct SyncGroupPayload;
-						auto& NativeInterface = SyncGroupPayload.InitializeAs<FAnimNextNativeDataInterface_SynchronizeUsingGroups>();
-						NativeInterface.GroupName = NAME_None;
-						NativeInterface.GroupRole = EAnimGroupSynchronizationRole::AlwaysFollower;
-						NativeInterface.SyncMode = SharedData->GetSyncMode(Binding);
-						NativeInterface.MatchSyncPoint = true;
-						GraphPayload.AddNative(MoveTemp(SyncGroupPayload));
-
-						FAnimNextNativeDataInterface_BlendSpacePlayer& NewNativeInterface = AssetPlayerPayload.InitializeAs<FAnimNextNativeDataInterface_BlendSpacePlayer>();
-						NewNativeInterface.BlendSpace = Cast<UBlendSpace>(AnimationAsset);
-
-						NewNativeInterface.XAxisSamplePoint = MotionMatchingState.SearchResult.BlendParameters.X;
-						NewNativeInterface.YAxisSamplePoint = MotionMatchingState.SearchResult.BlendParameters.Y;
-
-						NewNativeInterface.PlayRate = DesiredPlayRate;
-						NewNativeInterface.StartPosition = MotionMatchingState.SearchResult.SelectedTime;
-						NewNativeInterface.Loop = MotionMatchingState.SearchResult.bLoop;
+						SynchronizeUsingGroups.GroupName = NAME_None;
+						SynchronizeUsingGroups.GroupRole = EAnimGroupSynchronizationRole::AlwaysFollower;
+						SynchronizeUsingGroups.SyncMode = SharedData->GetSyncMode(Binding);
+						SynchronizeUsingGroups.bMatchSyncPoint = true;
+						FactoryParams.PushPublicTrait(SynchronizeUsingGroups);
+
+						FBlendSpacePlayerData BlendSpacePlayer;
+						BlendSpacePlayer.BlendSpace = Cast<UBlendSpace>(AnimationAsset);
+
+						BlendSpacePlayer.XAxisSamplePoint = MotionMatchingState.SearchResult.BlendParameters.X;
+						BlendSpacePlayer.YAxisSamplePoint = MotionMatchingState.SearchResult.BlendParameters.Y;
+
+						BlendSpacePlayer.PlayRate = DesiredPlayRate;
+						BlendSpacePlayer.StartPosition = MotionMatchingState.SearchResult.SelectedTime;
+						BlendSpacePlayer.bLoop = MotionMatchingState.SearchResult.bLoop;
+						FactoryParams.PushPublicTrait(BlendSpacePlayer);
 					}
 					else
 					{
-						FAnimNextNativeDataInterface_AnimSequencePlayer& NewNativeInterface = AssetPlayerPayload.InitializeAs<FAnimNextNativeDataInterface_AnimSequencePlayer>();
-						NewNativeInterface.AnimSequence = Cast<UAnimSequence>(AnimationAsset);
-						NewNativeInterface.PlayRate = DesiredPlayRate;
-						NewNativeInterface.StartPosition = MotionMatchingState.SearchResult.SelectedTime;
-						NewNativeInterface.Loop = MotionMatchingState.SearchResult.bLoop;
+						FSequencePlayerData AnimSequencePlayer;
+						AnimSequencePlayer.AnimSequence = Cast<UAnimSequence>(AnimationAsset);
+						AnimSequencePlayer.PlayRate = DesiredPlayRate;
+						AnimSequencePlayer.StartPosition = MotionMatchingState.SearchResult.SelectedTime;
+						AnimSequencePlayer.bLoop = MotionMatchingState.SearchResult.bLoop;
+						FactoryParams.PushPublicTrait(AnimSequencePlayer);
 					}
 
-					GraphPayload.AddNative(MoveTemp(AssetPlayerPayload));
-
-					const UAnimNextAnimationGraph* AnimationGraph = IGraphFactory::GetGraphFromObjectWithFallback(Context, Binding, AnimationAsset, GraphPayload);
+					const UAnimNextAnimationGraph* AnimationGraph = IGraphFactory::GetOrBuildGraph(Context, Binding, AnimationAsset, FactoryParams);
 					if (AnimationGraph != nullptr)
 					{
-						NewGraphRequest.GraphPayload = MoveTemp(GraphPayload);
+						IBlendStack::FGraphRequest NewGraphRequest;
+						NewGraphRequest.FactoryParams = MoveTemp(FactoryParams);
+						NewGraphRequest.BlendArgs = SharedData->BlendArguments;
 						NewGraphRequest.AnimationGraph = AnimationGraph;
 						NewGraphRequest.FactoryObject = AnimationAsset;
 						BlendStackBinding.PushGraph(Context, MoveTemp(NewGraphRequest));
@@ -389,35 +384,28 @@ namespace UE::AnimNext
 			}
 		}
 
-		// @TODO: This should be automatic and not manually piped per asset type.
 		// Sync latest motion matching state to current graph
-		IBlendStack::FGraphRequestPtr ActiveGraphRequest;
-		BlendStackBinding.GetActiveGraph(Context, ActiveGraphRequest);
+		FAnimNextGraphInstance* ActiveGraphInstance = nullptr;
+		BlendStackBinding.GetActiveGraphInstance(Context, ActiveGraphInstance);
 
 		// Only override current graph's inputs if we have a valid search result.
-		if (ActiveGraphRequest && MotionMatchingState.SearchResult.SelectedAnim)
+		if (ActiveGraphInstance && MotionMatchingState.SearchResult.SelectedAnim)
 		{
-			check(ActiveGraphRequest->AnimationGraph);
-			for (FStructView& Payload : ActiveGraphRequest->GraphPayload.Get())
+			ActiveGraphInstance->AccessVariablesStruct<FSequencePlayerData>([DesiredPlayRate](FSequencePlayerData& InSequencePlayer)
 			{
-				if (FAnimNextNativeDataInterface_AnimSequencePlayer* PlayAnimPayload = Payload.GetPtr<FAnimNextNativeDataInterface_AnimSequencePlayer>())
-				{
-					PlayAnimPayload->PlayRate = DesiredPlayRate;
-					break;
-				}
-				else if (FAnimNextNativeDataInterface_BlendSpacePlayer* BlendSpacePayload = Payload.GetPtr<FAnimNextNativeDataInterface_BlendSpacePlayer>())
+				InSequencePlayer.PlayRate = DesiredPlayRate;
+			});
+
+			ActiveGraphInstance->AccessVariablesStruct<FBlendSpacePlayerData>([&SharedData, &Binding, DesiredPlayRate](FBlendSpacePlayerData& InBlendSpacePlayer)
+			{
+				InBlendSpacePlayer.PlayRate = DesiredPlayRate;
+				// always updating the blend paramters if bUpdateBlendSpaceInputs
+				if (SharedData->GetbUpdateBlendSpaceInputs(Binding))
 				{
-					BlendSpacePayload->PlayRate = DesiredPlayRate;
-					
-					// always updating the blend paramters if bUpdateBlendSpaceInputs
-					if (SharedData->GetbUpdateBlendSpaceInputs(Binding))
-					{
-						BlendSpacePayload->XAxisSamplePoint = SharedData->GetXAxisSamplePoint(Binding);
-						BlendSpacePayload->YAxisSamplePoint = SharedData->GetYAxisSamplePoint(Binding);
-					}
-					break;
+					InBlendSpacePlayer.XAxisSamplePoint = SharedData->GetXAxisSamplePoint(Binding);
+					InBlendSpacePlayer.YAxisSamplePoint = SharedData->GetYAxisSamplePoint(Binding);
 				}
-			}
+			});
 		}
 
 		PublishResults(Binding);
@@ -435,9 +423,11 @@ namespace UE::AnimNext
 
 					CalculateFullAlignedTransformsAtTime(MotionMatchingState.SearchResult, MultiAnimAsset->GetPlayLength(FVector::ZeroVector), SharedData->GetbWarpUsingRootBone(Binding), FullAlignedTransforms);
 					
-					FAnimNextDataInterfaceInstance& DataInterfaceInstance = Binding.GetTraitPtr().GetNodeInstance()->GetOwner();
+					FUAFAssetInstance& Instance = Binding.GetTraitPtr().GetNodeInstance()->GetOwner();
 					FName AlignmentTransformName = SharedData->GetAlignmentTransformVariableName(Binding);
-					DataInterfaceInstance.SetVariable(AlignmentTransformName, FullAlignedTransforms[CurrentResultRoleIndex]);
+					PRAGMA_DISABLE_DEPRECATION_WARNINGS
+					Instance.SetVariable(FAnimNextVariableReference(AlignmentTransformName), FullAlignedTransforms[CurrentResultRoleIndex]);
+					PRAGMA_ENABLE_DEPRECATION_WARNINGS
 				}
 			}
 		}
@@ -520,7 +510,9 @@ namespace UE::AnimNext
 									SearchResult.Role = Role;
 									SearchResult.ActorRootTransforms = ActorRootTransforms;
 									SearchResult.ActorRootBoneTransforms = ActorRootBoneTransforms;
-									OtherModuleInstance->SetVariable(ResultOutputVariableName, SearchResult);
+									PRAGMA_DISABLE_DEPRECATION_WARNINGS
+									OtherModuleInstance->SetVariable(FAnimNextVariableReference(ResultOutputVariableName), SearchResult);
+									PRAGMA_ENABLE_DEPRECATION_WARNINGS
 								}
 							},
 							UE::AnimNext::ETaskRunLocation::Before);
@@ -529,7 +521,9 @@ namespace UE::AnimNext
 					{
 						FPoseSearchBlueprintResult SearchResultCopy = SearchResult;
 						SearchResultCopy.AnimContexts.Reset();
-						ModuleInstance->SetVariable(ResultOutputVariableName, SearchResultCopy);
+						PRAGMA_DISABLE_DEPRECATION_WARNINGS
+						ModuleInstance->SetVariable(FAnimNextVariableReference(ResultOutputVariableName), SearchResultCopy);
+						PRAGMA_ENABLE_DEPRECATION_WARNINGS
 					}
 				}
 			};
diff --git a/Engine/Plugins/Experimental/AnimNextPoseSearch/Source/AnimNextPoseSearch/Private/PoseHistoryChooserParameter.cpp b/Engine/Plugins/Experimental/AnimNextPoseSearch/Source/AnimNextPoseSearch/Private/PoseHistoryChooserParameter.cpp
--- a/Engine/Plugins/Experimental/AnimNextPoseSearch/Source/AnimNextPoseSearch/Private/PoseHistoryChooserParameter.cpp
+++ b/Engine/Plugins/Experimental/AnimNextPoseSearch/Source/AnimNextPoseSearch/Private/PoseHistoryChooserParameter.cpp
@@ -1,20 +1,21 @@
 // Copyright Epic Games, Inc. All Rights Reserved.
 
 #include "PoseHistoryChooserParameter.h"
-#include "DataInterface/AnimNextDataInterfaceInstance.h"
+#include "UAFAssetInstance.h"
 #include "StructUtils/PropertyBag.h"
+#include "UObject/FortniteMainBranchObjectVersion.h"
 
 namespace 
 {
 
 // Helper function to find the first data interface instance in the context 
-FAnimNextDataInterfaceInstance* GetFirstAnimNextDataInterfaceInstance(FChooserEvaluationContext& Context)
+FUAFAssetInstance* GetFirstAnimNextDataInterfaceInstance(FChooserEvaluationContext& Context)
 {
 	for(const FStructView& Param : Context.Params)
 	{
-		if(Param.GetScriptStruct() == FAnimNextDataInterfaceInstance::StaticStruct())
+		if(Param.GetScriptStruct() == FUAFAssetInstance::StaticStruct())
 		{
-			return Param.GetPtr<FAnimNextDataInterfaceInstance>();
+			return Param.GetPtr<FUAFAssetInstance>();
 		}
 	}
 
@@ -25,9 +26,27 @@ FAnimNextDataInterfaceInstance* GetFirstAnimNextDataInterfaceInstance(FChooserEv
 
 bool FPoseHistoryAnimProperty::GetValue(FChooserEvaluationContext& Context, FPoseHistoryReference& OutResult) const
 {
-	if(FAnimNextDataInterfaceInstance* Instance = GetFirstAnimNextDataInterfaceInstance(Context))
+	if(FUAFAssetInstance* Instance = GetFirstAnimNextDataInterfaceInstance(Context))
 	{
-		return Instance->GetVariable(VariableName, OutResult) == EPropertyBagResult::Success;
+		return Instance->GetVariable(Variable, OutResult) == EPropertyBagResult::Success;
 	}
 	return false;
 }
+
+#if WITH_EDITORONLY_DATA
+bool FPoseHistoryAnimProperty::Serialize(FArchive& Ar)
+{
+	Ar.UsingCustomVersion(FFortniteMainBranchObjectVersion::GUID);
+	return false;
+}
+
+void FPoseHistoryAnimProperty::PostSerialize(const FArchive& Ar)
+{
+	if (Ar.IsLoading() && Ar.CustomVer(FFortniteMainBranchObjectVersion::GUID) < FFortniteMainBranchObjectVersion::AnimNextVariableReferences)
+	{
+		PRAGMA_DISABLE_DEPRECATION_WARNINGS
+		Variable = FAnimNextVariableReference(VariableName_DEPRECATED);
+		PRAGMA_ENABLE_DEPRECATION_WARNINGS
+	}
+}
+#endif
\ No newline at end of file
diff --git a/Engine/Plugins/Experimental/AnimNextPoseSearch/Source/AnimNextPoseSearch/Private/PoseHistoryChooserParameter.h b/Engine/Plugins/Experimental/AnimNextPoseSearch/Source/AnimNextPoseSearch/Private/PoseHistoryChooserParameter.h
--- a/Engine/Plugins/Experimental/AnimNextPoseSearch/Source/AnimNextPoseSearch/Private/PoseHistoryChooserParameter.h
+++ b/Engine/Plugins/Experimental/AnimNextPoseSearch/Source/AnimNextPoseSearch/Private/PoseHistoryChooserParameter.h
@@ -6,6 +6,7 @@
 #include "IChooserParameterFloat.h"
 #include "IChooserParameterEnum.h"
 #include "PoseSearch/Chooser/ChooserParameterPoseHistoryBase.h"
+#include "Variables/AnimNextVariableReference.h"
 #include "PoseHistoryChooserParameter.generated.h"
 
 USTRUCT(DisplayName = "Pose History Anim Param")
@@ -14,9 +15,19 @@ struct FPoseHistoryAnimProperty :  public FChooserParameterPoseHistoryBase
 	GENERATED_BODY()
 public:
 
+#if WITH_EDITORONLY_DATA
+	UPROPERTY()
+	FName VariableName_DEPRECATED;
+#endif
+
 	UPROPERTY(DisplayName = "Variable", EditAnywhere, Category="Variable")
-	FName VariableName;
+	FAnimNextVariableReference Variable;
 
 	bool GetValue(FChooserEvaluationContext& Context, FPoseHistoryReference& OutResult) const override;
-	virtual bool IsBound() const { return !VariableName.IsNone(); }
+	virtual bool IsBound() const { return !Variable.IsNone(); }
+
+#if WITH_EDITORONLY_DATA
+	bool Serialize(FArchive& Ar);
+	void PostSerialize(const FArchive& Ar);
+#endif	
 };
\ No newline at end of file
diff --git a/Engine/Plugins/Experimental/AnimNextPoseSearch/Source/AnimNextPoseSearch/Private/PoseSearchFeatureChannel_DistanceFromAnimNextVar.cpp b/Engine/Plugins/Experimental/AnimNextPoseSearch/Source/AnimNextPoseSearch/Private/PoseSearchFeatureChannel_DistanceFromAnimNextVar.cpp
--- a/Engine/Plugins/Experimental/AnimNextPoseSearch/Source/AnimNextPoseSearch/Private/PoseSearchFeatureChannel_DistanceFromAnimNextVar.cpp
+++ b/Engine/Plugins/Experimental/AnimNextPoseSearch/Source/AnimNextPoseSearch/Private/PoseSearchFeatureChannel_DistanceFromAnimNextVar.cpp
@@ -2,21 +2,22 @@
 
 #include "PoseSearchFeatureChannel_DistanceFromAnimNextVar.h"
 #include "Component/AnimNextComponent.h"
-#include "DataInterface/AnimNextDataInterfaceInstance.h"
-
+#include "UAFAssetInstance.h"
 #include "PoseSearch/PoseSearchContext.h"
 #include "StructUtils/PropertyBag.h"
+#include "UObject/FortniteMainBranchObjectVersion.h"
+#include "Variables/AnimNextVariableReference.h"
 
 
 namespace {
 // Helper function to find the first data interface instance in the context 
-FAnimNextDataInterfaceInstance* GetFirstDataInterfaceInstance(FChooserEvaluationContext& Context)
+FUAFAssetInstance* GetFirstAnimNextInstance(FChooserEvaluationContext& Context)
 {
 	for(const FStructView& Param : Context.Params)
 	{
-		if(Param.GetScriptStruct() == FAnimNextDataInterfaceInstance::StaticStruct())
+		if(Param.GetScriptStruct() == FUAFAssetInstance::StaticStruct())
 		{
-			return Param.GetPtr<FAnimNextDataInterfaceInstance>();
+			return Param.GetPtr<FUAFAssetInstance>();
 		}
 	}
 
@@ -32,10 +33,26 @@ void UPoseSearchFeatureChannel_DistanceFromAnimNextVar::BuildQuery(UE::PoseSearc
 	
 	float Distance = 0.0f;
 
-	if (FAnimNextDataInterfaceInstance* Instance = GetFirstDataInterfaceInstance(*Context))
+	if (FUAFAssetInstance* Instance = GetFirstAnimNextInstance(*Context))
 	{
-		Instance->GetVariable(DistanceVariableName, Distance);
+		Instance->GetVariable(DistanceVariable, Distance);
 	}
 	
 	FFeatureVectorHelper::EncodeFloat(SearchContext.EditFeatureVector(), ChannelDataOffset, Distance);
 }
+
+void UPoseSearchFeatureChannel_DistanceFromAnimNextVar::Serialize(FArchive& Ar)
+{
+	Super::Serialize(Ar);
+
+	Ar.UsingCustomVersion(FFortniteMainBranchObjectVersion::GUID);
+
+#if WITH_EDITORONLY_DATA
+	if (Ar.IsLoading() && Ar.CustomVer(FFortniteMainBranchObjectVersion::GUID) < FFortniteMainBranchObjectVersion::AnimNextVariableReferences)
+	{
+		PRAGMA_DISABLE_DEPRECATION_WARNINGS
+		DistanceVariable = FAnimNextVariableReference(DistanceVariableName_DEPRECATED);
+		PRAGMA_ENABLE_DEPRECATION_WARNINGS
+	}
+#endif
+}
\ No newline at end of file
diff --git a/Engine/Plugins/Experimental/AnimNextPoseSearch/Source/AnimNextPoseSearch/Private/PoseSearchFeatureChannel_DistanceFromAnimNextVar.h b/Engine/Plugins/Experimental/AnimNextPoseSearch/Source/AnimNextPoseSearch/Private/PoseSearchFeatureChannel_DistanceFromAnimNextVar.h
--- a/Engine/Plugins/Experimental/AnimNextPoseSearch/Source/AnimNextPoseSearch/Private/PoseSearchFeatureChannel_DistanceFromAnimNextVar.h
+++ b/Engine/Plugins/Experimental/AnimNextPoseSearch/Source/AnimNextPoseSearch/Private/PoseSearchFeatureChannel_DistanceFromAnimNextVar.h
@@ -3,15 +3,22 @@
 #pragma once
 
 #include "PoseSearch/PoseSearchFeatureChannel_Distance.h"
+#include "Variables/AnimNextVariableReference.h"
 #include "PoseSearchFeatureChannel_DistanceFromAnimNextVar.generated.h"
 
 UCLASS(Experimental, EditInlineNew, Blueprintable, meta = (DisplayName = "Distance Channel (AnimNext Variable)"), CollapseCategories)
 class UPoseSearchFeatureChannel_DistanceFromAnimNextVar : public UPoseSearchFeatureChannel_Distance
 {
 	GENERATED_BODY()
-	
+
+#if WITH_EDITORONLY_DATA
+	UPROPERTY()
+	FName DistanceVariableName_DEPRECATED;
+#endif
+
 	UPROPERTY(DisplayName = "Variable", EditAnywhere, Category="Variable")
-	FName DistanceVariableName;
+	FAnimNextVariableReference DistanceVariable;
 
 	virtual void BuildQuery(UE::PoseSearch::FSearchContext& SearchContext) const override;
+	virtual void Serialize(FArchive& Ar) override;
 };
