UAF: Fix crash loading levels using BlendSpace player traits

BlendSpace GC was not forwarding to its subgraphs, which are no longer kept referenced by external systems post-refactor.

diff --git a/Engine/Plugins/Experimental/AnimNextAnimGraph/Source/AnimNextAnimGraph/Private/Traits/BlendSpacePlayer.cpp b/Engine/Plugins/Experimental/AnimNextAnimGraph/Source/AnimNextAnimGraph/Private/Traits/BlendSpacePlayer.cpp
--- a/Engine/Plugins/Experimental/AnimNextAnimGraph/Source/AnimNextAnimGraph/Private/Traits/BlendSpacePlayer.cpp
+++ b/Engine/Plugins/Experimental/AnimNextAnimGraph/Source/AnimNextAnimGraph/Private/Traits/BlendSpacePlayer.cpp
@@ -317,5 +317,13 @@ namespace UE::AnimNext
 		FInstanceData* InstanceData = Binding.GetInstanceData<FInstanceData>();
 
 		Collector.AddReferencedObject(InstanceData->BlendSpace);
+
+		for (FBlendSpacePlayerTrait::FSampleGraphState& GraphState : InstanceData->SampleGraphs)
+		{
+			if (FAnimNextGraphInstance* ImplPtr = GraphState.Instance.Get())
+			{
+				Collector.AddPropertyReferencesWithStructARO(FAnimNextGraphInstance::StaticStruct(), ImplPtr);
+			}
+		}	
 	}
 }
