Anim Next Graph support for latent pins which evaluate only on become relevant

diff --git a/Engine/Plugins/Experimental/AnimNextAnimGraph/Source/AnimNextAnimGraph/Internal/Graph/AnimNextGraphLatentPropertiesContextData.h b/Engine/Plugins/Experimental/AnimNextAnimGraph/Source/AnimNextAnimGraph/Internal/Graph/AnimNextGraphLatentPropertiesContextData.h
--- a/Engine/Plugins/Experimental/AnimNextAnimGraph/Source/AnimNextAnimGraph/Internal/Graph/AnimNextGraphLatentPropertiesContextData.h
+++ b/Engine/Plugins/Experimental/AnimNextAnimGraph/Source/AnimNextAnimGraph/Internal/Graph/AnimNextGraphLatentPropertiesContextData.h
@@ -20,22 +20,25 @@ struct FAnimNextGraphLatentPropertiesContextData : public FAnimNextGraphContextD
 
 	FAnimNextGraphLatentPropertiesContextData() = default;
 
-	FAnimNextGraphLatentPropertiesContextData(FAnimNextModuleInstance* InModuleInstance, const FAnimNextGraphInstance* InInstance, const TConstArrayView<UE::AnimNext::FLatentPropertyHandle>& InLatentHandles, void* InDestinationBasePtr, bool bInIsFrozen)
+	FAnimNextGraphLatentPropertiesContextData(FAnimNextModuleInstance* InModuleInstance, const FAnimNextGraphInstance* InInstance, const TConstArrayView<UE::AnimNext::FLatentPropertyHandle>& InLatentHandles, void* InDestinationBasePtr, bool bInIsFrozen, bool bInJustBecameRelevant)
 		: FAnimNextGraphContextData(InModuleInstance, InInstance)
 		, LatentHandles(InLatentHandles)
 		, DestinationBasePtr(InDestinationBasePtr)
 		, bIsFrozen(bInIsFrozen)
+		, bJustBecameRelevant(bInJustBecameRelevant)
 	{
 	}
 
 	const TConstArrayView<UE::AnimNext::FLatentPropertyHandle>& GetLatentHandles() const { return LatentHandles; }
 	void* GetDestinationBasePtr() const { return DestinationBasePtr; }
 	bool IsFrozen() const { return bIsFrozen; }
+	bool JustBecameRelevant() const { return bJustBecameRelevant; }
 
 private:
 	TConstArrayView<UE::AnimNext::FLatentPropertyHandle> LatentHandles;
 	void* DestinationBasePtr = nullptr;
 	bool bIsFrozen = false;
+	bool bJustBecameRelevant = false;
 
 	friend struct FAnimNextGraphInstance;
 	friend struct FAnimNextExecuteContext;
diff --git a/Engine/Plugins/Experimental/AnimNextAnimGraph/Source/AnimNextAnimGraph/Private/Graph/AnimNextGraphInstance.cpp b/Engine/Plugins/Experimental/AnimNextAnimGraph/Source/AnimNextAnimGraph/Private/Graph/AnimNextGraphInstance.cpp
--- a/Engine/Plugins/Experimental/AnimNextAnimGraph/Source/AnimNextAnimGraph/Private/Graph/AnimNextGraphInstance.cpp
+++ b/Engine/Plugins/Experimental/AnimNextAnimGraph/Source/AnimNextAnimGraph/Private/Graph/AnimNextGraphInstance.cpp
@@ -124,7 +124,7 @@ void FAnimNextGraphInstance::MarkAsUpdated()
 	bHasUpdatedOnce = true;
 }
 
-void FAnimNextGraphInstance::ExecuteLatentPins(const TConstArrayView<UE::AnimNext::FLatentPropertyHandle>& LatentHandles, void* DestinationBasePtr, bool bIsFrozen)
+void FAnimNextGraphInstance::ExecuteLatentPins(const TConstArrayView<UE::AnimNext::FLatentPropertyHandle>& LatentHandles, void* DestinationBasePtr, bool bIsFrozen, bool bJustBecameRelevant)
 {
 	SCOPE_CYCLE_COUNTER(STAT_AnimNext_Graph_RigVM);
 
@@ -144,14 +144,14 @@ void FAnimNextGraphInstance::ExecuteLatentPins(const TConstArrayView<UE::AnimNex
 		}
 
 		// Insert our context data for the scope of execution
-		FAnimNextGraphLatentPropertiesContextData ContextData(ModuleInstance, this, LatentHandles, DestinationBasePtr, bIsFrozen);
+		FAnimNextGraphLatentPropertiesContextData ContextData(ModuleInstance, this, LatentHandles, DestinationBasePtr, bIsFrozen, bJustBecameRelevant);
 		UE::AnimNext::FScopedExecuteContextData ContextDataScope(AnimNextContext, ContextData);
 
 		VM->ExecuteVM(ExtendedExecuteContext, FRigUnit_AnimNextShimRoot::EventName);
 	}
 }
 
-void FAnimNextGraphInstance::CopyVariablesToLatentPins(const TConstArrayView<UE::AnimNext::FLatentPropertyHandle>& LatentHandles, void* DestinationBasePtr, bool bIsFrozen)
+void FAnimNextGraphInstance::CopyVariablesToLatentPins(const TConstArrayView<UE::AnimNext::FLatentPropertyHandle>& LatentHandles, void* DestinationBasePtr, bool bIsFrozen, bool bJustBecameRelevant)
 {
 	for (UE::AnimNext::FLatentPropertyHandle Handle : LatentHandles)
 	{
@@ -173,6 +173,12 @@ void FAnimNextGraphInstance::CopyVariablesToLatentPins(const TConstArrayView<UE:
 			continue;
 		}
 
+		if (!bJustBecameRelevant && Handle.OnBecomeRelevant())
+		{
+			// This handle should only update on become relevant
+			continue;
+		}
+
 		uint8* DestinationPtr = (uint8*)DestinationBasePtr + Handle.GetLatentPropertyOffset();
 		AccessVariablePropertyByIndex(Handle.GetLatentPropertyIndex(), [DestinationPtr](const FProperty* InProperty, TArrayView<uint8> InData)
 		{
diff --git a/Engine/Plugins/Experimental/AnimNextAnimGraph/Source/AnimNextAnimGraph/Private/Graph/RigUnit_AnimNextGraphEvaluator.cpp b/Engine/Plugins/Experimental/AnimNextAnimGraph/Source/AnimNextAnimGraph/Private/Graph/RigUnit_AnimNextGraphEvaluator.cpp
--- a/Engine/Plugins/Experimental/AnimNextAnimGraph/Source/AnimNextAnimGraph/Private/Graph/RigUnit_AnimNextGraphEvaluator.cpp
+++ b/Engine/Plugins/Experimental/AnimNextAnimGraph/Source/AnimNextAnimGraph/Private/Graph/RigUnit_AnimNextGraphEvaluator.cpp
@@ -34,6 +34,7 @@ void FRigUnit_AnimNextGraphEvaluator::StaticExecute(FRigVMExtendedExecuteContext
 	const TConstArrayView<UE::AnimNext::FLatentPropertyHandle>& LatentHandles = LatentPropertiesContextData.GetLatentHandles();
 	uint8* DestinationBasePtr = (uint8*)LatentPropertiesContextData.GetDestinationBasePtr();
 	const bool bIsFrozen = LatentPropertiesContextData.IsFrozen();
+	const bool bJustBecameRelevant = LatentPropertiesContextData.JustBecameRelevant();
 
 	for (UE::AnimNext::FLatentPropertyHandle Handle : LatentHandles)
 	{
@@ -49,7 +50,7 @@ void FRigUnit_AnimNextGraphEvaluator::StaticExecute(FRigVMExtendedExecuteContext
 			continue;
 		}
 
-		if (bIsFrozen && Handle.CanFreeze())
+		if ((bIsFrozen && Handle.CanFreeze()) || (!bJustBecameRelevant && Handle.OnBecomeRelevant()))
 		{
 			// This handle can freeze and we are frozen, no need to update it
 			continue;
diff --git a/Engine/Plugins/Experimental/AnimNextAnimGraph/Source/AnimNextAnimGraph/Private/TraitCore/Trait.cpp b/Engine/Plugins/Experimental/AnimNextAnimGraph/Source/AnimNextAnimGraph/Private/TraitCore/Trait.cpp
--- a/Engine/Plugins/Experimental/AnimNextAnimGraph/Source/AnimNextAnimGraph/Private/TraitCore/Trait.cpp
+++ b/Engine/Plugins/Experimental/AnimNextAnimGraph/Source/AnimNextAnimGraph/Private/TraitCore/Trait.cpp
@@ -189,6 +189,7 @@ namespace UE::AnimNext
 
 				// Always false for now, we don't support freezing yet
 				Metadata.bCanFreeze = false;
+				Metadata.bOnBecomeRelevant = Property->HasMetaData("OnBecomeRelevant");
 
 				OutLatentPropertyHandles.Add(Metadata);
 				NumHandlesAdded++;
diff --git a/Engine/Plugins/Experimental/AnimNextAnimGraph/Source/AnimNextAnimGraph/Private/TraitCore/TraitReader.cpp b/Engine/Plugins/Experimental/AnimNextAnimGraph/Source/AnimNextAnimGraph/Private/TraitCore/TraitReader.cpp
--- a/Engine/Plugins/Experimental/AnimNextAnimGraph/Source/AnimNextAnimGraph/Private/TraitCore/TraitReader.cpp
+++ b/Engine/Plugins/Experimental/AnimNextAnimGraph/Source/AnimNextAnimGraph/Private/TraitCore/TraitReader.cpp
@@ -150,6 +150,7 @@ namespace UE::AnimNext
 					LatentHeader.bHasValidLatentProperties = false;
 					LatentHeader.bCanAllPropertiesFreeze = true;
 					LatentHeader.bAreAllPropertiesVariableAccesses = true;
+					LatentHeader.bAllLatentPropertiesOnBecomeRelevant = true;
 				}
 
 				if (NumLatentProperties == 0)
@@ -160,6 +161,7 @@ namespace UE::AnimNext
 				bool bHasValidLatentProperties = !!LatentHeader.bHasValidLatentProperties;
 				bool bCanAllPropertiesFreeze = !!LatentHeader.bCanAllPropertiesFreeze;
 				bool bAreAllPropertiesVariableAccesses = !!LatentHeader.bAreAllPropertiesVariableAccesses;
+				bool bAllLatentPropertiesOnBecomeRelevant = !!LatentHeader.bAllLatentPropertiesOnBecomeRelevant;
 
 				const FTraitRegistryHandle TraitHandle = TraitTemplate.GetRegistryHandle();
 				const FTrait* Trait = TraitRegistry.Find(TraitHandle);
@@ -172,6 +174,7 @@ namespace UE::AnimNext
 					uint16 RigVMIndex = MAX_uint16;
 					uint32 CurrentLatentPropertyOffset = 0;
 					bool bCanFreeze = true;
+					bool bOnBecomeRelevant = true;
 
 					if (Trait != nullptr)
 					{
@@ -188,6 +191,8 @@ namespace UE::AnimNext
 							RigVMIndex = Metadata.RigVMIndex;
 							CurrentLatentPropertyOffset = LatentPropertyOffset;
 							bCanFreeze = Metadata.bCanFreeze;
+							bOnBecomeRelevant = Metadata.bOnBecomeRelevant;
+							bAllLatentPropertiesOnBecomeRelevant &= Metadata.bOnBecomeRelevant;
 
 							bHasValidLatentProperties = true;
 							bCanAllPropertiesFreeze &= bCanFreeze;
@@ -198,12 +203,13 @@ namespace UE::AnimNext
 						}
 					}
 
-					LatentHandles[LatentPropertyIndex] = FLatentPropertyHandle(RigVMIndex, CurrentLatentPropertyOffset, bCanFreeze);
+					LatentHandles[LatentPropertyIndex] = FLatentPropertyHandle(RigVMIndex, CurrentLatentPropertyOffset, bCanFreeze, bOnBecomeRelevant);
 				}
 
 				LatentHeader.bHasValidLatentProperties = bHasValidLatentProperties;
 				LatentHeader.bCanAllPropertiesFreeze = bCanAllPropertiesFreeze;
 				LatentHeader.bAreAllPropertiesVariableAccesses = bAreAllPropertiesVariableAccesses;
+				LatentHeader.bAllLatentPropertiesOnBecomeRelevant = bAllLatentPropertiesOnBecomeRelevant;
 			}
 
 			// Set our final node instance data size that factors in our latent properties
diff --git a/Engine/Plugins/Experimental/AnimNextAnimGraph/Source/AnimNextAnimGraph/Private/TraitCore/TraitStackBinding.cpp b/Engine/Plugins/Experimental/AnimNextAnimGraph/Source/AnimNextAnimGraph/Private/TraitCore/TraitStackBinding.cpp
--- a/Engine/Plugins/Experimental/AnimNextAnimGraph/Source/AnimNextAnimGraph/Private/TraitCore/TraitStackBinding.cpp
+++ b/Engine/Plugins/Experimental/AnimNextAnimGraph/Source/AnimNextAnimGraph/Private/TraitCore/TraitStackBinding.cpp
@@ -352,7 +352,7 @@ namespace UE::AnimNext
 		return false;
 	}
 
-	void FTraitStackBinding::SnapshotLatentProperties(bool bIsFrozen)
+	void FTraitStackBinding::SnapshotLatentProperties(bool bIsFrozen, bool bJustBecameRelevant)
 	{
 		if (!IsValid())
 		{
@@ -367,7 +367,7 @@ namespace UE::AnimNext
 		{
 			return;	// All latent properties are inline, nothing to snapshot
 		}
-		else if (bIsFrozen && LatentHeader.bCanAllPropertiesFreeze)
+		else if ((bIsFrozen && LatentHeader.bCanAllPropertiesFreeze) || (!bJustBecameRelevant && LatentHeader.bAllLatentPropertiesOnBecomeRelevant))
 		{
 			return;	// We are frozen and all latent properties support freezing, nothing to snapshot
 		}
@@ -379,11 +379,11 @@ namespace UE::AnimNext
 
 		if (LatentHeader.bAreAllPropertiesVariableAccesses)
 		{
-			GraphInstance.CopyVariablesToLatentPins(TConstArrayView<FLatentPropertyHandle>(LatentHandles, NumLatentHandles), NodeInstance, bIsFrozen);
+			GraphInstance.CopyVariablesToLatentPins(TConstArrayView<FLatentPropertyHandle>(LatentHandles, NumLatentHandles), NodeInstance, bIsFrozen, bJustBecameRelevant);
 		}
 		else
 		{
-			GraphInstance.ExecuteLatentPins(TConstArrayView<FLatentPropertyHandle>(LatentHandles, NumLatentHandles), NodeInstance, bIsFrozen);
+			GraphInstance.ExecuteLatentPins(TConstArrayView<FLatentPropertyHandle>(LatentHandles, NumLatentHandles), NodeInstance, bIsFrozen, bJustBecameRelevant);
 		}
 	}
 }
diff --git a/Engine/Plugins/Experimental/AnimNextAnimGraph/Source/AnimNextAnimGraph/Private/TraitInterfaces/IUpdate.cpp b/Engine/Plugins/Experimental/AnimNextAnimGraph/Source/AnimNextAnimGraph/Private/TraitInterfaces/IUpdate.cpp
--- a/Engine/Plugins/Experimental/AnimNextAnimGraph/Source/AnimNextAnimGraph/Private/TraitInterfaces/IUpdate.cpp
+++ b/Engine/Plugins/Experimental/AnimNextAnimGraph/Source/AnimNextAnimGraph/Private/TraitInterfaces/IUpdate.cpp
@@ -550,7 +550,8 @@ namespace UE::AnimNext
 				// When a trait stack is blending out, its properties are frozen by default
 				// unless a property opts to always update regardless.
 				const bool bIsFrozen = Entry->TraitState.IsBlendingOut();
-				Entry->TraitStack.SnapshotLatentProperties(bIsFrozen);
+				const bool bJustBecameRelevant = Entry->TraitState.IsNewlyRelevant();
+				Entry->TraitStack.SnapshotLatentProperties(bIsFrozen, bJustBecameRelevant);
 
 				const bool bImplementsIUpdate = Entry->TraitStack.GetInterface(Entry->UpdateTrait);
 
diff --git a/Engine/Plugins/Experimental/AnimNextAnimGraph/Source/AnimNextAnimGraph/Public/Graph/AnimNextGraphInstance.h b/Engine/Plugins/Experimental/AnimNextAnimGraph/Source/AnimNextAnimGraph/Public/Graph/AnimNextGraphInstance.h
--- a/Engine/Plugins/Experimental/AnimNextAnimGraph/Source/AnimNextAnimGraph/Public/Graph/AnimNextGraphInstance.h
+++ b/Engine/Plugins/Experimental/AnimNextAnimGraph/Source/AnimNextAnimGraph/Public/Graph/AnimNextGraphInstance.h
@@ -106,11 +106,11 @@ struct FAnimNextGraphInstance : public FUAFAssetInstance
 private:
 	// Executes a list of latent RigVM pins and writes the result into the destination pointer (latent handle offsets are using the destination as base)
 	// When frozen, latent handles that can freeze are skipped, all others will execute
-	UE_API void ExecuteLatentPins(const TConstArrayView<UE::AnimNext::FLatentPropertyHandle>& LatentHandles, void* DestinationBasePtr, bool bIsFrozen);
+	UE_API void ExecuteLatentPins(const TConstArrayView<UE::AnimNext::FLatentPropertyHandle>& LatentHandles, void* DestinationBasePtr, bool bIsFrozen, bool bJustBecameRelevant);
 
 	// When a group of latent pins do not require RigVM to operate, this copies the values directly from variables rather than calling ExecuteLatentPins.
 	// When frozen, latent handles that can freeze are skipped, all others will execute
-	void CopyVariablesToLatentPins(const TConstArrayView<UE::AnimNext::FLatentPropertyHandle>& LatentHandles, void* DestinationBasePtr, bool bIsFrozen);
+	void CopyVariablesToLatentPins(const TConstArrayView<UE::AnimNext::FLatentPropertyHandle>& LatentHandles, void* DestinationBasePtr, bool bIsFrozen, bool bJustBecameRelevant);
 
 #if WITH_EDITORONLY_DATA
 	// During graph compilation, if we have existing graph instances, we freeze them by releasing their memory before thawing them
diff --git a/Engine/Plugins/Experimental/AnimNextAnimGraph/Source/AnimNextAnimGraph/Public/TraitCore/LatentPropertyHandle.h b/Engine/Plugins/Experimental/AnimNextAnimGraph/Source/AnimNextAnimGraph/Public/TraitCore/LatentPropertyHandle.h
--- a/Engine/Plugins/Experimental/AnimNextAnimGraph/Source/AnimNextAnimGraph/Public/TraitCore/LatentPropertyHandle.h
+++ b/Engine/Plugins/Experimental/AnimNextAnimGraph/Source/AnimNextAnimGraph/Public/TraitCore/LatentPropertyHandle.h
@@ -28,7 +28,11 @@ namespace UE::AnimNext
 
 		// Whether or not the property is a simple variable copy
 		bool bUsesVariableCopy = false;
+		
+		// Whether or not the property evaluates only once on become relevant
+		bool bOnBecomeRelevant = false;
 	};
+	
 
 	// Serializes latent property metadata
 	inline FArchive& operator<<(FArchive& Ar, FLatentPropertyMetadata& Metadata)
@@ -37,6 +41,7 @@ namespace UE::AnimNext
 		Ar << Metadata.RigVMIndex;
 		Ar << Metadata.bCanFreeze;
 		Ar << Metadata.bUsesVariableCopy;
+		Ar << Metadata.bOnBecomeRelevant;
 		return Ar;
 	}
 
@@ -53,6 +58,9 @@ namespace UE::AnimNext
 
 		// Whether or not this base trait has all its latent properties support freezing
 		uint8 bCanAllPropertiesFreeze : 1 = false;
+		
+		// Whether or not this base trait has all its latent properties only evaluate on first frame
+		uint8 bAllLatentPropertiesOnBecomeRelevant : 1 = false;
 
 		// Whether the latent properties for this trait are all simple variable accesses
 		uint8 bAreAllPropertiesVariableAccesses : 1 = false;
@@ -98,11 +106,14 @@ namespace UE::AnimNext
 		// Returns whether or not the latent property supports freezing during snapshots
 		// Latent properties that do not support freezing always update in a snapshot
 		constexpr bool CanFreeze() const noexcept { return IsIndexValid() ? !!(PackedIndexCanFreezeValue & HANDLE_CAN_FREEZE_MASK) : false; }
+		
+		// Returns whether or not the latent property evaluates only on become relevant
+		constexpr bool OnBecomeRelevant() const noexcept { return IsIndexValid() ? !!(PackedIndexCanFreezeValue & HANDLE_ON_BECOME_RELEVANT_MASK) : false; }
 
 	private:
 		// Constructs a FLatentPropertyHandle instance
-		FLatentPropertyHandle(uint16 InRigRMIndex, uint32 InPropertyOffset, bool bInCanFreeze)
-			: PackedIndexCanFreezeValue((InRigRMIndex & HANDLE_INDEX_MASK) | (uint16(bInCanFreeze) << HANDLE_CAN_FREEZE_SHIFT_OFFSET))
+		FLatentPropertyHandle(uint16 InRigRMIndex, uint32 InPropertyOffset, bool bInCanFreeze, bool bInOnBecomeRelevant)
+			: PackedIndexCanFreezeValue((InRigRMIndex & HANDLE_INDEX_MASK) | (uint16(bInCanFreeze) << HANDLE_CAN_FREEZE_SHIFT_OFFSET) | (uint16(bInOnBecomeRelevant) << HANDLE_ON_BECOME_RELEVANT_SHIFT_OFFSET))
 			, OffsetValue(uint16(InPropertyOffset))
 		{
 			check(InRigRMIndex <= INVALID_HANDLE_VALUE || InRigRMIndex == MAX_uint16);	// Must fit on 15 bits when valid
@@ -110,16 +121,18 @@ namespace UE::AnimNext
 		}
 
 		// A mask for the RigVM index portion of the packed value
-		static constexpr uint16 HANDLE_INDEX_MASK = 0x7FFF;
+		static constexpr uint16 HANDLE_INDEX_MASK = 0x3FFF;
 
 		// A mask for the CanFreeze metadata flag in the packed value
 		static constexpr uint16 HANDLE_CAN_FREEZE_MASK = 0x8000;
+		static constexpr uint16 HANDLE_ON_BECOME_RELEVANT_MASK = 0x4000;
 
 		// A shift offset for the CanFreeze metadata flag in the packed value
 		static constexpr uint16 HANDLE_CAN_FREEZE_SHIFT_OFFSET = 15;
+		static constexpr uint16 HANDLE_ON_BECOME_RELEVANT_SHIFT_OFFSET = 14;
 
 		// Latent property handles are a 15 bit index and a 1 bit flag to denote if the property can freeze, ~0 is invalid
-		static constexpr uint16 INVALID_HANDLE_VALUE = 0x7FFF;
+		static constexpr uint16 INVALID_HANDLE_VALUE = 0x3FFF;
 
 		// Latent property offsets are relative to the start of the node instance data, 0 is invalid
 		static constexpr uint16 INVALID_OFFSET_VALUE = 0;
diff --git a/Engine/Plugins/Experimental/AnimNextAnimGraph/Source/AnimNextAnimGraph/Public/TraitCore/TraitStackBinding.h b/Engine/Plugins/Experimental/AnimNextAnimGraph/Source/AnimNextAnimGraph/Public/TraitCore/TraitStackBinding.h
--- a/Engine/Plugins/Experimental/AnimNextAnimGraph/Source/AnimNextAnimGraph/Public/TraitCore/TraitStackBinding.h
+++ b/Engine/Plugins/Experimental/AnimNextAnimGraph/Source/AnimNextAnimGraph/Public/TraitCore/TraitStackBinding.h
@@ -44,7 +44,7 @@ namespace UE::AnimNext
 		// Takes a snapshot of all latent properties on this trait stack.
 		// Properties can be marked as always updating or as supporting freezing (e.g. when a branch of the graph blends out).
 		// A freezable property does not update when a snapshot is taken of a frozen stack.
-		UE_API void SnapshotLatentProperties(bool bIsFrozen);
+		UE_API void SnapshotLatentProperties(bool bIsFrozen, bool bIsFirstUpdate);
 
 		// Returns a trait pointer to the base of the trait stack.
 		FWeakTraitPtr GetBaseTraitPtr() const { return FWeakTraitPtr(NodeInstance, BaseTraitIndex); }
diff --git a/Engine/Plugins/Experimental/AnimNextAnimGraph/Source/AnimNextAnimGraph/Public/Traits/SequencePlayerTraitData.h b/Engine/Plugins/Experimental/AnimNextAnimGraph/Source/AnimNextAnimGraph/Public/Traits/SequencePlayerTraitData.h
--- a/Engine/Plugins/Experimental/AnimNextAnimGraph/Source/AnimNextAnimGraph/Public/Traits/SequencePlayerTraitData.h
+++ b/Engine/Plugins/Experimental/AnimNextAnimGraph/Source/AnimNextAnimGraph/Public/Traits/SequencePlayerTraitData.h
@@ -1,4 +1,4 @@
-﻿// Copyright Epic Games, Inc. All Rights Reserved.
+// Copyright Epic Games, Inc. All Rights Reserved.
 
 #pragma once
 
@@ -13,19 +13,19 @@ struct FAnimNextSequencePlayerTraitSharedData : public FAnimNextTraitSharedData
 	GENERATED_BODY()
 
 	/** The sequence to play. */
-	UPROPERTY(EditAnywhere, Category = "Sequence Player", meta=(ExportAsReference="true", FactorySource))
+	UPROPERTY(EditAnywhere, Category = "Sequence Player", meta=(ExportAsReference="true", FactorySource, OnBecomeRelevant))
 	TObjectPtr<const UAnimSequence> AnimSequence;
 
 	/** The play rate multiplier at which this sequence plays. */
-	UPROPERTY(EditAnywhere, Category = "Sequence Player")
+	UPROPERTY(EditAnywhere, Category = "Sequence Player", meta = (OnBecomeRelevant))
 	float PlayRate = 1.0f;
 
 	/** The time at which we should start playing this sequence. */
-	UPROPERTY(EditAnywhere, Category = "Sequence Player")
+	UPROPERTY(EditAnywhere, Category = "Sequence Player", meta = (OnBecomeRelevant))
 	float StartPosition = 0.0f;
 
 	/** Whether or not this sequence playback will loop. */
-	UPROPERTY(EditAnywhere, Category = "Sequence Player")
+	UPROPERTY(EditAnywhere, Category = "Sequence Player", meta = (OnBecomeRelevant))
 	bool bLoop = false;
 
 	// Latent pin support boilerplate
