UAF: Fix crash when selecting a 1D blendspace in a state tree

Query superclass chain when trying to factory-instantiate graphs from objects

diff --git a/Engine/Plugins/Experimental/AnimNextAnimGraph/Source/AnimNextAnimGraph/Private/Factory/AnimGraphFactory.cpp b/Engine/Plugins/Experimental/AnimNextAnimGraph/Source/AnimNextAnimGraph/Private/Factory/AnimGraphFactory.cpp
--- a/Engine/Plugins/Experimental/AnimNextAnimGraph/Source/AnimNextAnimGraph/Private/Factory/AnimGraphFactory.cpp
+++ b/Engine/Plugins/Experimental/AnimNextAnimGraph/Source/AnimNextAnimGraph/Private/Factory/AnimGraphFactory.cpp
@@ -18,6 +18,19 @@ static TStripedMap<32, uint64, TWeakObjectPtr<const UAnimNextAnimationGraph>> GA
 static TMap<const UClass*, TPair<FAnimNextFactoryParams, FAnimGraphFactory::FParamsInitializer>> GDefaultParamsForObject;
 static TArray<UClass*> GAllFactoryClasses;
 
+const TPair<FAnimNextFactoryParams, FAnimGraphFactory::FParamsInitializer>* FindParamsForClass(const UClass* InClass)
+{
+	const UClass* CurrentClass = InClass;
+	const TPair<FAnimNextFactoryParams, FAnimGraphFactory::FParamsInitializer>* FoundParams = nullptr;
+	while (FoundParams == nullptr && CurrentClass != nullptr)
+	{
+		FoundParams = GDefaultParamsForObject.Find(CurrentClass);
+		CurrentClass = CurrentClass->GetSuperClass();
+	}
+
+	return FoundParams;
+}
+
 }
 
 namespace UE::AnimNext
@@ -30,7 +43,7 @@ uint64 IAnimGraphBuilder::GetKey() const
 
 const UAnimNextAnimationGraph* FAnimGraphFactory::GetDefaultGraphHost()
 {
-	const TPair<FAnimNextFactoryParams, FParamsInitializer>* ParamsTask = Private::GDefaultParamsForObject.Find(UAnimNextAnimationGraph::StaticClass());
+	const TPair<FAnimNextFactoryParams, FParamsInitializer>* ParamsTask = Private::FindParamsForClass(UAnimNextAnimationGraph::StaticClass());
 	return BuildGraph(ParamsTask->Key.GetBuilder());
 }
 
@@ -80,7 +93,7 @@ const UAnimNextAnimationGraph* FAnimGraphFactory::BuildGraph(const IAnimGraphBui
 const FAnimNextFactoryParams& FAnimGraphFactory::GetDefaultParamsForClass(const UClass* InClass)
 {
 	check(InClass);
-	const TPair<FAnimNextFactoryParams, FParamsInitializer>* ParamsTask = Private::GDefaultParamsForObject.Find(InClass);
+	const TPair<FAnimNextFactoryParams, FParamsInitializer>* ParamsTask = Private::FindParamsForClass(InClass);
 	check(ParamsTask != nullptr);	// Unregistered object type
 	return ParamsTask->Key;
 }
@@ -88,7 +101,7 @@ const FAnimNextFactoryParams& FAnimGraphFactory::GetDefaultParamsForClass(const
 FAnimNextFactoryParams FAnimGraphFactory::GetDefaultParamsForObject(const UObject* InObject)
 {
 	check(InObject);
-	const TPair<FAnimNextFactoryParams, FParamsInitializer>* ParamsTask = Private::GDefaultParamsForObject.Find(InObject->GetClass());
+	const TPair<FAnimNextFactoryParams, FParamsInitializer>* ParamsTask = Private::FindParamsForClass(InObject->GetClass());
 	check(ParamsTask != nullptr);	// Unregistered object type
 
 	FAnimNextFactoryParams ParamsCopy = ParamsTask->Key;
@@ -105,7 +118,7 @@ FAnimNextFactoryParams FAnimGraphFactory::GetDefaultParamsForObject(const UObjec
 void FAnimGraphFactory::InitializeDefaultParamsForObject(const UObject* InObject, FAnimNextFactoryParams& InOutParams)
 {
 	check(InObject);
-	const TPair<FAnimNextFactoryParams, FParamsInitializer>* ParamsTask = Private::GDefaultParamsForObject.Find(InObject->GetClass());
+	const TPair<FAnimNextFactoryParams, FParamsInitializer>* ParamsTask = Private::FindParamsForClass(InObject->GetClass());
 	check(ParamsTask != nullptr);	// Unregistered object type
 
 	if (InOutParams.Builder.Stacks.Num() > 0)
