UAF: Namespace renamed

diff --git a/Engine/Plugins/Experimental/AnimNextAnimGraph/Source/AnimNextAnimGraph/Private/Factory/AnimGraphFactory.cpp b/Engine/Plugins/Experimental/AnimNextAnimGraph/Source/AnimNextAnimGraph/Private/Factory/AnimGraphFactory.cpp
--- a/Engine/Plugins/Experimental/AnimNextAnimGraph/Source/AnimNextAnimGraph/Private/Factory/AnimGraphFactory.cpp
+++ b/Engine/Plugins/Experimental/AnimNextAnimGraph/Source/AnimNextAnimGraph/Private/Factory/AnimGraphFactory.cpp
@@ -8,7 +8,7 @@
 #include "Graph/AnimNextAnimationGraph.h"
 #include "TraitCore/TraitWriter.h"
 
-namespace UE::AnimNext::Private
+namespace UE::UAF::Private
 {
 
 // References to all generated factory graphs, keyed by method
@@ -33,7 +33,7 @@ const TPair<FAnimNextFactoryParams, FAnimGraphFactory::FParamsInitializer>* Find
 
 }
 
-namespace UE::AnimNext
+namespace UE::UAF
 {
 
 uint64 IAnimGraphBuilder::GetKey() const
diff --git a/Engine/Plugins/Experimental/AnimNextAnimGraph/Source/AnimNextAnimGraph/Private/Factory/AnimNextFactoryParams.cpp b/Engine/Plugins/Experimental/AnimNextAnimGraph/Source/AnimNextAnimGraph/Private/Factory/AnimNextFactoryParams.cpp
--- a/Engine/Plugins/Experimental/AnimNextAnimGraph/Source/AnimNextAnimGraph/Private/Factory/AnimNextFactoryParams.cpp
+++ b/Engine/Plugins/Experimental/AnimNextAnimGraph/Source/AnimNextAnimGraph/Private/Factory/AnimNextFactoryParams.cpp
@@ -17,8 +17,8 @@ void FAnimNextFactoryParams::InitializeInstance(FUAFAssetInstance& InInstance) c
 		}
 	}
 
-	const UE::AnimNext::FInstanceTaskContext TaskContext(InInstance);
-	for (const UE::AnimNext::FInstanceTask& InitializeTask : InitializeTasks)
+	const UE::UAF::FInstanceTaskContext TaskContext(InInstance);
+	for (const UE::UAF::FInstanceTask& InitializeTask : InitializeTasks)
 	{
 		InitializeTask(TaskContext);
 	}
diff --git a/Engine/Plugins/Experimental/AnimNextAnimGraph/Source/AnimNextAnimGraph/Private/Factory/AnimNextSimpleAnimGraphBuilder.cpp b/Engine/Plugins/Experimental/AnimNextAnimGraph/Source/AnimNextAnimGraph/Private/Factory/AnimNextSimpleAnimGraphBuilder.cpp
--- a/Engine/Plugins/Experimental/AnimNextAnimGraph/Source/AnimNextAnimGraph/Private/Factory/AnimNextSimpleAnimGraphBuilder.cpp
+++ b/Engine/Plugins/Experimental/AnimNextAnimGraph/Source/AnimNextAnimGraph/Private/Factory/AnimNextSimpleAnimGraphBuilder.cpp
@@ -11,9 +11,9 @@
 #include "TraitCore/TraitWriter.h"
 #include "Traits/ReferencePoseTrait.h"
 
-bool FAnimNextSimpleAnimGraphBuilder::Build(UE::AnimNext::FAnimGraphBuilderContext& InContext) const
+bool FAnimNextSimpleAnimGraphBuilder::Build(UE::UAF::FAnimGraphBuilderContext& InContext) const
 {
-	using namespace UE::AnimNext;
+	using namespace UE::UAF;
 
 	ensure(Stacks.Num() == 1);	// Currently only support a single stack
 
@@ -141,7 +141,7 @@ bool FAnimNextSimpleAnimGraphBuilder::Build(UE::AnimNext::FAnimGraphBuilderConte
 
 void FAnimNextSimpleAnimGraphBuilder::ValidateTraitStruct(int32 InStackIndex, TConstStructView<FAnimNextTraitSharedData> InStruct)
 {
-	using namespace UE::AnimNext;
+	using namespace UE::UAF;
 
 	ensure(InStackIndex == 0);	// Currently only support a single stack
 	if (InStackIndex != 0)
@@ -171,7 +171,7 @@ void FAnimNextSimpleAnimGraphBuilder::ValidateTraitStruct(int32 InStackIndex, TC
 
 void FAnimNextSimpleAnimGraphBuilder::PushTraitStructViewToStack(int32 InStackIndex, TConstStructView<FAnimNextTraitSharedData> InStruct)
 {
-	using namespace UE::AnimNext;
+	using namespace UE::UAF;
 
 	ValidateTraitStruct(InStackIndex, InStruct);
 
@@ -194,7 +194,7 @@ void FAnimNextSimpleAnimGraphBuilder::PushTraitStructViewToStack(int32 InStackIn
 
 void FAnimNextSimpleAnimGraphBuilder::PushTraitInstancedStructToStack(int32 InStackIndex, TInstancedStruct<FAnimNextTraitSharedData>&& InStruct)
 {
-	using namespace UE::AnimNext;
+	using namespace UE::UAF;
 
 	ValidateTraitStruct(InStackIndex, InStruct);
 
diff --git a/Engine/Plugins/Experimental/AnimNextAnimGraph/Source/AnimNextAnimGraph/Private/Graph/AnimGraph/AnimNode_AnimNextGraph.cpp b/Engine/Plugins/Experimental/AnimNextAnimGraph/Source/AnimNextAnimGraph/Private/Graph/AnimGraph/AnimNode_AnimNextGraph.cpp
--- a/Engine/Plugins/Experimental/AnimNextAnimGraph/Source/AnimNextAnimGraph/Private/Graph/AnimGraph/AnimNode_AnimNextGraph.cpp
+++ b/Engine/Plugins/Experimental/AnimNextAnimGraph/Source/AnimNextAnimGraph/Private/Graph/AnimGraph/AnimNode_AnimNextGraph.cpp
@@ -51,7 +51,7 @@ void FAnimNode_AnimNextGraph::Update_AnyThread(const FAnimationUpdateContext& Co
 {
 	DECLARE_SCOPE_HIERARCHICAL_COUNTER_FUNC()
 
-	using namespace UE::AnimNext;
+	using namespace UE::UAF;
 
 	SourceLink.Update(Context);
 
@@ -63,7 +63,7 @@ void FAnimNode_AnimNextGraph::Update_AnyThread(const FAnimationUpdateContext& Co
 
 		FUpdateGraphContext UpdateGraphContext(*GraphInstance.Get(), Context.GetDeltaTime());
 		UpdateGraphContext.SetBindingObject(Context.AnimInstanceProxy->GetSkelMeshComponent());
-		UE::AnimNext::UpdateGraph(UpdateGraphContext);
+		UE::UAF::UpdateGraph(UpdateGraphContext);
 	}
 
 	FAnimNode_CustomProperty::Update_AnyThread(Context);
@@ -75,7 +75,7 @@ void FAnimNode_AnimNextGraph::Initialize_AnyThread(const FAnimationInitializeCon
 {
 	DECLARE_SCOPE_HIERARCHICAL_COUNTER_FUNC()
 
-	using namespace UE::AnimNext;
+	using namespace UE::UAF;
 
 	SourceLink.Initialize(Context);
 
@@ -107,7 +107,7 @@ void FAnimNode_AnimNextGraph::Evaluate_AnyThread(FPoseContext& Output)
 {
 	DECLARE_SCOPE_HIERARCHICAL_COUNTER_FUNC()
 
-	using namespace UE::AnimNext;
+	using namespace UE::UAF;
 
 	if (!CVarAnimNextForceAnimBP.GetValueOnAnyThread() && GraphInstance.IsValid())
 	{
@@ -119,14 +119,14 @@ void FAnimNode_AnimNextGraph::Evaluate_AnyThread(FPoseContext& Output)
 
 		const int32 LODLevel = Output.AnimInstanceProxy->GetLODLevel();
 
-		const UE::AnimNext::FReferencePose& RefPose = RefPoseHandle.GetRef<UE::AnimNext::FReferencePose>();
+		const UE::UAF::FReferencePose& RefPose = RefPoseHandle.GetRef<UE::UAF::FReferencePose>();
 		FAnimNextGraphLODPose ResultPose;
 		ResultPose.LODPose = FLODPoseHeap(RefPose, LODLevel, true, Output.ExpectsAdditivePose());
 
 		{
-			UE::AnimNext::FEvaluateGraphContext EvaluateGraphContext(*GraphInstance.Get(), RefPose, LODLevel);
+			UE::UAF::FEvaluateGraphContext EvaluateGraphContext(*GraphInstance.Get(), RefPose, LODLevel);
 			EvaluateGraphContext.SetBindingObject(SkeletalMeshComponent);
-			const FEvaluationProgram EvaluationProgram = UE::AnimNext::EvaluateGraph(EvaluateGraphContext);
+			const FEvaluationProgram EvaluationProgram = UE::UAF::EvaluateGraph(EvaluateGraphContext);
 
 			FEvaluationVM EvaluationVM(EEvaluationFlags::All, RefPose, LODLevel);
 			bool bHasValidOutput = false;
diff --git a/Engine/Plugins/Experimental/AnimNextAnimGraph/Source/AnimNextAnimGraph/Private/Graph/AnimNextAnimationGraph.cpp b/Engine/Plugins/Experimental/AnimNextAnimGraph/Source/AnimNextAnimGraph/Private/Graph/AnimNextAnimationGraph.cpp
--- a/Engine/Plugins/Experimental/AnimNextAnimGraph/Source/AnimNextAnimGraph/Private/Graph/AnimNextAnimationGraph.cpp
+++ b/Engine/Plugins/Experimental/AnimNextAnimGraph/Source/AnimNextAnimGraph/Private/Graph/AnimNextAnimationGraph.cpp
@@ -35,7 +35,7 @@ UAnimNextAnimationGraph::UAnimNextAnimationGraph(const FObjectInitializer& Objec
 	ExtendedExecuteContext.SetContextPublicDataStruct(FAnimNextExecuteContext::StaticStruct());
 }
 
-TSharedPtr<FAnimNextGraphInstance> UAnimNextAnimationGraph::AllocateInstance(const UE::AnimNext::FGraphAllocationParams& InParams) const
+TSharedPtr<FAnimNextGraphInstance> UAnimNextAnimationGraph::AllocateInstance(const UE::UAF::FGraphAllocationParams& InParams) const
 {
 	SCOPE_CYCLE_COUNTER(STAT_AnimNext_Graph_AllocateInstance);
 
@@ -61,7 +61,7 @@ TSharedPtr<FAnimNextGraphInstance> UAnimNextAnimationGraph::AllocateInstance(con
 	InstanceImpl->InitializeVariables(InParams.Overrides);
 
 	{
-		UE::AnimNext::FExecutionContext Context(*InstanceImpl.Get());
+		UE::UAF::FExecutionContext Context(*InstanceImpl.Get());
 
 		if (InParams.ParentContext != nullptr)
 		{
@@ -71,8 +71,8 @@ TSharedPtr<FAnimNextGraphInstance> UAnimNextAnimationGraph::AllocateInstance(con
 		{
 			const FAnimNextSkeletalMeshComponentReferenceComponent& ComponentReference = InParams.ModuleInstance->GetComponent<FAnimNextSkeletalMeshComponentReferenceComponent>();
 
-			UE::AnimNext::FDataHandle RefPoseHandle = UE::AnimNext::FDataRegistry::Get()->GetOrGenerateReferencePose(ComponentReference.GetComponent());
-			const UE::AnimNext::FReferencePose& RefPose = RefPoseHandle.GetRef<UE::AnimNext::FReferencePose>();
+			UE::UAF::FDataHandle RefPoseHandle = UE::UAF::FDataRegistry::Get()->GetOrGenerateReferencePose(ComponentReference.GetComponent());
+			const UE::UAF::FReferencePose& RefPose = RefPoseHandle.GetRef<UE::UAF::FReferencePose>();
 			Context.SetBindingObject(RefPose.SkeletalMeshComponent);
 		}
 		InstanceImpl->GraphInstancePtr = Context.AllocateNodeInstance(*InstanceImpl, ResolvedRootTraitHandle);
@@ -175,7 +175,7 @@ void UAnimNextAnimationGraph::GetPreloadDependencies(TArray<UObject*>& OutDeps)
 
 bool UAnimNextAnimationGraph::LoadFromArchiveBuffer(const TArray<uint8>& InSharedDataArchiveBuffer)
 {
-	using namespace UE::AnimNext;
+	using namespace UE::UAF;
 
 	// Reconstruct our graph shared data
 	FMemoryReader GraphSharedDataArchive(InSharedDataArchiveBuffer);
diff --git a/Engine/Plugins/Experimental/AnimNextAnimGraph/Source/AnimNextAnimGraph/Private/Graph/AnimNextGraphInstance.cpp b/Engine/Plugins/Experimental/AnimNextAnimGraph/Source/AnimNextAnimGraph/Private/Graph/AnimNextGraphInstance.cpp
--- a/Engine/Plugins/Experimental/AnimNextAnimGraph/Source/AnimNextAnimGraph/Private/Graph/AnimNextGraphInstance.cpp
+++ b/Engine/Plugins/Experimental/AnimNextAnimGraph/Source/AnimNextAnimGraph/Private/Graph/AnimNextGraphInstance.cpp
@@ -65,7 +65,7 @@ FName FAnimNextGraphInstance::GetEntryPoint() const
 	return EntryPoint;
 }
 
-UE::AnimNext::FWeakTraitPtr FAnimNextGraphInstance::GetGraphRootPtr() const
+UE::UAF::FWeakTraitPtr FAnimNextGraphInstance::GetGraphRootPtr() const
 {
 	return GraphInstancePtr;
 }
@@ -126,7 +126,7 @@ void FAnimNextGraphInstance::MarkAsUpdated()
 	bHasUpdatedOnce = true;
 }
 
-void FAnimNextGraphInstance::ExecuteLatentPins(const TConstArrayView<UE::AnimNext::FLatentPropertyHandle>& LatentHandles, void* DestinationBasePtr, bool bIsFrozen, bool bJustBecameRelevant)
+void FAnimNextGraphInstance::ExecuteLatentPins(const TConstArrayView<UE::UAF::FLatentPropertyHandle>& LatentHandles, void* DestinationBasePtr, bool bIsFrozen, bool bJustBecameRelevant)
 {
 	SCOPE_CYCLE_COUNTER(STAT_AnimNext_Graph_RigVM);
 
@@ -147,15 +147,15 @@ void FAnimNextGraphInstance::ExecuteLatentPins(const TConstArrayView<UE::AnimNex
 
 		// Insert our context data for the scope of execution
 		FAnimNextGraphLatentPropertiesContextData ContextData(ModuleInstance, this, LatentHandles, DestinationBasePtr, bIsFrozen, bJustBecameRelevant);
-		UE::AnimNext::FScopedExecuteContextData ContextDataScope(AnimNextContext, ContextData);
+		UE::UAF::FScopedExecuteContextData ContextDataScope(AnimNextContext, ContextData);
 
 		VM->ExecuteVM(ExtendedExecuteContext, FRigUnit_AnimNextShimRoot::EventName);
 	}
 }
 
-void FAnimNextGraphInstance::CopyVariablesToLatentPins(const TConstArrayView<UE::AnimNext::FLatentPropertyHandle>& LatentHandles, void* DestinationBasePtr, bool bIsFrozen, bool bJustBecameRelevant)
+void FAnimNextGraphInstance::CopyVariablesToLatentPins(const TConstArrayView<UE::UAF::FLatentPropertyHandle>& LatentHandles, void* DestinationBasePtr, bool bIsFrozen, bool bJustBecameRelevant)
 {
-	for (UE::AnimNext::FLatentPropertyHandle Handle : LatentHandles)
+	for (UE::UAF::FLatentPropertyHandle Handle : LatentHandles)
 	{
 		if (!Handle.IsIndexValid())
 		{
@@ -219,7 +219,7 @@ void FAnimNextGraphInstance::Thaw()
 		}
 
 		{
-			UE::AnimNext::FExecutionContext Context(*this);
+			UE::UAF::FExecutionContext Context(*this);
 			if(const FAnimNextTraitHandle* FoundHandle = AnimationGraph->ResolvedRootTraitHandles.Find(EntryPoint))
 			{
 				GraphInstancePtr = Context.AllocateNodeInstance(*this, *FoundHandle);
diff --git a/Engine/Plugins/Experimental/AnimNextAnimGraph/Source/AnimNextAnimGraph/Private/Graph/AnimNextModuleAnimGraphComponent.cpp b/Engine/Plugins/Experimental/AnimNextAnimGraph/Source/AnimNextAnimGraph/Private/Graph/AnimNextModuleAnimGraphComponent.cpp
--- a/Engine/Plugins/Experimental/AnimNextAnimGraph/Source/AnimNextAnimGraph/Private/Graph/AnimNextModuleAnimGraphComponent.cpp
+++ b/Engine/Plugins/Experimental/AnimNextAnimGraph/Source/AnimNextAnimGraph/Private/Graph/AnimNextModuleAnimGraphComponent.cpp
@@ -6,7 +6,7 @@
 #include "Graph/AnimNextAnimationGraph.h"
 #include "Module/AnimNextModuleInstance.h"
 
-TWeakPtr<FAnimNextGraphInstance> FAnimNextModuleAnimGraphComponent::AllocateInstance(const UAnimNextAnimationGraph* InAnimationGraph, const TSharedPtr<const UE::AnimNext::FVariableOverridesCollection>& InOverrides)
+TWeakPtr<FAnimNextGraphInstance> FAnimNextModuleAnimGraphComponent::AllocateInstance(const UAnimNextAnimationGraph* InAnimationGraph, const TSharedPtr<const UE::UAF::FVariableOverridesCollection>& InOverrides)
 {
 	check(InAnimationGraph);
 
diff --git a/Engine/Plugins/Experimental/AnimNextAnimGraph/Source/AnimNextAnimGraph/Private/Graph/AnimNextSkeletalMeshComponentLibrary.cpp b/Engine/Plugins/Experimental/AnimNextAnimGraph/Source/AnimNextAnimGraph/Private/Graph/AnimNextSkeletalMeshComponentLibrary.cpp
--- a/Engine/Plugins/Experimental/AnimNextAnimGraph/Source/AnimNextAnimGraph/Private/Graph/AnimNextSkeletalMeshComponentLibrary.cpp
+++ b/Engine/Plugins/Experimental/AnimNextAnimGraph/Source/AnimNextAnimGraph/Private/Graph/AnimNextSkeletalMeshComponentLibrary.cpp
@@ -5,6 +5,6 @@
 
 FAnimNextGraphReferencePose UAnimNextSkeletalMeshComponentLibrary::GetReferencePose(USkeletalMeshComponent* InComponent)
 {
-	UE::AnimNext::FDataHandle RefPoseHandle = UE::AnimNext::FDataRegistry::Get()->GetOrGenerateReferencePose(InComponent);
+	UE::UAF::FDataHandle RefPoseHandle = UE::UAF::FDataRegistry::Get()->GetOrGenerateReferencePose(InComponent);
 	return FAnimNextGraphReferencePose(RefPoseHandle);
 }
diff --git a/Engine/Plugins/Experimental/AnimNextAnimGraph/Source/AnimNextAnimGraph/Private/Graph/GC_GraphInstanceComponent.cpp b/Engine/Plugins/Experimental/AnimNextAnimGraph/Source/AnimNextAnimGraph/Private/Graph/GC_GraphInstanceComponent.cpp
--- a/Engine/Plugins/Experimental/AnimNextAnimGraph/Source/AnimNextAnimGraph/Private/Graph/GC_GraphInstanceComponent.cpp
+++ b/Engine/Plugins/Experimental/AnimNextAnimGraph/Source/AnimNextAnimGraph/Private/Graph/GC_GraphInstanceComponent.cpp
@@ -5,14 +5,14 @@
 #include "TraitCore/ExecutionContext.h"
 #include "TraitInterfaces/IGarbageCollection.h"
 
-void FGCGraphInstanceComponent::Register(const UE::AnimNext::FWeakTraitPtr& InTraitPtr)
+void FGCGraphInstanceComponent::Register(const UE::UAF::FWeakTraitPtr& InTraitPtr)
 {
 	TraitsWithReferences.Add(InTraitPtr);
 }
 
-void FGCGraphInstanceComponent::Unregister(const UE::AnimNext::FWeakTraitPtr& InTraitPtr)
+void FGCGraphInstanceComponent::Unregister(const UE::UAF::FWeakTraitPtr& InTraitPtr)
 {
-	using namespace UE::AnimNext;
+	using namespace UE::UAF;
 	
 	const int32 EntryIndex = TraitsWithReferences.IndexOfByPredicate(
 		[&InTraitPtr](const FWeakTraitPtr& TraitPtr)
@@ -28,7 +28,7 @@ void FGCGraphInstanceComponent::Unregister(const UE::AnimNext::FWeakTraitPtr& In
 
 void FGCGraphInstanceComponent::AddStructReferencedObjects(FReferenceCollector& Collector)
 {
-	using namespace UE::AnimNext;
+	using namespace UE::UAF;
 	
 	FExecutionContext Context;
 	FTraitStackBinding TraitStack;
diff --git a/Engine/Plugins/Experimental/AnimNextAnimGraph/Source/AnimNextAnimGraph/Private/Graph/GC_GraphInstanceComponent.h b/Engine/Plugins/Experimental/AnimNextAnimGraph/Source/AnimNextAnimGraph/Private/Graph/GC_GraphInstanceComponent.h
--- a/Engine/Plugins/Experimental/AnimNextAnimGraph/Source/AnimNextAnimGraph/Private/Graph/GC_GraphInstanceComponent.h
+++ b/Engine/Plugins/Experimental/AnimNextAnimGraph/Source/AnimNextAnimGraph/Private/Graph/GC_GraphInstanceComponent.h
@@ -23,17 +23,17 @@ struct FGCGraphInstanceComponent : public FUAFGraphInstanceComponent
 
 	// Registers the provided trait with the GC system
 	// Once registered, IGarbageCollection::AddReferencedObjects will be called on it during GC
-	void Register(const UE::AnimNext::FWeakTraitPtr& InTraitPtr);
+	void Register(const UE::UAF::FWeakTraitPtr& InTraitPtr);
 
 	// Unregisters the provided trait from the GC system
-	void Unregister(const UE::AnimNext::FWeakTraitPtr& InTraitPtr);
+	void Unregister(const UE::UAF::FWeakTraitPtr& InTraitPtr);
 
 	// Called during garbage collection to collect strong object references
 	void AddStructReferencedObjects(FReferenceCollector& Collector);
 
 private:
 	// List of trait handles that contain UObject references and implement IGarbageCollection
-	TArray<UE::AnimNext::FWeakTraitPtr> TraitsWithReferences;
+	TArray<UE::UAF::FWeakTraitPtr> TraitsWithReferences;
 };
 
 template<>
diff --git a/Engine/Plugins/Experimental/AnimNextAnimGraph/Source/AnimNextAnimGraph/Private/Graph/RigDecorator_AnimNextCppTrait.cpp b/Engine/Plugins/Experimental/AnimNextAnimGraph/Source/AnimNextAnimGraph/Private/Graph/RigDecorator_AnimNextCppTrait.cpp
--- a/Engine/Plugins/Experimental/AnimNextAnimGraph/Source/AnimNextAnimGraph/Private/Graph/RigDecorator_AnimNextCppTrait.cpp
+++ b/Engine/Plugins/Experimental/AnimNextAnimGraph/Source/AnimNextAnimGraph/Private/Graph/RigDecorator_AnimNextCppTrait.cpp
@@ -41,7 +41,7 @@ void FRigDecorator_AnimNextCppDecorator::GetProgrammaticPins(URigVMController* I
 
 
 	// Give the Trait an opportunity to generate its own pins
-	const UE::AnimNext::FTrait* Trait = GetTrait();
+	const UE::UAF::FTrait* Trait = GetTrait();
 	check(Trait != nullptr);
 	FAnimNextTraitSharedData* TraitSharedData = (FAnimNextTraitSharedData*)DefaultValueMemoryScope.GetStructMemory();
 
@@ -106,8 +106,8 @@ void FRigDecorator_AnimNextCppDecorator::GetProgrammaticPins(URigVMController* I
 	}
 }
 
-const UE::AnimNext::FTrait* FRigDecorator_AnimNextCppDecorator::GetTrait() const
+const UE::UAF::FTrait* FRigDecorator_AnimNextCppDecorator::GetTrait() const
 {
-	return UE::AnimNext::FTraitRegistry::Get().Find(DecoratorSharedDataStruct);
+	return UE::UAF::FTraitRegistry::Get().Find(DecoratorSharedDataStruct);
 }
 #endif
diff --git a/Engine/Plugins/Experimental/AnimNextAnimGraph/Source/AnimNextAnimGraph/Private/Graph/RigUnit_AnimNextGraphEvaluator.cpp b/Engine/Plugins/Experimental/AnimNextAnimGraph/Source/AnimNextAnimGraph/Private/Graph/RigUnit_AnimNextGraphEvaluator.cpp
--- a/Engine/Plugins/Experimental/AnimNextAnimGraph/Source/AnimNextAnimGraph/Private/Graph/RigUnit_AnimNextGraphEvaluator.cpp
+++ b/Engine/Plugins/Experimental/AnimNextAnimGraph/Source/AnimNextAnimGraph/Private/Graph/RigUnit_AnimNextGraphEvaluator.cpp
@@ -8,7 +8,7 @@
 
 #include UE_INLINE_GENERATED_CPP_BY_NAME(RigUnit_AnimNextGraphEvaluator)
 
-namespace UE::AnimNext::Private
+namespace UE::UAF::Private
 {
 	static TMap<uint32, FAnimNextGraphEvaluatorExecuteDefinition> GRegisteredGraphEvaluatorMethods;
 
@@ -31,12 +31,12 @@ void FRigUnit_AnimNextGraphEvaluator::StaticExecute(FRigVMExtendedExecuteContext
 	const FAnimNextExecuteContext& VMExecuteContext = RigVMExecuteContext.GetPublicData<FAnimNextExecuteContext>();
 	const FAnimNextGraphLatentPropertiesContextData& LatentPropertiesContextData = VMExecuteContext.GetContextData<FAnimNextGraphLatentPropertiesContextData>();
 
-	const TConstArrayView<UE::AnimNext::FLatentPropertyHandle>& LatentHandles = LatentPropertiesContextData.GetLatentHandles();
+	const TConstArrayView<UE::UAF::FLatentPropertyHandle>& LatentHandles = LatentPropertiesContextData.GetLatentHandles();
 	uint8* DestinationBasePtr = (uint8*)LatentPropertiesContextData.GetDestinationBasePtr();
 	const bool bIsFrozen = LatentPropertiesContextData.IsFrozen();
 	const bool bJustBecameRelevant = LatentPropertiesContextData.JustBecameRelevant();
 
-	for (UE::AnimNext::FLatentPropertyHandle Handle : LatentHandles)
+	for (UE::UAF::FLatentPropertyHandle Handle : LatentHandles)
 	{
 		if (!Handle.IsIndexValid())
 		{
@@ -78,7 +78,7 @@ void FRigUnit_AnimNextGraphEvaluator::StaticExecute(FRigVMExtendedExecuteContext
 
 void FRigUnit_AnimNextGraphEvaluator::RegisterExecuteMethod(const FAnimNextGraphEvaluatorExecuteDefinition& ExecuteDefinition)
 {
-	using namespace UE::AnimNext::Private;
+	using namespace UE::UAF::Private;
 
 	if (GRegisteredGraphEvaluatorMethods.Contains(ExecuteDefinition.Hash))
 	{
@@ -95,7 +95,7 @@ void FRigUnit_AnimNextGraphEvaluator::RegisterExecuteMethod(const FAnimNextGraph
 
 const FAnimNextGraphEvaluatorExecuteDefinition* FRigUnit_AnimNextGraphEvaluator::FindExecuteMethod(uint32 ExecuteMethodHash)
 {
-	using namespace UE::AnimNext::Private;
+	using namespace UE::UAF::Private;
 
 	return GRegisteredGraphEvaluatorMethods.Find(ExecuteMethodHash);
 }
diff --git a/Engine/Plugins/Experimental/AnimNextAnimGraph/Source/AnimNextAnimGraph/Private/Graph/RigUnit_AnimNextRunAnimationGraph_v1.cpp b/Engine/Plugins/Experimental/AnimNextAnimGraph/Source/AnimNextAnimGraph/Private/Graph/RigUnit_AnimNextRunAnimationGraph_v1.cpp
--- a/Engine/Plugins/Experimental/AnimNextAnimGraph/Source/AnimNextAnimGraph/Private/Graph/RigUnit_AnimNextRunAnimationGraph_v1.cpp
+++ b/Engine/Plugins/Experimental/AnimNextAnimGraph/Source/AnimNextAnimGraph/Private/Graph/RigUnit_AnimNextRunAnimationGraph_v1.cpp
@@ -20,7 +20,7 @@ FRigUnit_AnimNextRunAnimationGraph_v1_Execute()
 {
 	SCOPED_NAMED_EVENT(AnimNext_Run_Graph_V1, FColor::Orange);
 
-	using namespace UE::AnimNext;
+	using namespace UE::UAF;
 
 	const FAnimNextModuleContextData& ModuleContextData = ExecuteContext.GetContextData<FAnimNextModuleContextData>();
 	FAnimNextModuleInstance& ModuleInstance = ModuleContextData.GetModuleInstance();
@@ -66,7 +66,7 @@ FRigUnit_AnimNextRunAnimationGraph_v1_Execute()
 		AnimNextExecuteContext.SetDeltaTime(ExecuteContext.GetDeltaTime());
 	}
 
-	const UE::AnimNext::FReferencePose& RefPose = ReferencePose.ReferencePose.GetRef<UE::AnimNext::FReferencePose>();
+	const UE::UAF::FReferencePose& RefPose = ReferencePose.ReferencePose.GetRef<UE::UAF::FReferencePose>();
 
 	int32 DesiredLOD = LOD;
 	if (DesiredLOD == -1)
diff --git a/Engine/Plugins/Experimental/AnimNextAnimGraph/Source/AnimNextAnimGraph/Private/Graph/RigUnit_AnimNextRunAnimationGraph_v2.cpp b/Engine/Plugins/Experimental/AnimNextAnimGraph/Source/AnimNextAnimGraph/Private/Graph/RigUnit_AnimNextRunAnimationGraph_v2.cpp
--- a/Engine/Plugins/Experimental/AnimNextAnimGraph/Source/AnimNextAnimGraph/Private/Graph/RigUnit_AnimNextRunAnimationGraph_v2.cpp
+++ b/Engine/Plugins/Experimental/AnimNextAnimGraph/Source/AnimNextAnimGraph/Private/Graph/RigUnit_AnimNextRunAnimationGraph_v2.cpp
@@ -24,7 +24,7 @@ FRigUnit_AnimNextRunAnimationGraph_v2_Execute()
 {
 	SCOPED_NAMED_EVENT(AnimNext_Run_Graph_V2, FColor::Orange);
 
-	using namespace UE::AnimNext;
+	using namespace UE::UAF;
 
 	const FAnimNextModuleContextData& ModuleContextData = ExecuteContext.GetContextData<FAnimNextModuleContextData>();
 	FAnimNextModuleInstance& ModuleInstance = ModuleContextData.GetModuleInstance();
@@ -34,7 +34,7 @@ FRigUnit_AnimNextRunAnimationGraph_v2_Execute()
 		return;
 	}
 
-	const UE::AnimNext::FReferencePose& RefPose = ReferencePose.ReferencePose.GetRef<UE::AnimNext::FReferencePose>();
+	const UE::UAF::FReferencePose& RefPose = ReferencePose.ReferencePose.GetRef<UE::UAF::FReferencePose>();
 
 	int32 DesiredLOD = LOD;
 	if (DesiredLOD == -1)
diff --git a/Engine/Plugins/Experimental/AnimNextAnimGraph/Source/AnimNextAnimGraph/Private/Graph/SyncGroup_GraphInstanceComponent.cpp b/Engine/Plugins/Experimental/AnimNextAnimGraph/Source/AnimNextAnimGraph/Private/Graph/SyncGroup_GraphInstanceComponent.cpp
--- a/Engine/Plugins/Experimental/AnimNextAnimGraph/Source/AnimNextAnimGraph/Private/Graph/SyncGroup_GraphInstanceComponent.cpp
+++ b/Engine/Plugins/Experimental/AnimNextAnimGraph/Source/AnimNextAnimGraph/Private/Graph/SyncGroup_GraphInstanceComponent.cpp
@@ -128,7 +128,7 @@
 // perform one advance that does trigger events for leaders. Followers never trigger events.
 //
 
-namespace UE::AnimNext
+namespace UE::UAF
 {
 	namespace Private
 	{
@@ -2135,7 +2135,7 @@ FSyncGroupGraphInstanceComponent& FSyncGroupGraphInstanceComponent::operator=(FS
 
 FSyncGroupGraphInstanceComponent::~FSyncGroupGraphInstanceComponent()
 {
-	using namespace UE::AnimNext;
+	using namespace UE::UAF;
 
 	// We should have released all used group names before we are destroyed
 	check(UsedUniqueGroupNames.IsEmpty());
@@ -2153,9 +2153,9 @@ FSyncGroupGraphInstanceComponent::~FSyncGroupGraphInstanceComponent()
 	FirstFreeUniqueGroupName = nullptr;
 }
 
-void FSyncGroupGraphInstanceComponent::RegisterWithGroup(const UE::AnimNext::FSyncGroupParameters& GroupParameters, const UE::AnimNext::FWeakTraitPtr& TraitPtr, const UE::AnimNext::FTraitUpdateState& TraitState)
+void FSyncGroupGraphInstanceComponent::RegisterWithGroup(const UE::UAF::FSyncGroupParameters& GroupParameters, const UE::UAF::FWeakTraitPtr& TraitPtr, const UE::UAF::FTraitUpdateState& TraitState)
 {
-	using namespace UE::AnimNext::Private;
+	using namespace UE::UAF::Private;
 
 	if (!ensure(TraitPtr.IsValid()))
 	{
@@ -2184,7 +2184,7 @@ void FSyncGroupGraphInstanceComponent::RegisterWithGroup(const UE::AnimNext::FSy
 
 	if (bIsMemberNew)
 	{
-		GroupState.Members.Add(FSyncGroupMember{ TraitState, UE::AnimNext::FTraitPtr(TraitPtr), GroupParameters, true, true });
+		GroupState.Members.Add(FSyncGroupMember{ TraitState, UE::UAF::FTraitPtr(TraitPtr), GroupParameters, true, true });
 
 #if UE_DEBUG_SYNC_GROUPS
 		if (CDebugVarEnableSyncLog.GetValueOnAnyThread())
@@ -2201,7 +2201,7 @@ void FSyncGroupGraphInstanceComponent::RegisterWithGroup(const UE::AnimNext::FSy
 
 FName FSyncGroupGraphInstanceComponent::CreateUniqueGroupName()
 {
-	using namespace UE::AnimNext;
+	using namespace UE::UAF;
 	
 	// We create a unique group name by incrementing a counter and using it to generate a new FName
 	// We also recycle them to keep the counter as low as possible because in some configurations
@@ -2237,7 +2237,7 @@ FName FSyncGroupGraphInstanceComponent::CreateUniqueGroupName()
 
 void FSyncGroupGraphInstanceComponent::ReleaseUniqueGroupName(FName GroupName)
 {
-	using namespace UE::AnimNext::Private;
+	using namespace UE::UAF::Private;
 
 	FSyncGroupUniqueName* Entry = nullptr;
 	if (ensure(UsedUniqueGroupNames.RemoveAndCopyValue(GroupName, Entry)))
@@ -2271,9 +2271,9 @@ void FSyncGroupGraphInstanceComponent::ReleaseUniqueGroupName(FName GroupName)
 	}
 }
 
-void FSyncGroupGraphInstanceComponent::PreUpdate(UE::AnimNext::FExecutionContext& Context)
+void FSyncGroupGraphInstanceComponent::PreUpdate(UE::UAF::FExecutionContext& Context)
 {
-	using namespace UE::AnimNext::Private;
+	using namespace UE::UAF::Private;
 
 	for (FSyncGroupState& GroupState : SyncGroups)
 	{
@@ -2288,9 +2288,9 @@ void FSyncGroupGraphInstanceComponent::PreUpdate(UE::AnimNext::FExecutionContext
 	}
 }
 
-void FSyncGroupGraphInstanceComponent::PostUpdate(UE::AnimNext::FExecutionContext& Context)
+void FSyncGroupGraphInstanceComponent::PostUpdate(UE::UAF::FExecutionContext& Context)
 {
-	using namespace UE::AnimNext::Private;
+	using namespace UE::UAF::Private;
 
 	// First purge stale entries and update our bookkeeping
 	int32 MaxNumMembers = 0;
diff --git a/Engine/Plugins/Experimental/AnimNextAnimGraph/Source/AnimNextAnimGraph/Private/Graph/SyncGroup_GraphInstanceComponent.h b/Engine/Plugins/Experimental/AnimNextAnimGraph/Source/AnimNextAnimGraph/Private/Graph/SyncGroup_GraphInstanceComponent.h
--- a/Engine/Plugins/Experimental/AnimNextAnimGraph/Source/AnimNextAnimGraph/Private/Graph/SyncGroup_GraphInstanceComponent.h
+++ b/Engine/Plugins/Experimental/AnimNextAnimGraph/Source/AnimNextAnimGraph/Private/Graph/SyncGroup_GraphInstanceComponent.h
@@ -8,7 +8,7 @@
 #include "Graph/UAFGraphInstanceComponent.h"
 #include "SyncGroup_GraphInstanceComponent.generated.h"
 
-namespace UE::AnimNext
+namespace UE::UAF
 {
 	namespace Private
 	{
@@ -39,7 +39,7 @@ struct FSyncGroupGraphInstanceComponent : public FUAFGraphInstanceComponent
 	virtual ~FSyncGroupGraphInstanceComponent();
 
 	// Registers the specified trait with group based synchronization
-	void RegisterWithGroup(const UE::AnimNext::FSyncGroupParameters& GroupParameters, const UE::AnimNext::FWeakTraitPtr& TraitPtr, const UE::AnimNext::FTraitUpdateState& TraitState);
+	void RegisterWithGroup(const UE::UAF::FSyncGroupParameters& GroupParameters, const UE::UAF::FWeakTraitPtr& TraitPtr, const UE::UAF::FTraitUpdateState& TraitState);
 
 	// Create a unique group name suitable for spawned sub-graphs to self-synchronize
 	// Unique group names are a limited resource, when no longed needed they must be released
@@ -50,21 +50,21 @@ struct FSyncGroupGraphInstanceComponent : public FUAFGraphInstanceComponent
 	void ReleaseUniqueGroupName(FName GroupName);
 
 	// FGraphInstanceComponent impl
-	virtual void PreUpdate(UE::AnimNext::FExecutionContext& Context) override;
-	virtual void PostUpdate(UE::AnimNext::FExecutionContext& Context) override;
+	virtual void PreUpdate(UE::UAF::FExecutionContext& Context) override;
+	virtual void PostUpdate(UE::UAF::FExecutionContext& Context) override;
 
 private:
 	// A map of sync group name -> group index
 	TMap<FName, int32> SyncGroupMap;
 
 	// A list of groups and their data
-	TArray<UE::AnimNext::Private::FSyncGroupState> SyncGroups;
+	TArray<UE::UAF::Private::FSyncGroupState> SyncGroups;
 
 	// The first free unique group name
-	UE::AnimNext::Private::FSyncGroupUniqueName* FirstFreeUniqueGroupName = nullptr;
+	UE::UAF::Private::FSyncGroupUniqueName* FirstFreeUniqueGroupName = nullptr;
 
 	// The map of currently used group names
-	TMap<FName, UE::AnimNext::Private::FSyncGroupUniqueName*> UsedUniqueGroupNames;
+	TMap<FName, UE::UAF::Private::FSyncGroupUniqueName*> UsedUniqueGroupNames;
 
 	// A counter tracking the next unique group name to allocate
 	int32 UniqueGroupNameCounter = 0;
diff --git a/Engine/Plugins/Experimental/AnimNextAnimGraph/Source/AnimNextAnimGraph/Private/Graph/TraceAnimNextGraphInstances.h b/Engine/Plugins/Experimental/AnimNextAnimGraph/Source/AnimNextAnimGraph/Private/Graph/TraceAnimNextGraphInstances.h
--- a/Engine/Plugins/Experimental/AnimNextAnimGraph/Source/AnimNextAnimGraph/Private/Graph/TraceAnimNextGraphInstances.h
+++ b/Engine/Plugins/Experimental/AnimNextAnimGraph/Source/AnimNextAnimGraph/Private/Graph/TraceAnimNextGraphInstances.h
@@ -8,7 +8,7 @@
 
 struct FAnimNextGraphInstance;
 
-namespace UE::AnimNext
+namespace UE::UAF
 {
 	struct FEvaluationProgram;
 	
@@ -16,8 +16,8 @@ namespace UE::AnimNext
 	void TraceEvaluationProgram(const FEvaluationProgram& Program, const FAnimNextGraphInstance& RootGraph);
 }
 
-#define TRACE_ANIMNEXT_GRAPHINSTANCES(RootGraph) UE::AnimNext::TraceGraphInstances(RootGraph);
-#define TRACE_ANIMNEXT_EVALUATIONPROGRAM(Program, RootGraph) UE::AnimNext::TraceEvaluationProgram(Program, RootGraph);
+#define TRACE_ANIMNEXT_GRAPHINSTANCES(RootGraph) UE::UAF::TraceGraphInstances(RootGraph);
+#define TRACE_ANIMNEXT_EVALUATIONPROGRAM(Program, RootGraph) UE::UAF::TraceEvaluationProgram(Program, RootGraph);
 
 #else
 
diff --git a/Engine/Plugins/Experimental/AnimNextAnimGraph/Source/AnimNextAnimGraph/Private/Injection/AnimNextModuleInjectionComponent.cpp b/Engine/Plugins/Experimental/AnimNextAnimGraph/Source/AnimNextAnimGraph/Private/Injection/AnimNextModuleInjectionComponent.cpp
--- a/Engine/Plugins/Experimental/AnimNextAnimGraph/Source/AnimNextAnimGraph/Private/Injection/AnimNextModuleInjectionComponent.cpp
+++ b/Engine/Plugins/Experimental/AnimNextAnimGraph/Source/AnimNextAnimGraph/Private/Injection/AnimNextModuleInjectionComponent.cpp
@@ -12,10 +12,10 @@
 FAnimNextModuleInjectionComponent::FAnimNextModuleInjectionComponent()
 {
 	FAnimNextModuleInstance& ModuleInstance = GetModuleInstance();
-	InjectionInfo = UE::AnimNext::FInjectionInfo(ModuleInstance);
+	InjectionInfo = UE::UAF::FInjectionInfo(ModuleInstance);
 
 	// Register re-injection for each (user) tick function
-	for(UE::AnimNext::FModuleEventTickFunction& TickFunction : ModuleInstance.GetTickFunctions())
+	for(UE::UAF::FModuleEventTickFunction& TickFunction : ModuleInstance.GetTickFunctions())
 	{
 		if(TickFunction.bUserEvent)
 		{
@@ -26,19 +26,19 @@ FAnimNextModuleInjectionComponent::FAnimNextModuleInjectionComponent()
 
 void FAnimNextModuleInjectionComponent::OnTraitEvent(FAnimNextTraitEvent& Event)
 {
-	if(UE::AnimNext::FInjection_InjectEvent* InjectionEvent = Event.AsType<UE::AnimNext::FInjection_InjectEvent>())
+	if(UE::UAF::FInjection_InjectEvent* InjectionEvent = Event.AsType<UE::UAF::FInjection_InjectEvent>())
 	{
 		OnInjectionEvent(*InjectionEvent);
 	}
-	else if(UE::AnimNext::FInjection_UninjectEvent* UninjectionEvent = Event.AsType<UE::AnimNext::FInjection_UninjectEvent>())
+	else if(UE::UAF::FInjection_UninjectEvent* UninjectionEvent = Event.AsType<UE::UAF::FInjection_UninjectEvent>())
 	{
 		OnUninjectionEvent(*UninjectionEvent);
 	}
 }
 
-void FAnimNextModuleInjectionComponent::OnInjectionEvent(UE::AnimNext::FInjection_InjectEvent& InEvent)
+void FAnimNextModuleInjectionComponent::OnInjectionEvent(UE::UAF::FInjection_InjectEvent& InEvent)
 {
-	using namespace UE::AnimNext;
+	using namespace UE::UAF;
 
 	if(InEvent.IsHandled())
 	{
@@ -105,9 +105,9 @@ void FAnimNextModuleInjectionComponent::OnInjectionEvent(UE::AnimNext::FInjectio
 	ensureAlways(Result == EPropertyBagResult::Success);
 }
 
-void FAnimNextModuleInjectionComponent::OnUninjectionEvent(UE::AnimNext::FInjection_UninjectEvent& InEvent)
+void FAnimNextModuleInjectionComponent::OnUninjectionEvent(UE::UAF::FInjection_UninjectEvent& InEvent)
 {
-	using namespace UE::AnimNext;
+	using namespace UE::UAF;
 
 	const FInjectionRequestArgs& RequestArgs = InEvent.Request->GetArgs();
 
@@ -163,9 +163,9 @@ void FAnimNextModuleInjectionComponent::OnUninjectionEvent(UE::AnimNext::FInject
 	}
 }
 
-void FAnimNextModuleInjectionComponent::OnReapplyInjection(const UE::AnimNext::FModuleTaskContext& InContext)
+void FAnimNextModuleInjectionComponent::OnReapplyInjection(const UE::UAF::FModuleTaskContext& InContext)
 {
-	using namespace UE::AnimNext;
+	using namespace UE::UAF;
 	
 	FAnimNextModuleInjectionComponent& Component = InContext.ModuleInstance.GetComponent<FAnimNextModuleInjectionComponent>();
 
diff --git a/Engine/Plugins/Experimental/AnimNextAnimGraph/Source/AnimNextAnimGraph/Private/Injection/AnimNextModuleInjectionComponent.h b/Engine/Plugins/Experimental/AnimNextAnimGraph/Source/AnimNextAnimGraph/Private/Injection/AnimNextModuleInjectionComponent.h
--- a/Engine/Plugins/Experimental/AnimNextAnimGraph/Source/AnimNextAnimGraph/Private/Injection/AnimNextModuleInjectionComponent.h
+++ b/Engine/Plugins/Experimental/AnimNextAnimGraph/Source/AnimNextAnimGraph/Private/Injection/AnimNextModuleInjectionComponent.h
@@ -8,7 +8,7 @@
 #include "UObject/GCObject.h"
 #include "AnimNextModuleInjectionComponent.generated.h"
 
-namespace UE::AnimNext
+namespace UE::UAF
 {
 	struct FInjection_InjectEvent;
 	struct FInjection_UninjectEvent;
@@ -24,7 +24,7 @@ struct FAnimNextModuleInjectionComponent : public FUAFModuleInstanceComponent
 	FAnimNextModuleInjectionComponent();
 
 	// Get the cached injection info for our module 
-	const UE::AnimNext::FInjectionInfo& GetInjectionInfo() const { return InjectionInfo; }
+	const UE::UAF::FInjectionInfo& GetInjectionInfo() const { return InjectionInfo; }
 
 	void AddStructReferencedObjects(class FReferenceCollector& Collector);
 
@@ -32,16 +32,16 @@ private:
 	// FUAFModuleInstanceComponent interface
 	virtual void OnTraitEvent(FAnimNextTraitEvent& Event) override;
 
-	void OnInjectionEvent(UE::AnimNext::FInjection_InjectEvent& InEvent);
-	void OnUninjectionEvent(UE::AnimNext::FInjection_UninjectEvent& InEvent);
+	void OnInjectionEvent(UE::UAF::FInjection_InjectEvent& InEvent);
+	void OnUninjectionEvent(UE::UAF::FInjection_UninjectEvent& InEvent);
 
-	static void OnReapplyInjection(const UE::AnimNext::FModuleTaskContext& InContext);
+	static void OnReapplyInjection(const UE::UAF::FModuleTaskContext& InContext);
 
 	uint32 IncrementSerialNumber();
 
 private:
 	// Info for injection
-	UE::AnimNext::FInjectionInfo InjectionInfo;
+	UE::UAF::FInjectionInfo InjectionInfo;
 
 	struct FInjectionRecord
 	{
@@ -52,8 +52,8 @@ private:
 
 		void AddReferencedObjects(FReferenceCollector& Collector);
 
-		TSharedPtr<UE::AnimNext::FInjectionRequest> GraphRequest;
-		TSharedPtr<UE::AnimNext::FInjectionRequest> ModifierRequest;
+		TSharedPtr<UE::UAF::FInjectionRequest> GraphRequest;
+		TSharedPtr<UE::UAF::FInjectionRequest> ModifierRequest;
 		uint32 SerialNumber = 0;
 	};
 	
diff --git a/Engine/Plugins/Experimental/AnimNextAnimGraph/Source/AnimNextAnimGraph/Private/Injection/GraphInstanceInjectionComponent.h b/Engine/Plugins/Experimental/AnimNextAnimGraph/Source/AnimNextAnimGraph/Private/Injection/GraphInstanceInjectionComponent.h
--- a/Engine/Plugins/Experimental/AnimNextAnimGraph/Source/AnimNextAnimGraph/Private/Injection/GraphInstanceInjectionComponent.h
+++ b/Engine/Plugins/Experimental/AnimNextAnimGraph/Source/AnimNextAnimGraph/Private/Injection/GraphInstanceInjectionComponent.h
@@ -18,11 +18,11 @@ struct FGraphInstanceInjectionComponent : public FUAFGraphInstanceComponent
 
 	FGraphInstanceInjectionComponent();
 
-	const UE::AnimNext::FInjectionInfo& GetInjectionInfo() const
+	const UE::UAF::FInjectionInfo& GetInjectionInfo() const
 	{
 		return InjectionInfo;
 	}
 
 private:
-	UE::AnimNext::FInjectionInfo InjectionInfo;
+	UE::UAF::FInjectionInfo InjectionInfo;
 };
diff --git a/Engine/Plugins/Experimental/AnimNextAnimGraph/Source/AnimNextAnimGraph/Private/Injection/InjectionCallbackProxy.cpp b/Engine/Plugins/Experimental/AnimNextAnimGraph/Source/AnimNextAnimGraph/Private/Injection/InjectionCallbackProxy.cpp
--- a/Engine/Plugins/Experimental/AnimNextAnimGraph/Source/AnimNextAnimGraph/Private/Injection/InjectionCallbackProxy.cpp
+++ b/Engine/Plugins/Experimental/AnimNextAnimGraph/Source/AnimNextAnimGraph/Private/Injection/InjectionCallbackProxy.cpp
@@ -14,8 +14,8 @@ UInjectionCallbackProxy* UInjectionCallbackProxy::CreateProxyObjectForInjection(
 	FAnimNextVariableReference Site,
 	UObject* Object,
 	FAnimNextFactoryParams FactoryParams,
-	UE::AnimNext::FInjectionBlendSettings BlendInSettings,
-	UE::AnimNext::FInjectionBlendSettings BlendOutSettings)
+	UE::UAF::FInjectionBlendSettings BlendInSettings,
+	UE::UAF::FInjectionBlendSettings BlendOutSettings)
 {
 	UInjectionCallbackProxy* Proxy = NewObject<UInjectionCallbackProxy>();
 	Proxy->SetFlags(RF_StrongRefOnFrame);
@@ -28,27 +28,27 @@ bool UInjectionCallbackProxy::Inject(
 	FAnimNextVariableReference Site,
 	UObject* Object,
 	FAnimNextFactoryParams&& FactoryParams,
-	const UE::AnimNext::FInjectionBlendSettings& BlendInSettings,
-	const UE::AnimNext::FInjectionBlendSettings& BlendOutSettings)
+	const UE::UAF::FInjectionBlendSettings& BlendInSettings,
+	const UE::UAF::FInjectionBlendSettings& BlendOutSettings)
 {
 	if (AnimNextComponent == nullptr)
 	{
 		return false;
 	}
 
-	UE::AnimNext::FInjectionRequestArgs RequestArgs;
+	UE::UAF::FInjectionRequestArgs RequestArgs;
 	RequestArgs.Site = FAnimNextInjectionSite(Site);
 	RequestArgs.Object = Object;
 	RequestArgs.BlendInSettings = BlendInSettings;
 	RequestArgs.BlendOutSettings = BlendOutSettings;
 	RequestArgs.FactoryParams = MoveTemp(FactoryParams);
 
-	UE::AnimNext::FInjectionLifetimeEvents LifetimeEvents;
+	UE::UAF::FInjectionLifetimeEvents LifetimeEvents;
 	LifetimeEvents.OnCompleted.BindUObject(this, &UInjectionCallbackProxy::OnInjectionCompleted);
 	LifetimeEvents.OnInterrupted.BindUObject(this, &UInjectionCallbackProxy::OnInjectionInterrupted);
 	LifetimeEvents.OnBlendingOut.BindUObject(this, &UInjectionCallbackProxy::OnInjectionBlendingOut);
 
-	PlayingRequest = UE::AnimNext::FInjectionUtils::Inject(AnimNextComponent, AnimNextComponent->GetModuleHandle(), MoveTemp(RequestArgs), MoveTemp(LifetimeEvents));
+	PlayingRequest = UE::UAF::FInjectionUtils::Inject(AnimNextComponent, AnimNextComponent->GetModuleHandle(), MoveTemp(RequestArgs), MoveTemp(LifetimeEvents));
 
 	bWasInterrupted = false;
 
@@ -68,7 +68,7 @@ EUninjectionResult UInjectionCallbackProxy::Uninject()
 		return EUninjectionResult::Failed;
 	}
 
-	UE::AnimNext::FInjectionUtils::Uninject(PlayingRequest);
+	UE::UAF::FInjectionUtils::Uninject(PlayingRequest);
 
 	return EUninjectionResult::Succeeded;
 }
@@ -80,7 +80,7 @@ void UInjectionCallbackProxy::SetVariable(const FAnimNextVariableReference& Vari
 
 DEFINE_FUNCTION(UInjectionCallbackProxy::execSetVariable)
 {
-	using namespace UE::AnimNext;
+	using namespace UE::UAF;
 
 	if(!P_THIS->PlayingRequest.IsValid())
 	{
@@ -158,14 +158,14 @@ void UInjectionCallbackProxy::Cancel()
 	Uninject();
 }
 
-void UInjectionCallbackProxy::OnInjectionCompleted(const UE::AnimNext::FInjectionRequest& Request)
+void UInjectionCallbackProxy::OnInjectionCompleted(const UE::UAF::FInjectionRequest& Request)
 {
 	if (!bWasInterrupted)
 	{
-		const UE::AnimNext::EInjectionStatus Status = Request.GetStatus();
-		check(!EnumHasAnyFlags(Status, UE::AnimNext::EInjectionStatus::Interrupted));
+		const UE::UAF::EInjectionStatus Status = Request.GetStatus();
+		check(!EnumHasAnyFlags(Status, UE::UAF::EInjectionStatus::Interrupted));
 
-		if (EnumHasAnyFlags(Status, UE::AnimNext::EInjectionStatus::Expired))
+		if (EnumHasAnyFlags(Status, UE::UAF::EInjectionStatus::Expired))
 		{
 			OnInterrupted.Broadcast();
 		}
@@ -178,14 +178,14 @@ void UInjectionCallbackProxy::OnInjectionCompleted(const UE::AnimNext::FInjectio
 	Reset();
 }
 
-void UInjectionCallbackProxy::OnInjectionInterrupted(const UE::AnimNext::FInjectionRequest& Request)
+void UInjectionCallbackProxy::OnInjectionInterrupted(const UE::UAF::FInjectionRequest& Request)
 {
 	bWasInterrupted = true;
 
 	OnInterrupted.Broadcast();
 }
 
-void UInjectionCallbackProxy::OnInjectionBlendingOut(const UE::AnimNext::FInjectionRequest& Request)
+void UInjectionCallbackProxy::OnInjectionBlendingOut(const UE::UAF::FInjectionRequest& Request)
 {
 	if (!bWasInterrupted)
 	{
diff --git a/Engine/Plugins/Experimental/AnimNextAnimGraph/Source/AnimNextAnimGraph/Private/PlayAnim/PlayAnimCallbackProxy.cpp b/Engine/Plugins/Experimental/AnimNextAnimGraph/Source/AnimNextAnimGraph/Private/PlayAnim/PlayAnimCallbackProxy.cpp
--- a/Engine/Plugins/Experimental/AnimNextAnimGraph/Source/AnimNextAnimGraph/Private/PlayAnim/PlayAnimCallbackProxy.cpp
+++ b/Engine/Plugins/Experimental/AnimNextAnimGraph/Source/AnimNextAnimGraph/Private/PlayAnim/PlayAnimCallbackProxy.cpp
@@ -23,8 +23,8 @@ UPlayAnimCallbackProxy* UPlayAnimCallbackProxy::CreateProxyObjectForPlayAnim(
 	UAnimSequence* AnimSequence,
 	float PlayRate,
 	float StartPosition,
-	UE::AnimNext::FInjectionBlendSettings BlendInSettings,
-	UE::AnimNext::FInjectionBlendSettings BlendOutSettings)
+	UE::UAF::FInjectionBlendSettings BlendInSettings,
+	UE::UAF::FInjectionBlendSettings BlendOutSettings)
 {
 	UPlayAnimCallbackProxy* Proxy = NewObject<UPlayAnimCallbackProxy>();
 	Proxy->SetFlags(RF_StrongRefOnFrame);
@@ -38,21 +38,21 @@ bool UPlayAnimCallbackProxy::Play(
 	UAnimSequence* AnimSequence,
 	float PlayRate,
 	float StartPosition,
-	const UE::AnimNext::FInjectionBlendSettings& BlendInSettings,
-	const UE::AnimNext::FInjectionBlendSettings& BlendOutSettings)
+	const UE::UAF::FInjectionBlendSettings& BlendInSettings,
+	const UE::UAF::FInjectionBlendSettings& BlendOutSettings)
 {
 	bool bPlayedSuccessfully = false;
 	if (AnimNextComponent != nullptr)
 	{
-		UE::AnimNext::FInjectionLifetimeEvents LifetimeEvents;
+		UE::UAF::FInjectionLifetimeEvents LifetimeEvents;
 		LifetimeEvents.OnCompleted.BindUObject(this, &UPlayAnimCallbackProxy::OnPlayAnimCompleted);
 		LifetimeEvents.OnInterrupted.BindUObject(this, &UPlayAnimCallbackProxy::OnPlayAnimInterrupted);
 		LifetimeEvents.OnBlendingOut.BindUObject(this, &UPlayAnimCallbackProxy::OnPlayAnimBlendingOut);
 
 		PRAGMA_DISABLE_DEPRECATION_WARNINGS
-		PlayingRequest = UE::AnimNext::FInjectionUtils::PlayAnim(
+		PlayingRequest = UE::UAF::FInjectionUtils::PlayAnim(
 			AnimNextComponent,
-			UE::AnimNext::FInjectionSite(FAnimNextVariableReference(SiteName)),
+			UE::UAF::FInjectionSite(FAnimNextVariableReference(SiteName)),
 			AnimSequence,
 			{ 
 				.PlayRate = PlayRate,
@@ -74,14 +74,14 @@ bool UPlayAnimCallbackProxy::Play(
 	return bPlayedSuccessfully;
 }
 
-void UPlayAnimCallbackProxy::OnPlayAnimCompleted(const UE::AnimNext::FInjectionRequest& Request)
+void UPlayAnimCallbackProxy::OnPlayAnimCompleted(const UE::UAF::FInjectionRequest& Request)
 {
 	if (!bWasInterrupted)
 	{
-		const UE::AnimNext::EInjectionStatus Status = Request.GetStatus();
-		check(!EnumHasAnyFlags(Status, UE::AnimNext::EInjectionStatus::Interrupted));
+		const UE::UAF::EInjectionStatus Status = Request.GetStatus();
+		check(!EnumHasAnyFlags(Status, UE::UAF::EInjectionStatus::Interrupted));
 
-		if (EnumHasAnyFlags(Status, UE::AnimNext::EInjectionStatus::Expired))
+		if (EnumHasAnyFlags(Status, UE::UAF::EInjectionStatus::Expired))
 		{
 			OnInterrupted.Broadcast();
 		}
@@ -94,14 +94,14 @@ void UPlayAnimCallbackProxy::OnPlayAnimCompleted(const UE::AnimNext::FInjectionR
 	Reset();
 }
 
-void UPlayAnimCallbackProxy::OnPlayAnimInterrupted(const UE::AnimNext::FInjectionRequest& Request)
+void UPlayAnimCallbackProxy::OnPlayAnimInterrupted(const UE::UAF::FInjectionRequest& Request)
 {
 	bWasInterrupted = true;
 
 	OnInterrupted.Broadcast();
 }
 
-void UPlayAnimCallbackProxy::OnPlayAnimBlendingOut(const UE::AnimNext::FInjectionRequest& Request)
+void UPlayAnimCallbackProxy::OnPlayAnimBlendingOut(const UE::UAF::FInjectionRequest& Request)
 {
 	if (!bWasInterrupted)
 	{
diff --git a/Engine/Plugins/Experimental/AnimNextAnimGraph/Source/AnimNextAnimGraph/Private/TraitCore/TraitHandle.cpp b/Engine/Plugins/Experimental/AnimNextAnimGraph/Source/AnimNextAnimGraph/Private/TraitCore/TraitHandle.cpp
--- a/Engine/Plugins/Experimental/AnimNextAnimGraph/Source/AnimNextAnimGraph/Private/TraitCore/TraitHandle.cpp
+++ b/Engine/Plugins/Experimental/AnimNextAnimGraph/Source/AnimNextAnimGraph/Private/TraitCore/TraitHandle.cpp
@@ -13,7 +13,7 @@ bool FAnimNextTraitHandle::Serialize(FArchive& Ar)
 		// On load, we hold a node ID that needs fix-up
 		check(GetNodeHandle().IsNodeID());
 
-		UE::AnimNext::FTraitReader& TraitReader = static_cast<UE::AnimNext::FTraitReader&>(Ar);
+		UE::UAF::FTraitReader& TraitReader = static_cast<UE::UAF::FTraitReader&>(Ar);
 		*this = TraitReader.ResolveTraitHandle(*this);
 		check(GetNodeHandle().IsSharedOffset());
 	}
diff --git a/Engine/Plugins/Experimental/AnimNextAnimGraph/Source/AnimNextAnimGraph/Private/TraitInterfaces/IUpdate.cpp b/Engine/Plugins/Experimental/AnimNextAnimGraph/Source/AnimNextAnimGraph/Private/TraitInterfaces/IUpdate.cpp
--- a/Engine/Plugins/Experimental/AnimNextAnimGraph/Source/AnimNextAnimGraph/Private/TraitInterfaces/IUpdate.cpp
+++ b/Engine/Plugins/Experimental/AnimNextAnimGraph/Source/AnimNextAnimGraph/Private/TraitInterfaces/IUpdate.cpp
@@ -8,7 +8,7 @@
 #include "Graph/AnimNextGraphInstance.h"
 #include "TraitCore/TraitEventRaising.h"
 
-namespace UE::AnimNext
+namespace UE::UAF
 {
 	namespace Private
 	{
@@ -125,7 +125,7 @@ namespace UE::AnimNext
 		}
 
 		// Raises every queued event from the list on the specified entry
-		static void RaiseTraitEvents(FUpdateTraversalContext& Context, Private::FUpdateEntry* UpdateEntry, const UE::AnimNext::FTraitEventList& EventList)
+		static void RaiseTraitEvents(FUpdateTraversalContext& Context, Private::FUpdateEntry* UpdateEntry, const UE::UAF::FTraitEventList& EventList)
 		{
 			// TODO: Performance note
 			// 
@@ -152,7 +152,7 @@ namespace UE::AnimNext
 			// function to avoid the virtual call we have for GetTypeUID.
 			// 
 
-			UE::AnimNext::RaiseTraitEvents(Context, UpdateEntry->TraitStack, EventList);
+			UE::UAF::RaiseTraitEvents(Context, UpdateEntry->TraitStack, EventList);
 		}
 	}
 
diff --git a/Engine/Plugins/Experimental/AnimNextAnimGraph/Source/AnimNextAnimGraph/Private/Traits/ApplyAdditiveTrait.cpp b/Engine/Plugins/Experimental/AnimNextAnimGraph/Source/AnimNextAnimGraph/Private/Traits/ApplyAdditiveTrait.cpp
--- a/Engine/Plugins/Experimental/AnimNextAnimGraph/Source/AnimNextAnimGraph/Private/Traits/ApplyAdditiveTrait.cpp
+++ b/Engine/Plugins/Experimental/AnimNextAnimGraph/Source/AnimNextAnimGraph/Private/Traits/ApplyAdditiveTrait.cpp
@@ -7,7 +7,7 @@
 #include "TraitCore/ExecutionContext.h"
 #include "EvaluationVM/Tasks/ApplyAdditiveKeyframe.h"
 
-namespace UE::AnimNext
+namespace UE::UAF
 {
 	AUTO_REGISTER_ANIM_TRAIT(FApplyAdditiveTrait)
 
@@ -167,4 +167,4 @@ namespace UE::AnimNext
 		}
 	}
 
-}  // UE::AnimNext
+}  // UE::UAF
diff --git a/Engine/Plugins/Experimental/AnimNextAnimGraph/Source/AnimNextAnimGraph/Private/Traits/BlendSpacePlayer.cpp b/Engine/Plugins/Experimental/AnimNextAnimGraph/Source/AnimNextAnimGraph/Private/Traits/BlendSpacePlayer.cpp
--- a/Engine/Plugins/Experimental/AnimNextAnimGraph/Source/AnimNextAnimGraph/Private/Traits/BlendSpacePlayer.cpp
+++ b/Engine/Plugins/Experimental/AnimNextAnimGraph/Source/AnimNextAnimGraph/Private/Traits/BlendSpacePlayer.cpp
@@ -14,7 +14,7 @@
 #include "TraitInterfaces/IGraphFactory.h"
 
 
-namespace UE::AnimNext
+namespace UE::UAF
 {
 	AUTO_REGISTER_ANIM_TRAIT(FBlendSpacePlayerTrait)
 
@@ -300,8 +300,8 @@ namespace UE::AnimNext
 				const FSampleGraphState& SampleState = InstanceData->SampleGraphs[HighestWeightIndex];
 				FTraitStackBinding ChildTraitStack;
 				ensure(Context.GetStack(SampleState.ChildPtr, ChildTraitStack));
-				UE::AnimNext::TTraitBinding<UE::AnimNext::ITimeline> Timeline;
-				if (ChildTraitStack.GetInterface<UE::AnimNext::ITimeline>(Timeline))
+				UE::UAF::TTraitBinding<UE::UAF::ITimeline> Timeline;
+				if (ChildTraitStack.GetInterface<UE::UAF::ITimeline>(Timeline))
 				{
 					return Timeline.GetState(Context);
 				}
diff --git a/Engine/Plugins/Experimental/AnimNextAnimGraph/Source/AnimNextAnimGraph/Private/Traits/CallFunction.cpp b/Engine/Plugins/Experimental/AnimNextAnimGraph/Source/AnimNextAnimGraph/Private/Traits/CallFunction.cpp
--- a/Engine/Plugins/Experimental/AnimNextAnimGraph/Source/AnimNextAnimGraph/Private/Traits/CallFunction.cpp
+++ b/Engine/Plugins/Experimental/AnimNextAnimGraph/Source/AnimNextAnimGraph/Private/Traits/CallFunction.cpp
@@ -13,7 +13,7 @@
 
 DEFINE_STAT(STAT_AnimNext_CallFunction);
 
-namespace UE::AnimNext
+namespace UE::UAF
 {
 	AUTO_REGISTER_ANIM_TRAIT(FCallFunctionTrait)
 
@@ -76,7 +76,7 @@ namespace UE::AnimNext
 		AnimNextContext.SetDeltaTime(InTraitState.GetDeltaTime());
 
 		FAnimNextGraphContextData ContextData(GraphInstance.GetModuleInstance(), &GraphInstance);
-		UE::AnimNext::FScopedExecuteContextData ContextDataScope(AnimNextContext, ContextData);
+		UE::UAF::FScopedExecuteContextData ContextDataScope(AnimNextContext, ContextData);
 
 		VM->ExecuteVM(ExtendedExecuteContext, SharedData.FunctionEvent);
 	}
diff --git a/Engine/Plugins/Experimental/AnimNextAnimGraph/Source/AnimNextAnimGraph/Private/Traits/InertializationDeadBlend.cpp b/Engine/Plugins/Experimental/AnimNextAnimGraph/Source/AnimNextAnimGraph/Private/Traits/InertializationDeadBlend.cpp
--- a/Engine/Plugins/Experimental/AnimNextAnimGraph/Source/AnimNextAnimGraph/Private/Traits/InertializationDeadBlend.cpp
+++ b/Engine/Plugins/Experimental/AnimNextAnimGraph/Source/AnimNextAnimGraph/Private/Traits/InertializationDeadBlend.cpp
@@ -10,7 +10,7 @@
 
 #include "Curves/CurveFloat.h"
 
-namespace UE::AnimNext
+namespace UE::UAF
 {
 	AUTO_REGISTER_ANIM_TRAIT(FInertializationDeadBlendTrait)
 
@@ -75,7 +75,7 @@ namespace UE::AnimNext
                 // or we cannot handle it so we effectively ignore the request
 
                 InstanceData->bRequestPending = false;
-                InstanceData->PendingRequest = UE::AnimNext::FInertializationRequest();
+                InstanceData->PendingRequest = UE::UAF::FInertializationRequest();
             }
             else if (InstanceData->PrevPose.IsEmpty())
             {
@@ -85,7 +85,7 @@ namespace UE::AnimNext
 
                 check(!InstanceData->CurrPose.IsEmpty());
 
-				UE::AnimNext::FDeadBlendTransitionTaskParameters Parameters;
+				UE::UAF::FDeadBlendTransitionTaskParameters Parameters;
 				Parameters.ExtrapolationHalfLife = SharedData->ExtrapolationHalfLife;
 				Parameters.ExtrapolationHalfLifeMin = SharedData->ExtrapolationHalfLifeMin;
 				Parameters.ExtrapolationHalfLifeMax = SharedData->ExtrapolationHalfLifeMax;
@@ -101,7 +101,7 @@ namespace UE::AnimNext
                 InstanceData->bRequestActive = true;
                 InstanceData->ActiveRequest = InstanceData->PendingRequest;
                 InstanceData->bRequestPending = false;
-                InstanceData->PendingRequest = UE::AnimNext::FInertializationRequest();
+                InstanceData->PendingRequest = UE::UAF::FInertializationRequest();
                 InstanceData->TimeSinceTransition = 0.0f;
             }
             else
@@ -111,7 +111,7 @@ namespace UE::AnimNext
                 check(!InstanceData->CurrPose.IsEmpty());
                 check(!InstanceData->PrevPose.IsEmpty());
 
-				UE::AnimNext::FDeadBlendTransitionTaskParameters Parameters;
+				UE::UAF::FDeadBlendTransitionTaskParameters Parameters;
 				Parameters.ExtrapolationHalfLife = SharedData->ExtrapolationHalfLife;
 				Parameters.ExtrapolationHalfLifeMin = SharedData->ExtrapolationHalfLifeMin;
 				Parameters.ExtrapolationHalfLifeMax = SharedData->ExtrapolationHalfLifeMax;
@@ -129,7 +129,7 @@ namespace UE::AnimNext
                 InstanceData->bRequestActive = true;
                 InstanceData->ActiveRequest = InstanceData->PendingRequest;
                 InstanceData->bRequestPending = false;
-                InstanceData->PendingRequest = UE::AnimNext::FInertializationRequest();
+                InstanceData->PendingRequest = UE::UAF::FInertializationRequest();
                 InstanceData->TimeSinceTransition = 0.0f;
             }
         } 
@@ -152,7 +152,7 @@ namespace UE::AnimNext
 				// Deactivate Dead Blending since transition is complete
 
 				InstanceData->bRequestActive = false;
-				InstanceData->ActiveRequest = UE::AnimNext::FInertializationRequest();
+				InstanceData->ActiveRequest = UE::UAF::FInertializationRequest();
 				InstanceData->TimeSinceTransition = 0.0f;
 				InstanceData->State.Empty();
 			}
diff --git a/Engine/Plugins/Experimental/AnimNextAnimGraph/Source/AnimNextAnimGraph/Private/Traits/ModifyCurveTrait.cpp b/Engine/Plugins/Experimental/AnimNextAnimGraph/Source/AnimNextAnimGraph/Private/Traits/ModifyCurveTrait.cpp
--- a/Engine/Plugins/Experimental/AnimNextAnimGraph/Source/AnimNextAnimGraph/Private/Traits/ModifyCurveTrait.cpp
+++ b/Engine/Plugins/Experimental/AnimNextAnimGraph/Source/AnimNextAnimGraph/Private/Traits/ModifyCurveTrait.cpp
@@ -4,7 +4,7 @@
 
 #include "EvaluationVM/EvaluationVM.h"
 
-namespace UE::AnimNext
+namespace UE::UAF
 {
 AUTO_REGISTER_ANIM_TRAIT(FModifyCurveTrait)
 
@@ -16,7 +16,7 @@ GENERATE_ANIM_TRAIT_IMPLEMENTATION(FModifyCurveTrait, TRAIT_INTERFACE_ENUMERATOR
 #undef TRAIT_INTERFACE_ENUMERATOR
 }
 
-void UE::AnimNext::FModifyCurveTrait::PostEvaluate(FEvaluateTraversalContext& Context, const TTraitBinding<IEvaluate>& Binding) const
+void UE::UAF::FModifyCurveTrait::PostEvaluate(FEvaluateTraversalContext& Context, const TTraitBinding<IEvaluate>& Binding) const
 {
 	IEvaluate::PostEvaluate(Context, Binding);
 
@@ -36,8 +36,8 @@ void UE::AnimNext::FModifyCurveTrait::PostEvaluate(FEvaluateTraversalContext& Co
 	Context.AppendTask(FModifyCurveTask::Make(InstanceData, SharedData));
 }
 
-FModifyCurveTask FModifyCurveTask::Make(UE::AnimNext::FModifyCurveTrait::FInstanceData* InstanceData,
-	const UE::AnimNext::FModifyCurveTrait::FSharedData* SharedData)
+FModifyCurveTask FModifyCurveTask::Make(UE::UAF::FModifyCurveTrait::FInstanceData* InstanceData,
+	const UE::UAF::FModifyCurveTrait::FSharedData* SharedData)
 {
 	FModifyCurveTask Task;
 	Task.InstanceData = InstanceData;
@@ -45,9 +45,9 @@ FModifyCurveTask FModifyCurveTask::Make(UE::AnimNext::FModifyCurveTrait::FInstan
 	return Task;
 }
 
-void FModifyCurveTask::Execute(UE::AnimNext::FEvaluationVM& VM) const
+void FModifyCurveTask::Execute(UE::UAF::FEvaluationVM& VM) const
 {
-	using namespace UE::AnimNext;
+	using namespace UE::UAF;
 	
 	if (const TUniquePtr<FKeyframeState>* Keyframe = VM.PeekValue<TUniquePtr<FKeyframeState>>(KEYFRAME_STACK_NAME, 0))
 	{
diff --git a/Engine/Plugins/Experimental/AnimNextAnimGraph/Source/AnimNextAnimGraph/Private/Traits/ModifyCurveTrait.h b/Engine/Plugins/Experimental/AnimNextAnimGraph/Source/AnimNextAnimGraph/Private/Traits/ModifyCurveTrait.h
--- a/Engine/Plugins/Experimental/AnimNextAnimGraph/Source/AnimNextAnimGraph/Private/Traits/ModifyCurveTrait.h
+++ b/Engine/Plugins/Experimental/AnimNextAnimGraph/Source/AnimNextAnimGraph/Private/Traits/ModifyCurveTrait.h
@@ -60,7 +60,7 @@ struct FAnimNextModifyCurveTraitSharedData : public FAnimNextTraitSharedData
 #undef TRAIT_LATENT_PROPERTIES_ENUMERATOR
 };
 
-namespace UE::AnimNext
+namespace UE::UAF
 {
 
 	/**
@@ -96,14 +96,14 @@ struct FModifyCurveTask : public FAnimNextEvaluationTask
 
 	DECLARE_ANIM_EVALUATION_TASK(FModifyCurveTask)
 
-	static FModifyCurveTask Make(UE::AnimNext::FModifyCurveTrait::FInstanceData* InstanceData, const UE::AnimNext::FModifyCurveTrait::FSharedData* SharedData);
+	static FModifyCurveTask Make(UE::UAF::FModifyCurveTrait::FInstanceData* InstanceData, const UE::UAF::FModifyCurveTrait::FSharedData* SharedData);
 
 	// Task entry point
-	virtual void Execute(UE::AnimNext::FEvaluationVM& VM) const override;
+	virtual void Execute(UE::UAF::FEvaluationVM& VM) const override;
 
 private:
 	static float ProcessCurveOperation(float CurrentValue, float NewValue, float Alpha, EAnimNext_ModifyCurveApplyMode ApplyMode);
 
-	UE::AnimNext::FModifyCurveTrait::FInstanceData* InstanceData = nullptr;
-	const UE::AnimNext::FModifyCurveTrait::FSharedData* SharedData = nullptr;
+	UE::UAF::FModifyCurveTrait::FInstanceData* InstanceData = nullptr;
+	const UE::UAF::FModifyCurveTrait::FSharedData* SharedData = nullptr;
 };
diff --git a/Engine/Plugins/Experimental/AnimNextAnimGraph/Source/AnimNextAnimGraph/Private/Traits/NotifyDispatcher.cpp b/Engine/Plugins/Experimental/AnimNextAnimGraph/Source/AnimNextAnimGraph/Private/Traits/NotifyDispatcher.cpp
--- a/Engine/Plugins/Experimental/AnimNextAnimGraph/Source/AnimNextAnimGraph/Private/Traits/NotifyDispatcher.cpp
+++ b/Engine/Plugins/Experimental/AnimNextAnimGraph/Source/AnimNextAnimGraph/Private/Traits/NotifyDispatcher.cpp
@@ -13,7 +13,7 @@
 #include "Module/AnimNextModuleInstance.h"
 #include "Module/AnimNextSkeletalMeshComponentReferenceComponent.h"
 
-namespace UE::AnimNext
+namespace UE::UAF
 {
 	AUTO_REGISTER_ANIM_TRAIT(FNotifyDispatcherTrait)
 
@@ -149,7 +149,7 @@ FAnimNextNotifyDispatcherComponent::FAnimNextNotifyDispatcherComponent()
 
 void FAnimNextNotifyDispatcherComponent::OnTraitEvent(FAnimNextTraitEvent& Event)
 {
-	if(UE::AnimNext::FNotifyDispatchEvent* NotifyDispatchEvent = Event.AsType<UE::AnimNext::FNotifyDispatchEvent>())
+	if(UE::UAF::FNotifyDispatchEvent* NotifyDispatchEvent = Event.AsType<UE::UAF::FNotifyDispatchEvent>())
 	{
 		NotifyQueue.AddAnimNotifies(NotifyDispatchEvent->Notifies, NotifyDispatchEvent->Weight);
 		NotifyDispatchEvent->MarkConsumed();
@@ -158,7 +158,7 @@ void FAnimNextNotifyDispatcherComponent::OnTraitEvent(FAnimNextTraitEvent& Event
 
 void FAnimNextNotifyDispatcherComponent::OnEndExecution(float InDeltaTime)
 {
-	using namespace UE::AnimNext;
+	using namespace UE::UAF;
 
 	// Early out if we dont need to do any work
 	if(NotifyQueue.AnimNotifies.Num() == 0 && ActiveAnimNotifyState.Num() == 0)
@@ -291,7 +291,7 @@ void FAnimNextNotifyDispatcherComponent::OnEndExecution(float InDeltaTime)
 	NotifyQueue.AnimNotifies.Reset();
 }
 
-void FAnimNextNotifyDispatcherComponent::TriggerSingleAnimNotify(float InDeltaTime, UE::AnimNext::FNotifyQueueDispatchEvent* InDispatcher, const FAnimNotifyEventReference& EventReference)
+void FAnimNextNotifyDispatcherComponent::TriggerSingleAnimNotify(float InDeltaTime, UE::UAF::FNotifyQueueDispatchEvent* InDispatcher, const FAnimNotifyEventReference& EventReference)
 {
 	// This is for non 'state' anim notifies.
 	const FAnimNotifyEvent* AnimNotifyEvent = EventReference.GetNotify();
diff --git a/Engine/Plugins/Experimental/AnimNextAnimGraph/Source/AnimNextAnimGraph/Private/Traits/NotifyDispatcher.h b/Engine/Plugins/Experimental/AnimNextAnimGraph/Source/AnimNextAnimGraph/Private/Traits/NotifyDispatcher.h
--- a/Engine/Plugins/Experimental/AnimNextAnimGraph/Source/AnimNextAnimGraph/Private/Traits/NotifyDispatcher.h
+++ b/Engine/Plugins/Experimental/AnimNextAnimGraph/Source/AnimNextAnimGraph/Private/Traits/NotifyDispatcher.h
@@ -16,7 +16,7 @@ struct FAnimNotifyEventReference;
 class USkeletalMeshComponent;
 struct FAnimNextModuleInstance;
 
-namespace UE::AnimNext
+namespace UE::UAF
 {
 	struct FNotifyQueueDispatchEvent;
 }
@@ -33,7 +33,7 @@ struct FAnimNextNotifyDispatcherComponent : public FUAFModuleInstanceComponent
 	virtual void OnEndExecution(float InDeltaTime) override;
 
 	// Triggers a single anim notify in the dispatcher
-	void TriggerSingleAnimNotify(float InDeltaTime, UE::AnimNext::FNotifyQueueDispatchEvent* InDispatcher, const FAnimNotifyEventReference& EventReference);
+	void TriggerSingleAnimNotify(float InDeltaTime, UE::UAF::FNotifyQueueDispatchEvent* InDispatcher, const FAnimNotifyEventReference& EventReference);
 
 	// Notify queue to dispatch
 	UPROPERTY()
@@ -49,7 +49,7 @@ struct FAnimNextNotifyDispatcherComponent : public FUAFModuleInstanceComponent
 	TObjectPtr<USkeletalMeshComponent> SkeletalMeshComponent;
 };
 
-namespace UE::AnimNext
+namespace UE::UAF
 {
 	/**
 	 * FNotifyDispatcherTrait
diff --git a/Engine/Plugins/Experimental/AnimNextAnimGraph/Source/AnimNextAnimGraph/Public/EvaluationVM/EvaluationTask.h b/Engine/Plugins/Experimental/AnimNextAnimGraph/Source/AnimNextAnimGraph/Public/EvaluationVM/EvaluationTask.h
--- a/Engine/Plugins/Experimental/AnimNextAnimGraph/Source/AnimNextAnimGraph/Public/EvaluationVM/EvaluationTask.h
+++ b/Engine/Plugins/Experimental/AnimNextAnimGraph/Source/AnimNextAnimGraph/Public/EvaluationVM/EvaluationTask.h
@@ -6,7 +6,7 @@
 
 #include "EvaluationTask.generated.h"
 
-namespace UE::AnimNext
+namespace UE::UAF
 {
 	struct FEvaluationVM;
 }
@@ -35,5 +35,5 @@ struct FAnimNextEvaluationTask
 	virtual const UScriptStruct* GetStruct() const { return FAnimNextEvaluationTask::StaticStruct(); }
 
 	// Task entry point
-	virtual void Execute(UE::AnimNext::FEvaluationVM& VM) const { LowLevelFatalError(TEXT("Pure virtual not implemented (FAnimNextEvaluationTask::Execute)")); }
+	virtual void Execute(UE::UAF::FEvaluationVM& VM) const { LowLevelFatalError(TEXT("Pure virtual not implemented (FAnimNextEvaluationTask::Execute)")); }
 };
diff --git a/Engine/Plugins/Experimental/AnimNextAnimGraph/Source/AnimNextAnimGraph/Public/EvaluationVM/EvaluationVM.h b/Engine/Plugins/Experimental/AnimNextAnimGraph/Source/AnimNextAnimGraph/Public/EvaluationVM/EvaluationVM.h
--- a/Engine/Plugins/Experimental/AnimNextAnimGraph/Source/AnimNextAnimGraph/Public/EvaluationVM/EvaluationVM.h
+++ b/Engine/Plugins/Experimental/AnimNextAnimGraph/Source/AnimNextAnimGraph/Public/EvaluationVM/EvaluationVM.h
@@ -15,8 +15,8 @@
 #define UE_API ANIMNEXTANIMGRAPH_API
 
 // If you wish to use a custom type with a stack in the FEvaluationVM, you'll have to enable
-// its usage through this macro to ensure that UE::AnimNext::GetTypeID() is declared.
-// Use this macro inside the UE::AnimNext namespace
+// its usage through this macro to ensure that UE::UAF::GetTypeID() is declared.
+// Use this macro inside the UE::UAF namespace
 #define ANIM_NEXT_ENABLE_EVALUATION_STACK_USAGE(Type) \
 	template<> \
 	constexpr uint32 GetTypeID<Type>() \
@@ -24,7 +24,7 @@
 		return UE::HashStringFNV1a32(#Type); \
 	}
 
-namespace UE::AnimNext
+namespace UE::UAF
 {
 	// Helper function that returns a UID for the specified type
 	// @see ANIM_NEXT_ENABLE_EVALUATION_STACK_USAGE
@@ -147,7 +147,7 @@ namespace UE::AnimNext
 	struct FEvaluationVM
 	{
 		FEvaluationVM() = default;
-		UE_API FEvaluationVM(EEvaluationFlags InEvaluationFlags, const UE::AnimNext::FReferencePose& InReferencePose, int32 InCurrentLOD);
+		UE_API FEvaluationVM(EEvaluationFlags InEvaluationFlags, const UE::UAF::FReferencePose& InReferencePose, int32 InCurrentLOD);
 
 		// Allow move
 		FEvaluationVM(FEvaluationVM&&) = default;
@@ -211,7 +211,7 @@ namespace UE::AnimNext
 		TMap<FName, FEvaluationVMStack> InternalStacks;
 
 		// Reference pose
-		const UE::AnimNext::FReferencePose* ReferencePose = nullptr;
+		const UE::UAF::FReferencePose* ReferencePose = nullptr;
 
 		// Current LOD we are evaluating at
 		int32 CurrentLOD = 0;
diff --git a/Engine/Plugins/Experimental/AnimNextAnimGraph/Source/AnimNextAnimGraph/Public/EvaluationVM/SerializableEvaluationProgram.h b/Engine/Plugins/Experimental/AnimNextAnimGraph/Source/AnimNextAnimGraph/Public/EvaluationVM/SerializableEvaluationProgram.h
--- a/Engine/Plugins/Experimental/AnimNextAnimGraph/Source/AnimNextAnimGraph/Public/EvaluationVM/SerializableEvaluationProgram.h
+++ b/Engine/Plugins/Experimental/AnimNextAnimGraph/Source/AnimNextAnimGraph/Public/EvaluationVM/SerializableEvaluationProgram.h
@@ -7,7 +7,7 @@
 
 #include "SerializableEvaluationProgram.generated.h"
 
-namespace UE::AnimNext
+namespace UE::UAF
 {
 	struct FEvaluationProgram;
 }
@@ -21,7 +21,7 @@ struct FSerializableEvaluationProgram
 	{
 	}
 	
-	FSerializableEvaluationProgram(const UE::AnimNext::FEvaluationProgram& Other);
+	FSerializableEvaluationProgram(const UE::UAF::FEvaluationProgram& Other);
 	
 	UPROPERTY(EditAnywhere, Category = "Tasks", meta = (ExpandByDefault))
 	TArray<FInstancedStruct> Tasks;	
diff --git a/Engine/Plugins/Experimental/AnimNextAnimGraph/Source/AnimNextAnimGraph/Public/EvaluationVM/Tasks/ApplyAdditiveKeyframe.h b/Engine/Plugins/Experimental/AnimNextAnimGraph/Source/AnimNextAnimGraph/Public/EvaluationVM/Tasks/ApplyAdditiveKeyframe.h
--- a/Engine/Plugins/Experimental/AnimNextAnimGraph/Source/AnimNextAnimGraph/Public/EvaluationVM/Tasks/ApplyAdditiveKeyframe.h
+++ b/Engine/Plugins/Experimental/AnimNextAnimGraph/Source/AnimNextAnimGraph/Public/EvaluationVM/Tasks/ApplyAdditiveKeyframe.h
@@ -11,7 +11,7 @@
 
 struct FInputScaleBiasClamp;
 
-namespace UE::AnimNext
+namespace UE::UAF
 {
 	struct FKeyframeState;
 }
@@ -34,11 +34,11 @@ struct FAnimNextApplyAdditiveKeyframeTask : public FAnimNextEvaluationTask
 	static UE_API FAnimNextApplyAdditiveKeyframeTask Make(const FName& AlphaSourceCurveName, const int8 AlphaCurveInputIndex, TFunction<float(float)> InputScaleBiasClampFn);
 
 	// Task entry point
-	UE_API virtual void Execute(UE::AnimNext::FEvaluationVM& VM) const override;
+	UE_API virtual void Execute(UE::UAF::FEvaluationVM& VM) const override;
 
 	// @TODO: Move this into shared method
 	// Compute the current alpha value for curves or return the interpolation alpha otherwise
-	UE_API float GetInterpolationAlpha(const UE::AnimNext::FKeyframeState* KeyframeA, const UE::AnimNext::FKeyframeState* KeyframeB) const;
+	UE_API float GetInterpolationAlpha(const UE::UAF::FKeyframeState* KeyframeA, const UE::UAF::FKeyframeState* KeyframeB) const;
 
 	// How much weight between the additive identity and the additive pose to apply: lerp(identity, additive, weight)
 	UPROPERTY()
diff --git a/Engine/Plugins/Experimental/AnimNextAnimGraph/Source/AnimNextAnimGraph/Public/EvaluationVM/Tasks/BlendKeyframes.h b/Engine/Plugins/Experimental/AnimNextAnimGraph/Source/AnimNextAnimGraph/Public/EvaluationVM/Tasks/BlendKeyframes.h
--- a/Engine/Plugins/Experimental/AnimNextAnimGraph/Source/AnimNextAnimGraph/Public/EvaluationVM/Tasks/BlendKeyframes.h
+++ b/Engine/Plugins/Experimental/AnimNextAnimGraph/Source/AnimNextAnimGraph/Public/EvaluationVM/Tasks/BlendKeyframes.h
@@ -11,7 +11,7 @@
 
 struct FInputScaleBiasClamp;
 
-namespace UE::AnimNext
+namespace UE::UAF
 {
 	struct FKeyframeState;
 }
@@ -48,10 +48,10 @@ struct FAnimNextBlendTwoKeyframesTask : public FAnimNextEvaluationTask
 	
 
 	// Task entry point
-	UE_API virtual void Execute(UE::AnimNext::FEvaluationVM& VM) const override;
+	UE_API virtual void Execute(UE::UAF::FEvaluationVM& VM) const override;
 
 	// Compute the current alpha value for curves or return the interpolation alpha otherwise
-	UE_API float GetInterpolationAlpha(const UE::AnimNext::FKeyframeState* KeyframeA, const UE::AnimNext::FKeyframeState* KeyframeB) const;
+	UE_API float GetInterpolationAlpha(const UE::UAF::FKeyframeState* KeyframeA, const UE::UAF::FKeyframeState* KeyframeB) const;
 
 	// The interpolation alpha between the two input keyframes
 	UPROPERTY(VisibleAnywhere, Category = Properties)
@@ -84,7 +84,7 @@ struct FAnimNextBlendOverwriteKeyframeWithScaleTask : public FAnimNextEvaluation
 	static UE_API FAnimNextBlendOverwriteKeyframeWithScaleTask Make(float ScaleFactor);
 
 	// Task entry point
-	UE_API virtual void Execute(UE::AnimNext::FEvaluationVM& VM) const override;
+	UE_API virtual void Execute(UE::UAF::FEvaluationVM& VM) const override;
 
 	// The scale factor to apply to the keyframe
 	UPROPERTY(VisibleAnywhere, Category = Properties)
@@ -111,7 +111,7 @@ struct FAnimNextBlendAddKeyframeWithScaleTask : public FAnimNextEvaluationTask
 	static UE_API FAnimNextBlendAddKeyframeWithScaleTask Make(float ScaleFactor);
 
 	// Task entry point
-	UE_API virtual void Execute(UE::AnimNext::FEvaluationVM& VM) const override;
+	UE_API virtual void Execute(UE::UAF::FEvaluationVM& VM) const override;
 
 	// The scale factor to apply to the keyframe
 	UPROPERTY(VisibleAnywhere, Category = Properties)
diff --git a/Engine/Plugins/Experimental/AnimNextAnimGraph/Source/AnimNextAnimGraph/Public/EvaluationVM/Tasks/BlendKeyframesPerBone.h b/Engine/Plugins/Experimental/AnimNextAnimGraph/Source/AnimNextAnimGraph/Public/EvaluationVM/Tasks/BlendKeyframesPerBone.h
--- a/Engine/Plugins/Experimental/AnimNextAnimGraph/Source/AnimNextAnimGraph/Public/EvaluationVM/Tasks/BlendKeyframesPerBone.h
+++ b/Engine/Plugins/Experimental/AnimNextAnimGraph/Source/AnimNextAnimGraph/Public/EvaluationVM/Tasks/BlendKeyframesPerBone.h
@@ -38,7 +38,7 @@ struct FAnimNextBlendOverwriteKeyframePerBoneWithScaleTask : public FAnimNextBle
 	static UE_API FAnimNextBlendOverwriteKeyframePerBoneWithScaleTask Make(const IBlendProfileInterface* BlendProfile, const FBlendSampleData& BlendData, float ScaleFactor);
 
 	// Task entry point
-	UE_API virtual void Execute(UE::AnimNext::FEvaluationVM& VM) const override;
+	UE_API virtual void Execute(UE::UAF::FEvaluationVM& VM) const override;
 
 	// The blend data associated with the keyframe to overwrite
 	// Cannot use a UPROPERTY with a pointer to a struct
@@ -73,7 +73,7 @@ struct FAnimNextBlendAddKeyframePerBoneWithScaleTask : public FAnimNextBlendAddK
 	static UE_API FAnimNextBlendAddKeyframePerBoneWithScaleTask Make(const IBlendProfileInterface* BlendProfile, const FBlendSampleData& BlendDataA, const FBlendSampleData& BlendDataB, float ScaleFactor);
 
 	// Task entry point
-	UE_API virtual void Execute(UE::AnimNext::FEvaluationVM& VM) const override;
+	UE_API virtual void Execute(UE::UAF::FEvaluationVM& VM) const override;
 
 	// The blend data associated with the keyframe A
 	// Cannot use a UPROPERTY with a pointer to a struct
@@ -113,7 +113,7 @@ struct FAnimNextBlendKeyframePerBoneWithScaleTask : public FAnimNextBlendAddKeyf
 	static UE_API FAnimNextBlendKeyframePerBoneWithScaleTask Make(TSharedPtr<IInterpolationIndexProvider> BlendProfile, const USkeleton* Skeleton, const TArray<float>& BoneMaskWeights, const UE::Anim::TNamedValueArray<FDefaultAllocator, UE::Anim::FCurveElement>& CurveMaskWeights, const TArray<FMaskedAttributeWeight>& AttributeMaskWeights, float ScaleFactor);
 
 	// Task entry point
-	UE_API virtual void Execute(UE::AnimNext::FEvaluationVM& VM) const override;
+	UE_API virtual void Execute(UE::UAF::FEvaluationVM& VM) const override;
 
 private:
 	TArray<float> BoneMaskWeights;
diff --git a/Engine/Plugins/Experimental/AnimNextAnimGraph/Source/AnimNextAnimGraph/Public/EvaluationVM/Tasks/ConvertRotationsToLocalSpace.h b/Engine/Plugins/Experimental/AnimNextAnimGraph/Source/AnimNextAnimGraph/Public/EvaluationVM/Tasks/ConvertRotationsToLocalSpace.h
--- a/Engine/Plugins/Experimental/AnimNextAnimGraph/Source/AnimNextAnimGraph/Public/EvaluationVM/Tasks/ConvertRotationsToLocalSpace.h
+++ b/Engine/Plugins/Experimental/AnimNextAnimGraph/Source/AnimNextAnimGraph/Public/EvaluationVM/Tasks/ConvertRotationsToLocalSpace.h
@@ -18,7 +18,7 @@ struct FAnimNextConvertRotationsToLocalSpaceTask : public FAnimNextEvaluationTas
 
 	static UE_API FAnimNextConvertRotationsToLocalSpaceTask Make(const int32 NumPoses);
 
-	UE_API virtual void Execute(UE::AnimNext::FEvaluationVM& VM) const override;
+	UE_API virtual void Execute(UE::UAF::FEvaluationVM& VM) const override;
 
 private:
 	// The number of poses on the stack to convert
diff --git a/Engine/Plugins/Experimental/AnimNextAnimGraph/Source/AnimNextAnimGraph/Public/EvaluationVM/Tasks/ConvertRotationsToMeshSpace.h b/Engine/Plugins/Experimental/AnimNextAnimGraph/Source/AnimNextAnimGraph/Public/EvaluationVM/Tasks/ConvertRotationsToMeshSpace.h
--- a/Engine/Plugins/Experimental/AnimNextAnimGraph/Source/AnimNextAnimGraph/Public/EvaluationVM/Tasks/ConvertRotationsToMeshSpace.h
+++ b/Engine/Plugins/Experimental/AnimNextAnimGraph/Source/AnimNextAnimGraph/Public/EvaluationVM/Tasks/ConvertRotationsToMeshSpace.h
@@ -18,7 +18,7 @@ struct FAnimNextConvertRotationsToMeshSpaceTask : public FAnimNextEvaluationTask
 
 	static UE_API FAnimNextConvertRotationsToMeshSpaceTask Make(const int32 NumPoses);
 
-	UE_API virtual void Execute(UE::AnimNext::FEvaluationVM& VM) const override;
+	UE_API virtual void Execute(UE::UAF::FEvaluationVM& VM) const override;
 
 private:
 	// The number of poses on the stack to convert
diff --git a/Engine/Plugins/Experimental/AnimNextAnimGraph/Source/AnimNextAnimGraph/Public/EvaluationVM/Tasks/DeadBlending.h b/Engine/Plugins/Experimental/AnimNextAnimGraph/Source/AnimNextAnimGraph/Public/EvaluationVM/Tasks/DeadBlending.h
--- a/Engine/Plugins/Experimental/AnimNextAnimGraph/Source/AnimNextAnimGraph/Public/EvaluationVM/Tasks/DeadBlending.h
+++ b/Engine/Plugins/Experimental/AnimNextAnimGraph/Source/AnimNextAnimGraph/Public/EvaluationVM/Tasks/DeadBlending.h
@@ -12,7 +12,7 @@
 
 #define UE_API ANIMNEXTANIMGRAPH_API
 
-namespace UE::AnimNext
+namespace UE::UAF
 {
 	struct FDeadBlendingState;
 
@@ -51,34 +51,34 @@ struct FAnimNextDeadBlendingTransitionTask : public FAnimNextEvaluationTask
 
 	// Make a DeadBlendingTransitionTask using two previous poses
 	static UE_API FAnimNextDeadBlendingTransitionTask Make(
-		UE::AnimNext::FDeadBlendingState* State,
-		const UE::AnimNext::FTransformArraySoAHeap* CurrPose,
-		const UE::AnimNext::FTransformArraySoAHeap* PrevPose,
+		UE::UAF::FDeadBlendingState* State,
+		const UE::UAF::FTransformArraySoAHeap* CurrPose,
+		const UE::UAF::FTransformArraySoAHeap* PrevPose,
 		const float DeltaTime,
-		const UE::AnimNext::FDeadBlendTransitionTaskParameters& Parameters);
+		const UE::UAF::FDeadBlendTransitionTaskParameters& Parameters);
 
 	// Make a DeadBlendingTransitionTask using only a single previous poses
 	static UE_API FAnimNextDeadBlendingTransitionTask Make(
-		UE::AnimNext::FDeadBlendingState* State,
-		const UE::AnimNext::FTransformArraySoAHeap* CurrPose,
-		const UE::AnimNext::FDeadBlendTransitionTaskParameters& Parameters);
+		UE::UAF::FDeadBlendingState* State,
+		const UE::UAF::FTransformArraySoAHeap* CurrPose,
+		const UE::UAF::FDeadBlendTransitionTaskParameters& Parameters);
 
 	// Task entry point
-	UE_API virtual void Execute(UE::AnimNext::FEvaluationVM& VM) const override;
+	UE_API virtual void Execute(UE::UAF::FEvaluationVM& VM) const override;
 
 	// Output
 
-	UE::AnimNext::FDeadBlendingState* State = nullptr;
+	UE::UAF::FDeadBlendingState* State = nullptr;
 
 	// Inputs
 
-	const UE::AnimNext::FTransformArraySoAHeap* CurrPose = nullptr;
-	const UE::AnimNext::FTransformArraySoAHeap* PrevPose = nullptr;
+	const UE::UAF::FTransformArraySoAHeap* CurrPose = nullptr;
+	const UE::UAF::FTransformArraySoAHeap* PrevPose = nullptr;
 
 	// Parameters
 
 	float DeltaTime = 0.0f;
-	UE::AnimNext::FDeadBlendTransitionTaskParameters Parameters;
+	UE::UAF::FDeadBlendTransitionTaskParameters Parameters;
 };
 
 /**
@@ -97,18 +97,18 @@ struct FAnimNextDeadBlendingApplyTask : public FAnimNextEvaluationTask
 
 	// Make a FAnimNextDeadBlendingApplyTask task
 	static UE_API FAnimNextDeadBlendingApplyTask Make(
-		UE::AnimNext::FDeadBlendingState* State,
+		UE::UAF::FDeadBlendingState* State,
 		const float BlendDuration,
 		const float TimeSinceTransition,
 		const EAlphaBlendOption& BlendMode,
 		const TWeakObjectPtr<UCurveFloat> CustomBlendCurve);
 
 	// Task entry point
-	UE_API virtual void Execute(UE::AnimNext::FEvaluationVM& VM) const override;
+	UE_API virtual void Execute(UE::UAF::FEvaluationVM& VM) const override;
 
 	// Input / Output
 
-	UE::AnimNext::FDeadBlendingState* State = nullptr;
+	UE::UAF::FDeadBlendingState* State = nullptr;
 
 	// Parameters
 
diff --git a/Engine/Plugins/Experimental/AnimNextAnimGraph/Source/AnimNextAnimGraph/Public/EvaluationVM/Tasks/ExecuteProgram.h b/Engine/Plugins/Experimental/AnimNextAnimGraph/Source/AnimNextAnimGraph/Public/EvaluationVM/Tasks/ExecuteProgram.h
--- a/Engine/Plugins/Experimental/AnimNextAnimGraph/Source/AnimNextAnimGraph/Public/EvaluationVM/Tasks/ExecuteProgram.h
+++ b/Engine/Plugins/Experimental/AnimNextAnimGraph/Source/AnimNextAnimGraph/Public/EvaluationVM/Tasks/ExecuteProgram.h
@@ -9,7 +9,7 @@
 
 #define UE_API ANIMNEXTANIMGRAPH_API
 
-namespace UE::AnimNext
+namespace UE::UAF
 {
 	struct FEvaluationProgram;
 }
@@ -27,13 +27,13 @@ struct FAnimNextExecuteProgramTask : public FAnimNextEvaluationTask
 
 	DECLARE_ANIM_EVALUATION_TASK(FAnimNextExecuteProgramTask)
 
-	static UE_API FAnimNextExecuteProgramTask Make(TSharedPtr<const UE::AnimNext::FEvaluationProgram> Program);
+	static UE_API FAnimNextExecuteProgramTask Make(TSharedPtr<const UE::UAF::FEvaluationProgram> Program);
 
 	// Task entry point
-	UE_API virtual void Execute(UE::AnimNext::FEvaluationVM& VM) const override;
+	UE_API virtual void Execute(UE::UAF::FEvaluationVM& VM) const override;
 
 	// The program to execute
-	TSharedPtr<const UE::AnimNext::FEvaluationProgram> Program = nullptr;
+	TSharedPtr<const UE::UAF::FEvaluationProgram> Program = nullptr;
 };
 
 #undef UE_API
diff --git a/Engine/Plugins/Experimental/AnimNextAnimGraph/Source/AnimNextAnimGraph/Public/EvaluationVM/Tasks/StoreKeyframe.h b/Engine/Plugins/Experimental/AnimNextAnimGraph/Source/AnimNextAnimGraph/Public/EvaluationVM/Tasks/StoreKeyframe.h
--- a/Engine/Plugins/Experimental/AnimNextAnimGraph/Source/AnimNextAnimGraph/Public/EvaluationVM/Tasks/StoreKeyframe.h
+++ b/Engine/Plugins/Experimental/AnimNextAnimGraph/Source/AnimNextAnimGraph/Public/EvaluationVM/Tasks/StoreKeyframe.h
@@ -20,13 +20,13 @@ struct FAnimNextSwapTransformsTask : public FAnimNextEvaluationTask
 
 	DECLARE_ANIM_EVALUATION_TASK(FAnimNextSwapTransformsTask)
 
-	static UE_API FAnimNextSwapTransformsTask Make(UE::AnimNext::FTransformArraySoAHeap* A, UE::AnimNext::FTransformArraySoAHeap* B);
+	static UE_API FAnimNextSwapTransformsTask Make(UE::UAF::FTransformArraySoAHeap* A, UE::UAF::FTransformArraySoAHeap* B);
 
 	// Task entry point
-	UE_API virtual void Execute(UE::AnimNext::FEvaluationVM& VM) const override;
+	UE_API virtual void Execute(UE::UAF::FEvaluationVM& VM) const override;
 
-	UE::AnimNext::FTransformArraySoAHeap* A = nullptr;
-	UE::AnimNext::FTransformArraySoAHeap* B = nullptr;
+	UE::UAF::FTransformArraySoAHeap* A = nullptr;
+	UE::UAF::FTransformArraySoAHeap* B = nullptr;
 };
 
 /*
@@ -39,12 +39,12 @@ struct FAnimNextStoreKeyframeTransformsTask : public FAnimNextEvaluationTask
 
 	DECLARE_ANIM_EVALUATION_TASK(FAnimNextStoreKeyframeTransformsTask)
 
-	static UE_API FAnimNextStoreKeyframeTransformsTask Make(UE::AnimNext::FTransformArraySoAHeap* Dest);
+	static UE_API FAnimNextStoreKeyframeTransformsTask Make(UE::UAF::FTransformArraySoAHeap* Dest);
 
 	// Task entry point
-	UE_API virtual void Execute(UE::AnimNext::FEvaluationVM& VM) const override;
+	UE_API virtual void Execute(UE::UAF::FEvaluationVM& VM) const override;
 
-	UE::AnimNext::FTransformArraySoAHeap* Dest = nullptr;
+	UE::UAF::FTransformArraySoAHeap* Dest = nullptr;
 };
 
 /*
@@ -60,7 +60,7 @@ struct FAnimNextDuplicateTopKeyframeTask : public FAnimNextEvaluationTask
 	static UE_API FAnimNextDuplicateTopKeyframeTask Make();
 
 	// Task entry point
-	UE_API virtual void Execute(UE::AnimNext::FEvaluationVM& VM) const override;
+	UE_API virtual void Execute(UE::UAF::FEvaluationVM& VM) const override;
 };
 
 #undef UE_API
diff --git a/Engine/Plugins/Experimental/AnimNextAnimGraph/Source/AnimNextAnimGraph/Public/Factory/AnimNextAnimGraphBuilder.h b/Engine/Plugins/Experimental/AnimNextAnimGraph/Source/AnimNextAnimGraph/Public/Factory/AnimNextAnimGraphBuilder.h
--- a/Engine/Plugins/Experimental/AnimNextAnimGraph/Source/AnimNextAnimGraph/Public/Factory/AnimNextAnimGraphBuilder.h
+++ b/Engine/Plugins/Experimental/AnimNextAnimGraph/Source/AnimNextAnimGraph/Public/Factory/AnimNextAnimGraphBuilder.h
@@ -7,7 +7,7 @@
 
 class UAnimNextAnimationGraph;
 
-namespace UE::AnimNext
+namespace UE::UAF
 {
 	struct FAnimGraphFactory;
 	class FTraitWriter;
@@ -15,7 +15,7 @@ namespace UE::AnimNext
 	struct FAnimGraphBuilderContext;
 }
 
-namespace UE::AnimNext
+namespace UE::UAF
 {
 
 // A method for making a graph in FAnimGraphFactory procedurally
@@ -41,7 +41,7 @@ private:
 USTRUCT()
 struct FAnimNextAnimGraphBuilder
 #if CPP
-	: public UE::AnimNext::IAnimGraphBuilder
+	: public UE::UAF::IAnimGraphBuilder
 #endif
 {
 	GENERATED_BODY()
diff --git a/Engine/Plugins/Experimental/AnimNextAnimGraph/Source/AnimNextAnimGraph/Public/Factory/AnimNextFactoryParams.h b/Engine/Plugins/Experimental/AnimNextAnimGraph/Source/AnimNextAnimGraph/Public/Factory/AnimNextFactoryParams.h
--- a/Engine/Plugins/Experimental/AnimNextAnimGraph/Source/AnimNextAnimGraph/Public/Factory/AnimNextFactoryParams.h
+++ b/Engine/Plugins/Experimental/AnimNextAnimGraph/Source/AnimNextAnimGraph/Public/Factory/AnimNextFactoryParams.h
@@ -16,12 +16,12 @@ class UAnimNextRigVMAsset;
 struct FAnimNextFactoryParams;
 struct FAnimNextGraphInstanceTaskInstanceData;
 
-namespace UE::AnimNext
+namespace UE::UAF
 {
 	struct FAnimGraphFactory;
 }
 
-namespace UE::AnimNext::Editor
+namespace UE::UAF::Editor
 {
 	class FAnimNextFactoryParamsDetails;
 };
@@ -67,13 +67,13 @@ struct FAnimNextFactoryParams
 	}
 
 	// Get the builder that these params hold
-	const UE::AnimNext::IAnimGraphBuilder& GetBuilder() const
+	const UE::UAF::IAnimGraphBuilder& GetBuilder() const
 	{
 		return Builder;
 	}
 
 	// Add an initialize task used to set up an instance when these params are applied
-	FAnimNextFactoryParams& AddInitializeTask(UE::AnimNext::FInstanceTask&& InTask)
+	FAnimNextFactoryParams& AddInitializeTask(UE::UAF::FInstanceTask&& InTask)
 	{
 		InitializeTasks.Add(MoveTemp(InTask));
 		return *this;
@@ -86,8 +86,8 @@ struct FAnimNextFactoryParams
 	ANIMNEXTANIMGRAPH_API void Reset();
 
 private:
-	friend UE::AnimNext::FAnimGraphFactory;
-	friend UE::AnimNext::Editor::FAnimNextFactoryParamsDetails;
+	friend UE::UAF::FAnimGraphFactory;
+	friend UE::UAF::Editor::FAnimNextFactoryParamsDetails;
 	friend FAnimNextGraphInstanceTaskInstanceData;	// For old data upgrade
 
 	// Method used to create a graph
@@ -95,5 +95,5 @@ private:
 	FAnimNextSimpleAnimGraphBuilder Builder;
 
 	// Tasks to run on the factory-generated graph instance via InitializeInstance, used for initial setup of variables etc.
-	TArray<UE::AnimNext::FInstanceTask> InitializeTasks;
+	TArray<UE::UAF::FInstanceTask> InitializeTasks;
 };
diff --git a/Engine/Plugins/Experimental/AnimNextAnimGraph/Source/AnimNextAnimGraph/Public/Factory/AnimNextSimpleAnimGraphBuilder.h b/Engine/Plugins/Experimental/AnimNextAnimGraph/Source/AnimNextAnimGraph/Public/Factory/AnimNextSimpleAnimGraphBuilder.h
--- a/Engine/Plugins/Experimental/AnimNextAnimGraph/Source/AnimNextAnimGraph/Public/Factory/AnimNextSimpleAnimGraphBuilder.h
+++ b/Engine/Plugins/Experimental/AnimNextAnimGraph/Source/AnimNextAnimGraph/Public/Factory/AnimNextSimpleAnimGraphBuilder.h
@@ -12,12 +12,12 @@ class FAnimationAnimNextRuntimeTest_SimpleGraphBuilder;
 struct FAnimNextGraphInstanceTaskInstanceData;
 struct FAnimNextFactoryParams;
 
-namespace UE::AnimNext
+namespace UE::UAF
 {
 	struct FAnimGraphFactory;
 };
 
-namespace UE::AnimNext::Editor
+namespace UE::UAF::Editor
 {
 	class FAnimNextFactoryParamsDetails;
 };
@@ -28,10 +28,10 @@ struct FAnimNextSimpleAnimGraphBuilderTraitStackDesc
 	GENERATED_BODY()
 
 private:
-	friend UE::AnimNext::FAnimGraphFactory;
+	friend UE::UAF::FAnimGraphFactory;
 	friend FAnimNextSimpleAnimGraphBuilder;
 	friend FAnimNextGraphInstanceTaskInstanceData;	// For old data upgrade
-	friend UE::AnimNext::Editor::FAnimNextFactoryParamsDetails;
+	friend UE::UAF::Editor::FAnimNextFactoryParamsDetails;
 	friend FAnimNextFactoryParams;
 
 	// Traits for this stack
@@ -69,17 +69,17 @@ struct FAnimNextSimpleAnimGraphBuilder : public FAnimNextAnimGraphBuilder
 	}
 	
 private:
-	friend UE::AnimNext::FAnimGraphFactory;
+	friend UE::UAF::FAnimGraphFactory;
 	friend FAnimationAnimNextRuntimeTest_SimpleGraphBuilder;
 	friend FAnimNextGraphInstanceTaskInstanceData;	// For old data upgrade
-	friend UE::AnimNext::Editor::FAnimNextFactoryParamsDetails;
+	friend UE::UAF::Editor::FAnimNextFactoryParamsDetails;
 	friend FAnimNextFactoryParams;
 
 	// Internal helper
 	void ValidateTraitStruct(int32 InStackIndex, TConstStructView<FAnimNextTraitSharedData> InStruct);
 
 	// FAnimNextAnimGraphBuilder interface
-	ANIMNEXTANIMGRAPH_API virtual bool Build(UE::AnimNext::FAnimGraphBuilderContext& InContext) const override;
+	ANIMNEXTANIMGRAPH_API virtual bool Build(UE::UAF::FAnimGraphBuilderContext& InContext) const override;
 	ANIMNEXTANIMGRAPH_API virtual uint64 RecalculateKey() const override;
 
 	// The trait stacks we will use to generate our data
diff --git a/Engine/Plugins/Experimental/AnimNextAnimGraph/Source/AnimNextAnimGraph/Public/Graph/AnimNextAnimGraph.h b/Engine/Plugins/Experimental/AnimNextAnimGraph/Source/AnimNextAnimGraph/Public/Graph/AnimNextAnimGraph.h
--- a/Engine/Plugins/Experimental/AnimNextAnimGraph/Source/AnimNextAnimGraph/Public/Graph/AnimNextAnimGraph.h
+++ b/Engine/Plugins/Experimental/AnimNextAnimGraph/Source/AnimNextAnimGraph/Public/Graph/AnimNextAnimGraph.h
@@ -26,7 +26,7 @@ struct FStructView;
 struct FAnimNextModuleInjectionComponent;
 struct FRigUnit_GetPostProcessAnimation;
 
-namespace UE::AnimNext
+namespace UE::UAF
 {
 	struct IEvaluationModifier;
 	struct FExecutionContext;
@@ -37,12 +37,12 @@ namespace UE::AnimNext
 	struct FInjectionInfo;
 }
 
-namespace UE::AnimNext::Editor
+namespace UE::UAF::Editor
 {
 	class FAnimNextGraphDetails;
 }
 
-namespace UE::AnimNext::AnimGraph
+namespace UE::UAF::AnimGraph
 {
 	class FAnimNextAnimGraphModule;
 }
@@ -56,12 +56,12 @@ struct FAnimNextGraphInjectionData
 	bool HasValidEvaluationModifier() const { return EvaluationModifier.IsValid(); }
 private:
 	// Evaluation modifier to allow IEvaluate pass to be controlled
-	TWeakPtr<UE::AnimNext::IEvaluationModifier> EvaluationModifier;
+	TWeakPtr<UE::UAF::IEvaluationModifier> EvaluationModifier;
 
 	// Serial number for the last injection that was made to this struct
 	uint32 InjectionSerialNumber = 0;
 
-	friend UE::AnimNext::FInjectionSiteTrait;
+	friend UE::UAF::FInjectionSiteTrait;
 	friend FAnimNextModuleInjectionComponent;
 };
 
@@ -99,13 +99,13 @@ private:
 
 	friend FRigUnit_AnimNextRunAnimationGraph_v1;	// We evaluate the instance
 	friend FRigUnit_AnimNextRunAnimationGraph_v2;	// We evaluate the instance
-	friend UE::AnimNext::FBlendStackCoreTrait;		// Requires access to take a reference to external graphs
-	friend UE::AnimNext::FInjectionSiteTrait;		// Requires access to take a reference to external graphs
-	friend UE::AnimNext::FInjectionInfo;
+	friend UE::UAF::FBlendStackCoreTrait;		// Requires access to take a reference to external graphs
+	friend UE::UAF::FInjectionSiteTrait;		// Requires access to take a reference to external graphs
+	friend UE::UAF::FInjectionInfo;
 	friend FAnimNextModuleInjectionComponent;
-	friend UE::AnimNext::Editor::FAnimNextGraphDetails;
+	friend UE::UAF::Editor::FAnimNextGraphDetails;
 	friend FRigUnit_GetPostProcessAnimation;
-	friend UE::AnimNext::AnimGraph::FAnimNextAnimGraphModule;
+	friend UE::UAF::AnimGraph::FAnimNextAnimGraphModule;
 };
 
 #undef UE_API
diff --git a/Engine/Plugins/Experimental/AnimNextAnimGraph/Source/AnimNextAnimGraph/Public/Graph/AnimNextAnimationGraph.h b/Engine/Plugins/Experimental/AnimNextAnimGraph/Source/AnimNextAnimGraph/Public/Graph/AnimNextAnimationGraph.h
--- a/Engine/Plugins/Experimental/AnimNextAnimGraph/Source/AnimNextAnimGraph/Public/Graph/AnimNextAnimationGraph.h
+++ b/Engine/Plugins/Experimental/AnimNextAnimGraph/Source/AnimNextAnimGraph/Public/Graph/AnimNextAnimationGraph.h
@@ -28,7 +28,7 @@ struct FAnimNextParam;
 struct FAnimNextModuleInstance;
 struct FAnimNextModuleAnimGraphComponent;
 
-namespace UE::AnimNext
+namespace UE::UAF
 {
 	struct FContext;
 	struct FExecutionContext;
@@ -43,29 +43,29 @@ namespace UE::AnimNext
 	struct FVariableOverridesCollection;
 }
 
-namespace UE::AnimNext::UncookedOnly
+namespace UE::UAF::UncookedOnly
 {
 	struct FUtils;
 }
 
-namespace UE::AnimNext::Editor
+namespace UE::UAF::Editor
 {
 	class FModuleEditor;
 	class FVariableCustomization;
 }
 
-namespace UE::AnimNext::AnimGraph
+namespace UE::UAF::AnimGraph
 {
 	class FAnimNextAnimGraphModule;
 }
 
-namespace UE::AnimNext::Graph
+namespace UE::UAF::Graph
 {
 	extern ANIMNEXTANIMGRAPH_API const FName EntryPointName;
 	extern ANIMNEXTANIMGRAPH_API const FName ResultName;
 }
 
-namespace UE::AnimNext
+namespace UE::UAF
 {
 	// Parameters to AllocateInstance
 	struct FGraphAllocationParams
@@ -79,7 +79,7 @@ namespace UE::AnimNext
 		// The entry point to use (deprecated)
 		FName EntryPoint = NAME_None;
 		// Any variable overrides to apply to this graph
-		TSharedPtr<const UE::AnimNext::FVariableOverridesCollection> Overrides;
+		TSharedPtr<const UE::UAF::FVariableOverridesCollection> Overrides;
 	};
 }
 
@@ -101,7 +101,7 @@ public:
 	// Allocates an instance of the graph according to the supplied parameters
 	// @param	InParams		Parameters to allocate the graph with
 	// @return the allocated instance. This can be invalid if the entry point was not found or the graph's root trait is invalid.
-	UE_API TSharedPtr<FAnimNextGraphInstance> AllocateInstance(const UE::AnimNext::FGraphAllocationParams& InParams = UE::AnimNext::FGraphAllocationParams()) const;
+	UE_API TSharedPtr<FAnimNextGraphInstance> AllocateInstance(const UE::UAF::FGraphAllocationParams& InParams = UE::UAF::FGraphAllocationParams()) const;
 
 private:
 	// Loads the graph data from the provided archive buffer and returns true on success, false otherwise
@@ -119,20 +119,20 @@ private:
 	friend class UAnimNextAnimationGraphFactory;
 	friend class UAnimNextAnimationGraph_EditorData;
 	friend class UAnimNextVariableEntry;
-	friend struct UE::AnimNext::UncookedOnly::FUtils;
-	friend class UE::AnimNext::Editor::FModuleEditor;
-	friend struct UE::AnimNext::FTestUtils;
+	friend struct UE::UAF::UncookedOnly::FUtils;
+	friend class UE::UAF::Editor::FModuleEditor;
+	friend struct UE::UAF::FTestUtils;
 	friend FAnimNextGraphInstance;
 	friend class UAnimGraphNode_AnimNextGraph;
-	friend UE::AnimNext::FExecutionContext;
+	friend UE::UAF::FExecutionContext;
 	friend struct FAnimNextScheduleGraphTask;
-	friend UE::AnimNext::AnimGraph::FAnimNextAnimGraphModule;
-	friend class UE::AnimNext::Editor::FVariableCustomization;
-	friend struct UE::AnimNext::FParametersProxy;
-	friend struct UE::AnimNext::FPlayAnimSlotTrait;
-	friend struct UE::AnimNext::FBlendStackCoreTrait;
+	friend UE::UAF::AnimGraph::FAnimNextAnimGraphModule;
+	friend class UE::UAF::Editor::FVariableCustomization;
+	friend struct UE::UAF::FParametersProxy;
+	friend struct UE::UAF::FPlayAnimSlotTrait;
+	friend struct UE::UAF::FBlendStackCoreTrait;
 	friend struct FAnimNextModuleAnimGraphComponent;
-	friend UE::AnimNext::FAnimGraphBuilderContext;
+	friend UE::UAF::FAnimGraphBuilderContext;
 
 protected:
 #if WITH_EDITORONLY_DATA
diff --git a/Engine/Plugins/Experimental/AnimNextAnimGraph/Source/AnimNextAnimGraph/Public/Graph/AnimNextGraphInstance.h b/Engine/Plugins/Experimental/AnimNextAnimGraph/Source/AnimNextAnimGraph/Public/Graph/AnimNextGraphInstance.h
--- a/Engine/Plugins/Experimental/AnimNextAnimGraph/Source/AnimNextAnimGraph/Public/Graph/AnimNextGraphInstance.h
+++ b/Engine/Plugins/Experimental/AnimNextAnimGraph/Source/AnimNextAnimGraph/Public/Graph/AnimNextGraphInstance.h
@@ -29,7 +29,7 @@ struct FRigUnit_AnimNextRunAnimationGraph_v2;
 struct FAnimNextStateTreeRigVMConditionBase;
 struct FAnimNextStateTreeRigVMTaskBase;
 
-namespace UE::AnimNext
+namespace UE::UAF
 {
 	class IDataInterfaceHost;
 	struct FExecutionContext;
@@ -77,7 +77,7 @@ struct FAnimNextGraphInstance : public FUAFAssetInstance
 	UE_API FName GetEntryPoint() const;
 	
 	// Returns a weak handle to the root trait instance
-	UE_API UE::AnimNext::FWeakTraitPtr GetGraphRootPtr() const;
+	UE_API UE::UAF::FWeakTraitPtr GetGraphRootPtr() const;
 
 	// Returns the module instance that owns us or nullptr if we are invalid
 	UE_API FAnimNextModuleInstance* GetModuleInstance() const;
@@ -106,11 +106,11 @@ struct FAnimNextGraphInstance : public FUAFAssetInstance
 private:
 	// Executes a list of latent RigVM pins and writes the result into the destination pointer (latent handle offsets are using the destination as base)
 	// When frozen, latent handles that can freeze are skipped, all others will execute
-	UE_API void ExecuteLatentPins(const TConstArrayView<UE::AnimNext::FLatentPropertyHandle>& LatentHandles, void* DestinationBasePtr, bool bIsFrozen, bool bJustBecameRelevant);
+	UE_API void ExecuteLatentPins(const TConstArrayView<UE::UAF::FLatentPropertyHandle>& LatentHandles, void* DestinationBasePtr, bool bIsFrozen, bool bJustBecameRelevant);
 
 	// When a group of latent pins do not require RigVM to operate, this copies the values directly from variables rather than calling ExecuteLatentPins.
 	// When frozen, latent handles that can freeze are skipped, all others will execute
-	void CopyVariablesToLatentPins(const TConstArrayView<UE::AnimNext::FLatentPropertyHandle>& LatentHandles, void* DestinationBasePtr, bool bIsFrozen, bool bJustBecameRelevant);
+	void CopyVariablesToLatentPins(const TConstArrayView<UE::UAF::FLatentPropertyHandle>& LatentHandles, void* DestinationBasePtr, bool bIsFrozen, bool bJustBecameRelevant);
 
 #if WITH_EDITORONLY_DATA
 	// During graph compilation, if we have existing graph instances, we freeze them by releasing their memory before thawing them
@@ -125,7 +125,7 @@ private:
 	FName EntryPoint;
 
 	// Hard reference to the graph instance data, we own it
-	UE::AnimNext::FTraitPtr GraphInstancePtr;
+	UE::UAF::FTraitPtr GraphInstancePtr;
 
 	// The module instance that owns the root, us and the components
 	FAnimNextModuleInstance* ModuleInstance = nullptr;
@@ -138,20 +138,20 @@ private:
 
 	friend UAnimNextAnimationGraph;			// The graph is the one that allocates instances
 	friend FRigUnit_AnimNextGraphEvaluator;	// We evaluate the instance
-	friend UE::AnimNext::FExecutionContext;
-	friend UE::AnimNext::FTraitStackBinding;
-	friend UE::AnimNext::FPlayAnimSlotTrait;
-	friend UE::AnimNext::FBlendStackCoreTrait;
-	friend UE::AnimNext::FInjectionRequest;
-	friend UE::AnimNext::FInjectionSiteTrait;
-	friend UE::AnimNext::FStateTreeTrait;	// Temp - remove: Needs to access mutable variables to copy-in data OnBecomeRelevant
+	friend UE::UAF::FExecutionContext;
+	friend UE::UAF::FTraitStackBinding;
+	friend UE::UAF::FPlayAnimSlotTrait;
+	friend UE::UAF::FBlendStackCoreTrait;
+	friend UE::UAF::FInjectionRequest;
+	friend UE::UAF::FInjectionSiteTrait;
+	friend UE::UAF::FStateTreeTrait;	// Temp - remove: Needs to access mutable variables to copy-in data OnBecomeRelevant
 	friend FAnimNextModuleInjectionComponent;
 	friend FRigUnit_AnimNextRunAnimationGraph_v1;
 	friend FRigUnit_AnimNextRunAnimationGraph_v2;
-	friend UE::AnimNext::FInjectionUtils;
+	friend UE::UAF::FInjectionUtils;
 	friend FAnimNextStateTreeRigVMConditionBase;	// Temp - remove: Needs to access mutable variables to copy-in data for function shim
 	friend FAnimNextStateTreeRigVMTaskBase;	// Temp - remove: Needs to access mutable variables to copy-in data for function shim
-	friend UE::AnimNext::FBlendSpacePlayerTrait;
+	friend UE::UAF::FBlendSpacePlayerTrait;
 };
 
 template<>
diff --git a/Engine/Plugins/Experimental/AnimNextAnimGraph/Source/AnimNextAnimGraph/Public/Graph/UAFGraphInstanceComponent.h b/Engine/Plugins/Experimental/AnimNextAnimGraph/Source/AnimNextAnimGraph/Public/Graph/UAFGraphInstanceComponent.h
--- a/Engine/Plugins/Experimental/AnimNextAnimGraph/Source/AnimNextAnimGraph/Public/Graph/UAFGraphInstanceComponent.h
+++ b/Engine/Plugins/Experimental/AnimNextAnimGraph/Source/AnimNextAnimGraph/Public/Graph/UAFGraphInstanceComponent.h
@@ -9,7 +9,7 @@
 struct FAnimNextTraitEvent;
 struct FAnimNextGraphInstance;
 
-namespace UE::AnimNext
+namespace UE::UAF
 {
 	struct FExecutionContext;
 }
@@ -40,12 +40,12 @@ struct FUAFGraphInstanceComponent : public FUAFAssetInstanceComponent
 	// Called before the update traversal begins, before any node has been visited
 	// Note that PreUpdate won't be called if a component is created during the update traversal until the next update
 	// The execution context provided is bound to the graph root and can be bound to anything the component wishes
-	virtual void PreUpdate(UE::AnimNext::FExecutionContext& Context) {}
+	virtual void PreUpdate(UE::UAF::FExecutionContext& Context) {}
 
 	// Called after the update traversal completes, after every node has been visited
 	// The execution context provided is bound to the graph root and can be bound to anything the component wishes
-	virtual void PostUpdate(UE::AnimNext::FExecutionContext& Context) {}
+	virtual void PostUpdate(UE::UAF::FExecutionContext& Context) {}
 
 	// Called before PreUpdate with input events and before PostUpdate with output events
-	virtual void OnTraitEvent(UE::AnimNext::FExecutionContext& Context, FAnimNextTraitEvent& Event) {}
+	virtual void OnTraitEvent(UE::UAF::FExecutionContext& Context, FAnimNextTraitEvent& Event) {}
 };
