UAF: Namespace renamed

diff --git a/Engine/Plugins/Experimental/AnimNextAnimGraph/Source/AnimNextAnimGraph/Public/Injection/InjectionRequest.h b/Engine/Plugins/Experimental/AnimNextAnimGraph/Source/AnimNextAnimGraph/Public/Injection/InjectionRequest.h
--- a/Engine/Plugins/Experimental/AnimNextAnimGraph/Source/AnimNextAnimGraph/Public/Injection/InjectionRequest.h
+++ b/Engine/Plugins/Experimental/AnimNextAnimGraph/Source/AnimNextAnimGraph/Public/Injection/InjectionRequest.h
@@ -23,7 +23,7 @@ struct FAnimNextGraphInstance;
 class UInjectionCallbackProxy;
 class UWorld;
 
-namespace UE::AnimNext
+namespace UE::UAF
 {
 	struct FInstanceTaskContext;
 	struct IEvaluationModifier;
@@ -106,11 +106,11 @@ struct FAnimNextInjectionRequestArgs
 	GENERATED_BODY()
 
 private:
-	friend UE::AnimNext::FInjectionRequest;
-	friend UE::AnimNext::FInjectionSiteTrait;
-	friend UE::AnimNext::FPlayAnimSlotTrait;
-	friend UE::AnimNext::FInjectionRequestTracker;
-	friend UE::AnimNext::FInjectionUtils;
+	friend UE::UAF::FInjectionRequest;
+	friend UE::UAF::FInjectionSiteTrait;
+	friend UE::UAF::FPlayAnimSlotTrait;
+	friend UE::UAF::FInjectionRequestTracker;
+	friend UE::UAF::FInjectionUtils;
 	friend FAnimNextModuleInjectionComponent;
 	friend UInjectionCallbackProxy;
 
@@ -131,7 +131,7 @@ private:
 	TObjectPtr<UObject> Object;
 
 	// Evaluation modifier to apply at the injection site. Shared ownership, references can persist in worker threads until un-injected.
-	TSharedPtr<UE::AnimNext::IEvaluationModifier> EvaluationModifier;
+	TSharedPtr<UE::UAF::IEvaluationModifier> EvaluationModifier;
 
 	UPROPERTY()
 	EAnimNextInjectionType Type = EAnimNextInjectionType::InjectObject;
@@ -149,7 +149,7 @@ private:
 	FAnimNextFactoryParams FactoryParams;
 };
 
-namespace UE::AnimNext
+namespace UE::UAF
 {
 	// Create a namespaced aliases to simplify usage
 	using EInjectionBlendMode = EAnimNextInjectionBlendMode;
diff --git a/Engine/Plugins/Experimental/AnimNextAnimGraph/Source/AnimNextAnimGraph/Public/TraitCore/EntryPointHandle.h b/Engine/Plugins/Experimental/AnimNextAnimGraph/Source/AnimNextAnimGraph/Public/TraitCore/EntryPointHandle.h
--- a/Engine/Plugins/Experimental/AnimNextAnimGraph/Source/AnimNextAnimGraph/Public/TraitCore/EntryPointHandle.h
+++ b/Engine/Plugins/Experimental/AnimNextAnimGraph/Source/AnimNextAnimGraph/Public/TraitCore/EntryPointHandle.h
@@ -25,18 +25,18 @@ struct FAnimNextEntryPointHandle final
 
 	// Creates an invalid entry point handle
 	constexpr FAnimNextEntryPointHandle() noexcept
-		: PackedTraitIndexAndNodeHandle(UE::AnimNext::FNodeHandle::INVALID_NODE_HANDLE_RAW_VALUE)
+		: PackedTraitIndexAndNodeHandle(UE::UAF::FNodeHandle::INVALID_NODE_HANDLE_RAW_VALUE)
 	{}
 
 	// Creates an entry point handle pointing to the first trait of the specified node
-	explicit FAnimNextEntryPointHandle(UE::AnimNext::FNodeHandle NodeHandle)
+	explicit FAnimNextEntryPointHandle(UE::UAF::FNodeHandle NodeHandle)
 		: PackedTraitIndexAndNodeHandle(NodeHandle.GetPackedValue() & NODE_HANDLE_MASK)
 	{
 		check(!NodeHandle.IsValid() || NodeHandle.IsNodeID());
 	}
 
 	// Creates an entry point handle pointing to the specified trait on the specified node
-	FAnimNextEntryPointHandle(UE::AnimNext::FNodeHandle NodeHandle, uint32 TraitIndex)
+	FAnimNextEntryPointHandle(UE::UAF::FNodeHandle NodeHandle, uint32 TraitIndex)
 		: PackedTraitIndexAndNodeHandle((NodeHandle.GetPackedValue() & NODE_HANDLE_MASK) | (TraitIndex << 24))
 	{
 		check(!NodeHandle.IsValid() || NodeHandle.IsNodeID());
@@ -50,7 +50,7 @@ struct FAnimNextEntryPointHandle final
 	constexpr uint32 GetTraitIndex() const noexcept { return PackedTraitIndexAndNodeHandle >> TRAIT_INDEX_SHIFT; }
 
 	// Returns a handle to the node referenced (its node ID)
-	constexpr UE::AnimNext::FNodeHandle GetNodeHandle() const noexcept { return UE::AnimNext::FNodeHandle::FromPackedValue(PackedTraitIndexAndNodeHandle); }
+	constexpr UE::UAF::FNodeHandle GetNodeHandle() const noexcept { return UE::UAF::FNodeHandle::FromPackedValue(PackedTraitIndexAndNodeHandle); }
 
 private:
 	// Bottom 24 bits are used by the node handle while the top 8 bits by the trait index
diff --git a/Engine/Plugins/Experimental/AnimNextAnimGraph/Source/AnimNextAnimGraph/Public/TraitCore/ITraitInterface.h b/Engine/Plugins/Experimental/AnimNextAnimGraph/Source/AnimNextAnimGraph/Public/TraitCore/ITraitInterface.h
--- a/Engine/Plugins/Experimental/AnimNextAnimGraph/Source/AnimNextAnimGraph/Public/TraitCore/ITraitInterface.h
+++ b/Engine/Plugins/Experimental/AnimNextAnimGraph/Source/AnimNextAnimGraph/Public/TraitCore/ITraitInterface.h
@@ -7,7 +7,7 @@
 
 #define UE_API ANIMNEXTANIMGRAPH_API
 
-namespace UE::AnimNext
+namespace UE::UAF
 {
 	struct ITraitInterface;
 }
@@ -17,16 +17,16 @@ namespace UE::AnimNext
 // In a trait interface struct declaration, this macro declares the necessary boilerplate we require
 #define DECLARE_ANIM_TRAIT_INTERFACE(InterfaceName) \
 	/* Globally unique UID for this interface */ \
-	static constexpr UE::AnimNext::FTraitInterfaceUID InterfaceUID = UE::AnimNext::FTraitInterfaceUID::MakeUID(TEXT(#InterfaceName)); \
-	virtual UE::AnimNext::FTraitInterfaceUID GetInterfaceUID() const override { return InterfaceUID; }
+	static constexpr UE::UAF::FTraitInterfaceUID InterfaceUID = UE::UAF::FTraitInterfaceUID::MakeUID(TEXT(#InterfaceName)); \
+	virtual UE::UAF::FTraitInterfaceUID GetInterfaceUID() const override { return InterfaceUID; }
 
 // Allows a trait interface to auto-register and unregister within the current execution scope
 // The trait interface must be found in the current scope without a namespace qualification
 #define AUTO_REGISTER_ANIM_TRAIT_INTERFACE(Interface) \
-		UE::AnimNext::FTraitInterfaceStaticInitHook Interface##Hook(MakeShared<Interface>());
+		UE::UAF::FTraitInterfaceStaticInitHook Interface##Hook(MakeShared<Interface>());
 
 
-namespace UE::AnimNext
+namespace UE::UAF
 {
 	struct FExecutionContext;	// Derived types will have functions that accept the execution context
 
@@ -83,7 +83,7 @@ namespace UE::AnimNext
 	*/
 	struct FTraitInterfaceStaticInitHook final
 	{
-		UE_API explicit FTraitInterfaceStaticInitHook(const TSharedPtr<UE::AnimNext::ITraitInterface>& InTraitInterface);
+		UE_API explicit FTraitInterfaceStaticInitHook(const TSharedPtr<UE::UAF::ITraitInterface>& InTraitInterface);
 		UE_API ~FTraitInterfaceStaticInitHook();
 
 	private:
diff --git a/Engine/Plugins/Experimental/AnimNextAnimGraph/Source/AnimNextAnimGraph/Public/TraitCore/Trait.h b/Engine/Plugins/Experimental/AnimNextAnimGraph/Source/AnimNextAnimGraph/Public/TraitCore/Trait.h
--- a/Engine/Plugins/Experimental/AnimNextAnimGraph/Source/AnimNextAnimGraph/Public/TraitCore/Trait.h
+++ b/Engine/Plugins/Experimental/AnimNextAnimGraph/Source/AnimNextAnimGraph/Public/TraitCore/Trait.h
@@ -26,29 +26,29 @@ struct FRigVMPinInfoArray;
 #define ANIM_NEXT_IMPL_DECLARE_ANIM_TRAIT_BASIC(TraitName, SuperTraitName) \
 	using TraitSuper = SuperTraitName; \
 	/* FTrait impl */ \
-	static constexpr UE::AnimNext::FTraitUID TraitUID = UE::AnimNext::FTraitUID::MakeUID(TEXT(#TraitName)); \
-	virtual UE::AnimNext::FTraitUID GetTraitUID() const override { return TraitUID; } \
+	static constexpr UE::UAF::FTraitUID TraitUID = UE::UAF::FTraitUID::MakeUID(TEXT(#TraitName)); \
+	virtual UE::UAF::FTraitUID GetTraitUID() const override { return TraitUID; } \
 	virtual FString GetTraitName() const override { return TEXT(#TraitName); }
 
 #define ANIM_NEXT_IMPL_DECLARE_ANIM_TRAIT_INSTANCING_SUPPORT() \
-	static const UE::AnimNext::FTraitMemoryLayout TraitMemoryDescription; \
-	virtual UE::AnimNext::FTraitMemoryLayout GetTraitMemoryDescription() const override { return TraitMemoryDescription; } \
+	static const UE::UAF::FTraitMemoryLayout TraitMemoryDescription; \
+	virtual UE::UAF::FTraitMemoryLayout GetTraitMemoryDescription() const override { return TraitMemoryDescription; } \
 	virtual UScriptStruct* GetTraitSharedDataStruct() const override { return FSharedData::StaticStruct(); } \
-	virtual void ConstructTraitInstance(const UE::AnimNext::FExecutionContext& Context, const UE::AnimNext::FTraitBinding& Binding) const override; \
-	virtual void DestructTraitInstance(const UE::AnimNext::FExecutionContext& Context, const UE::AnimNext::FTraitBinding& Binding) const override;
+	virtual void ConstructTraitInstance(const UE::UAF::FExecutionContext& Context, const UE::UAF::FTraitBinding& Binding) const override; \
+	virtual void DestructTraitInstance(const UE::UAF::FExecutionContext& Context, const UE::UAF::FTraitBinding& Binding) const override;
 
 #define ANIM_NEXT_IMPL_DECLARE_ANIM_TRAIT_INTERFACE_SUPPORT() \
-	virtual const UE::AnimNext::ITraitInterface* GetTraitInterface(UE::AnimNext::FTraitInterfaceUID InterfaceUID) const override; \
-	virtual TConstArrayView<UE::AnimNext::FTraitInterfaceUID> GetTraitInterfaces() const override; \
-	virtual TConstArrayView<UE::AnimNext::FTraitInterfaceUID> GetTraitRequiredInterfaces() const override;
+	virtual const UE::UAF::ITraitInterface* GetTraitInterface(UE::UAF::FTraitInterfaceUID InterfaceUID) const override; \
+	virtual TConstArrayView<UE::UAF::FTraitInterfaceUID> GetTraitInterfaces() const override; \
+	virtual TConstArrayView<UE::UAF::FTraitInterfaceUID> GetTraitRequiredInterfaces() const override;
 
 #define ANIM_NEXT_IMPL_DECLARE_ANIM_TRAIT_EVENT_SUPPORT() \
-	virtual UE::AnimNext::ETraitStackPropagation OnTraitEvent(UE::AnimNext::FExecutionContext& Context, UE::AnimNext::FTraitBinding& Binding, FAnimNextTraitEvent& Event) const override; \
-	virtual TConstArrayView<UE::AnimNext::FTraitEventUID> GetTraitEvents() const override;
+	virtual UE::UAF::ETraitStackPropagation OnTraitEvent(UE::UAF::FExecutionContext& Context, UE::UAF::FTraitBinding& Binding, FAnimNextTraitEvent& Event) const override; \
+	virtual TConstArrayView<UE::UAF::FTraitEventUID> GetTraitEvents() const override;
 
 #define ANIM_NEXT_IMPL_DECLARE_ANIM_TRAIT_LATENT_PROPERTY_SUPPORT() \
 	virtual uint32 GetNumLatentTraitProperties() const override { return -FSharedData::GetLatentPropertyIndex(~(size_t)0); } \
-	virtual UE::AnimNext::FTraitLatentPropertyMemoryLayout GetLatentPropertyMemoryLayout(const FAnimNextTraitSharedData& SharedData, FName PropertyName, uint32 PropertyIndex) const override; \
+	virtual UE::UAF::FTraitLatentPropertyMemoryLayout GetLatentPropertyMemoryLayout(const FAnimNextTraitSharedData& SharedData, FName PropertyName, uint32 PropertyIndex) const override; \
 	virtual bool IsPropertyLatent(const FAnimNextTraitSharedData& InSharedData, FName PropertyName) const override;
 
 // Helper macros
@@ -76,14 +76,14 @@ struct FRigVMPinInfoArray;
 
 // Implements various parts of FTrait
 #define ANIM_NEXT_IMPL_DEFINE_ANIM_TRAIT(TraitName) \
-	const UE::AnimNext::FTraitMemoryLayout TraitName::TraitMemoryDescription = \
-		UE::AnimNext::FTraitMemoryLayout{ \
+	const UE::UAF::FTraitMemoryLayout TraitName::TraitMemoryDescription = \
+		UE::UAF::FTraitMemoryLayout{ \
 			sizeof(TraitName), alignof(TraitName), \
 			/* Use the effective size to ensure optimal packing for empty shared/instance data */ \
-			UE::AnimNext::Private::TGetEffectiveSize<TraitName::FSharedData>::Value, alignof(TraitName::FSharedData), \
-			UE::AnimNext::Private::TGetEffectiveSize<TraitName::FInstanceData>::Value, alignof(TraitName::FInstanceData) \
+			UE::UAF::Private::TGetEffectiveSize<TraitName::FSharedData>::Value, alignof(TraitName::FSharedData), \
+			UE::UAF::Private::TGetEffectiveSize<TraitName::FInstanceData>::Value, alignof(TraitName::FInstanceData) \
 		}; \
-	void TraitName::ConstructTraitInstance(const UE::AnimNext::FExecutionContext& Context, const UE::AnimNext::FTraitBinding& Binding) const \
+	void TraitName::ConstructTraitInstance(const UE::UAF::FExecutionContext& Context, const UE::UAF::FTraitBinding& Binding) const \
 	{ \
 		static_assert(std::is_base_of<TraitSuper::FSharedData, TraitName::FSharedData>::value, "Trait shared data must derive from its Super's FSharedData"); \
 		static_assert(std::is_base_of<TraitSuper::FInstanceData, TraitName::FInstanceData>::value, "Trait instance data must derive from Super's FInstanceData"); \
@@ -94,7 +94,7 @@ struct FRigVMPinInfoArray;
 		/* Construct our typed instance last */ \
 		Data->Construct(Context, Binding); \
 	} \
-	void TraitName::DestructTraitInstance(const UE::AnimNext::FExecutionContext& Context, const UE::AnimNext::FTraitBinding& Binding) const \
+	void TraitName::DestructTraitInstance(const UE::UAF::FExecutionContext& Context, const UE::UAF::FTraitBinding& Binding) const \
 	{ \
 		/* Destruction is reverse order of construction above */ \
 		FInstanceData* Data = Binding.GetInstanceData<FInstanceData>(); \
@@ -105,10 +105,10 @@ struct FRigVMPinInfoArray;
 
 // Implements GetLatentPropertyMemoryLayout()
 #define ANIM_NEXT_IMPL_DEFINE_ANIM_TRAIT_GET_LATENT_PROPERTY_MEMORY_LAYOUT(TraitName) \
-	UE::AnimNext::FTraitLatentPropertyMemoryLayout TraitName::GetLatentPropertyMemoryLayout(const FAnimNextTraitSharedData& SharedData, FName PropertyName, uint32 PropertyIndex) const \
+	UE::UAF::FTraitLatentPropertyMemoryLayout TraitName::GetLatentPropertyMemoryLayout(const FAnimNextTraitSharedData& SharedData, FName PropertyName, uint32 PropertyIndex) const \
 	{ \
 		/* Thread safe cache initialization */ \
-		static TArray<UE::AnimNext::FTraitLatentPropertyMemoryLayout> CachedLatentPropertyMemoryLayouts = [this](){ TArray<UE::AnimNext::FTraitLatentPropertyMemoryLayout> Result; Result.SetNum(GetNumLatentTraitProperties()); return Result; }(); \
+		static TArray<UE::UAF::FTraitLatentPropertyMemoryLayout> CachedLatentPropertyMemoryLayouts = [this](){ TArray<UE::UAF::FTraitLatentPropertyMemoryLayout> Result; Result.SetNum(GetNumLatentTraitProperties()); return Result; }(); \
 		return GetLatentPropertyMemoryLayoutImpl(PropertyName, PropertyIndex, CachedLatentPropertyMemoryLayouts); \
 	}
 
@@ -130,7 +130,7 @@ struct FRigVMPinInfoArray;
 
 // Implements GetTraitInterface()
 #define ANIM_NEXT_IMPL_DEFINE_ANIM_TRAIT_GET_INTERFACE(TraitName, InterfaceEnumeratorMacro) \
-	const UE::AnimNext::ITraitInterface* TraitName::GetTraitInterface(UE::AnimNext::FTraitInterfaceUID InInterfaceUID) const \
+	const UE::UAF::ITraitInterface* TraitName::GetTraitInterface(UE::UAF::FTraitInterfaceUID InInterfaceUID) const \
 	{ \
 		InterfaceEnumeratorMacro(ANIM_NEXT_IMPL_GET_INTERFACE_IMPL_FOR_INTERFACE) \
 		/* Forward to base implementation */ \
@@ -142,10 +142,10 @@ struct FRigVMPinInfoArray;
 
 // Implements GetTraitInterfaces()
 #define ANIM_NEXT_IMPL_DEFINE_ANIM_TRAIT_GET_INTERFACES(TraitName, InterfaceEnumeratorMacro) \
-	TConstArrayView<UE::AnimNext::FTraitInterfaceUID> TraitName::GetTraitInterfaces() const \
+	TConstArrayView<UE::UAF::FTraitInterfaceUID> TraitName::GetTraitInterfaces() const \
 	{ \
 		/* Thread safe cache initialization */ \
-		static TArray<UE::AnimNext::FTraitInterfaceUID> CachedInterfaceList = FTrait::BuildTraitInterfaceList( \
+		static TArray<UE::UAF::FTraitInterfaceUID> CachedInterfaceList = FTrait::BuildTraitInterfaceList( \
 			TraitSuper::GetTraitInterfaces(), \
 			{ \
 				InterfaceEnumeratorMacro(ANIM_NEXT_IMPL_GET_INTERFACES_IMPL_FOR_INTERFACE) \
@@ -155,10 +155,10 @@ struct FRigVMPinInfoArray;
 
 // Implements GetTraitRequiredInterfaces()
 #define ANIM_NEXT_IMPL_DEFINE_ANIM_TRAIT_GET_REQUIRED_INTERFACES(TraitName, RequiredInterfaceEnumeratorMacro) \
-	TConstArrayView<UE::AnimNext::FTraitInterfaceUID> TraitName::GetTraitRequiredInterfaces() const \
+	TConstArrayView<UE::UAF::FTraitInterfaceUID> TraitName::GetTraitRequiredInterfaces() const \
 	{ \
 		/* Thread safe cache initialization */ \
-		static TArray<UE::AnimNext::FTraitInterfaceUID> CachedInterfaceList = FTrait::BuildTraitInterfaceList( \
+		static TArray<UE::UAF::FTraitInterfaceUID> CachedInterfaceList = FTrait::BuildTraitInterfaceList( \
 			TraitSuper::GetTraitRequiredInterfaces(), \
 			{ \
 				RequiredInterfaceEnumeratorMacro(ANIM_NEXT_IMPL_GET_INTERFACES_IMPL_FOR_INTERFACE) \
@@ -166,7 +166,7 @@ struct FRigVMPinInfoArray;
 		return CachedInterfaceList; \
 	}
 
-namespace UE::AnimNext::Private
+namespace UE::UAF::Private
 {
 	// Helper to grab the event type from an event handler function signature
 	template<typename HandlerEventType>
@@ -214,16 +214,16 @@ namespace UE::AnimNext::Private
 
 // Helper that handles the OnTraitEvent() details for each event specified by the generator macro
 #define ANIM_NEXT_IMPL_ON_TRAIT_EVENT_IMPL_FOR_EVENT(EventHandler) \
-	if (EventUID == UE::AnimNext::Private::EventHandlerTypeTrait<decltype(&EventHandler)>::EventType::TypeUID) \
+	if (EventUID == UE::UAF::Private::EventHandlerTypeTrait<decltype(&EventHandler)>::EventType::TypeUID) \
 	{ \
-		return EventHandler(Context, Binding, static_cast<UE::AnimNext::Private::EventHandlerTypeTrait<decltype(&EventHandler)>::EventType&>(Event)); \
+		return EventHandler(Context, Binding, static_cast<UE::UAF::Private::EventHandlerTypeTrait<decltype(&EventHandler)>::EventType&>(Event)); \
 	}
 
 // Implements OnTraitEvent()
 #define ANIM_NEXT_IMPL_DEFINE_ANIM_TRAIT_ON_TRAIT_EVENT(TraitName, EventEnumeratorMacro) \
-	UE::AnimNext::ETraitStackPropagation TraitName::OnTraitEvent(UE::AnimNext::FExecutionContext& Context, UE::AnimNext::FTraitBinding& Binding, FAnimNextTraitEvent& Event) const \
+	UE::UAF::ETraitStackPropagation TraitName::OnTraitEvent(UE::UAF::FExecutionContext& Context, UE::UAF::FTraitBinding& Binding, FAnimNextTraitEvent& Event) const \
 	{ \
-		const UE::AnimNext::FTraitEventUID EventUID = Event.GetTypeUID(); \
+		const UE::UAF::FTraitEventUID EventUID = Event.GetTypeUID(); \
 		EventEnumeratorMacro(ANIM_NEXT_IMPL_ON_TRAIT_EVENT_IMPL_FOR_EVENT) \
 		/* Forward to base implementation */ \
 		return TraitSuper::OnTraitEvent(Context, Binding, Event); \
@@ -231,14 +231,14 @@ namespace UE::AnimNext::Private
 
 // Helper that handles the GetTraitEvents() details for each event specified by the generator macro
 #define ANIM_NEXT_IMPL_GET_TRAIT_EVENTS_IMPL_FOR_EVENT(EventHandler) \
-	UE::AnimNext::Private::EventHandlerTypeTrait<decltype(&EventHandler)>::EventType::TypeUID,
+	UE::UAF::Private::EventHandlerTypeTrait<decltype(&EventHandler)>::EventType::TypeUID,
 
 // Implements GetTraitEvents()
 #define ANIM_NEXT_IMPL_DEFINE_ANIM_TRAIT_GET_TRAIT_EVENTS(TraitName, EventEnumeratorMacro) \
-	TConstArrayView<UE::AnimNext::FTraitEventUID> TraitName::GetTraitEvents() const \
+	TConstArrayView<UE::UAF::FTraitEventUID> TraitName::GetTraitEvents() const \
 	{ \
 		/* Thread safe cache initialization */ \
-		static TArray<UE::AnimNext::FTraitEventUID> CachedEventList = FTrait::BuildTraitEventList( \
+		static TArray<UE::UAF::FTraitEventUID> CachedEventList = FTrait::BuildTraitEventList( \
 			TraitSuper::GetTraitEvents(), \
 			{ \
 				EventEnumeratorMacro(ANIM_NEXT_IMPL_GET_TRAIT_EVENTS_IMPL_FOR_EVENT) \
@@ -268,14 +268,14 @@ namespace UE::AnimNext::Private
 // Allows a trait to auto-register and unregister within the current execution scope
 // The trait must be found in the current scope without a namespace qualification
 #define AUTO_REGISTER_ANIM_TRAIT(TraitName) \
-	UE::AnimNext::FTraitStaticInitHook TraitName##Hook( \
-		[](void* DestPtr, UE::AnimNext::FTraitMemoryLayout& MemoryDesc) -> UE::AnimNext::FTrait* \
+	UE::UAF::FTraitStaticInitHook TraitName##Hook( \
+		[](void* DestPtr, UE::UAF::FTraitMemoryLayout& MemoryDesc) -> UE::UAF::FTrait* \
 		{ \
 			MemoryDesc = TraitName::TraitMemoryDescription; \
 			return DestPtr != nullptr ? new(DestPtr) TraitName() : nullptr; \
 		});
 
-namespace UE::AnimNext
+namespace UE::UAF
 {
 	struct FTrait;
 	struct FTraitBinding;
diff --git a/Engine/Plugins/Experimental/AnimNextAnimGraph/Source/AnimNextAnimGraph/Public/TraitCore/TraitBinding.h b/Engine/Plugins/Experimental/AnimNextAnimGraph/Source/AnimNextAnimGraph/Public/TraitCore/TraitBinding.h
--- a/Engine/Plugins/Experimental/AnimNextAnimGraph/Source/AnimNextAnimGraph/Public/TraitCore/TraitBinding.h
+++ b/Engine/Plugins/Experimental/AnimNextAnimGraph/Source/AnimNextAnimGraph/Public/TraitCore/TraitBinding.h
@@ -16,7 +16,7 @@
 
 struct FAnimNextTraitSharedData;
 
-namespace UE::AnimNext
+namespace UE::UAF
 {
 	struct FExecutionContext;
 	struct FTraitInstanceData;
@@ -275,7 +275,7 @@ namespace UE::AnimNext
 	 * to bind a trait to a specific interface. See existing interfaces for examples.
 	 * 
 	 * Here, we forward declare the template which every interface must specialize. Because we
-	 * rely on specialization, it must be defined within the UE::AnimNext namespace where the
+	 * rely on specialization, it must be defined within the UE::UAF namespace where the
 	 * declaration exists.
 	 * 
 	 * Specializations must derive from FTraitBinding to provide the necessary machinery.
diff --git a/Engine/Plugins/Experimental/AnimNextAnimGraph/Source/AnimNextAnimGraph/Public/TraitCore/TraitHandle.h b/Engine/Plugins/Experimental/AnimNextAnimGraph/Source/AnimNextAnimGraph/Public/TraitCore/TraitHandle.h
--- a/Engine/Plugins/Experimental/AnimNextAnimGraph/Source/AnimNextAnimGraph/Public/TraitCore/TraitHandle.h
+++ b/Engine/Plugins/Experimental/AnimNextAnimGraph/Source/AnimNextAnimGraph/Public/TraitCore/TraitHandle.h
@@ -25,16 +25,16 @@ struct FAnimNextTraitHandle final
 
 	// Creates an invalid trait handle
 	constexpr FAnimNextTraitHandle() noexcept
-		: PackedTraitIndexAndNodeHandle(UE::AnimNext::FNodeHandle::INVALID_NODE_HANDLE_RAW_VALUE)
+		: PackedTraitIndexAndNodeHandle(UE::UAF::FNodeHandle::INVALID_NODE_HANDLE_RAW_VALUE)
 	{}
 
 	// Creates a trait handle pointing to the first trait of the specified node
-	explicit constexpr FAnimNextTraitHandle(UE::AnimNext::FNodeHandle NodeHandle) noexcept
+	explicit constexpr FAnimNextTraitHandle(UE::UAF::FNodeHandle NodeHandle) noexcept
 		: PackedTraitIndexAndNodeHandle(NodeHandle.GetPackedValue() & NODE_HANDLE_MASK)
 	{}
 
 	// Creates a trait handle pointing to the specified trait on the specified node
-	FAnimNextTraitHandle(UE::AnimNext::FNodeHandle NodeHandle, uint32 TraitIndex)
+	FAnimNextTraitHandle(UE::UAF::FNodeHandle NodeHandle, uint32 TraitIndex)
 		: PackedTraitIndexAndNodeHandle((NodeHandle.GetPackedValue() & NODE_HANDLE_MASK) | (TraitIndex << 24))
 	{
 		check(TraitIndex <= MAX_uint8);	// Make sure we don't truncate
@@ -47,7 +47,7 @@ struct FAnimNextTraitHandle final
 	constexpr uint32 GetTraitIndex() const noexcept { return PackedTraitIndexAndNodeHandle >> TRAIT_INDEX_SHIFT; }
 
 	// Returns a handle to the node referenced
-	constexpr UE::AnimNext::FNodeHandle GetNodeHandle() const noexcept { return UE::AnimNext::FNodeHandle::FromPackedValue(PackedTraitIndexAndNodeHandle); }
+	constexpr UE::UAF::FNodeHandle GetNodeHandle() const noexcept { return UE::UAF::FNodeHandle::FromPackedValue(PackedTraitIndexAndNodeHandle); }
 
 	ANIMNEXTANIMGRAPH_API bool Serialize(class FArchive& Ar);
 
diff --git a/Engine/Plugins/Experimental/AnimNextAnimGraph/Source/AnimNextAnimGraph/Public/TraitCore/TraitSharedData.h b/Engine/Plugins/Experimental/AnimNextAnimGraph/Source/AnimNextAnimGraph/Public/TraitCore/TraitSharedData.h
--- a/Engine/Plugins/Experimental/AnimNextAnimGraph/Source/AnimNextAnimGraph/Public/TraitCore/TraitSharedData.h
+++ b/Engine/Plugins/Experimental/AnimNextAnimGraph/Source/AnimNextAnimGraph/Public/TraitCore/TraitSharedData.h
@@ -9,7 +9,7 @@
 
 #include "TraitSharedData.generated.h"
 
-namespace UE::AnimNext
+namespace UE::UAF
 {
 	struct FTraitBinding;
 }
@@ -18,7 +18,7 @@ namespace UE::AnimNext
 	/* @see GetLatentPropertyIndex below for details */ \
 	[&Binding, LatentPropertyHandles](int32 LatentPropertyIndex) \
 	{ \
-		const UE::AnimNext::FLatentPropertyHandle LatentPropertyHandle = LatentPropertyHandles[LatentPropertyIndex - 1]; \
+		const UE::UAF::FLatentPropertyHandle LatentPropertyHandle = LatentPropertyHandles[LatentPropertyIndex - 1]; \
 		if (LatentPropertyHandle.IsOffsetValid()) \
 		{ \
 			/* We remove the 'const' explicitly to avoid exposing a mutable version of the getter */ \
@@ -28,10 +28,10 @@ namespace UE::AnimNext
 
 #define ANIM_NEXT_IMPL_DEFINE_LATENT_CONSTRUCTOR(EnumeratorMacro) \
 	/** @see FAnimNextTraitSharedData::ConstructLatentProperties */ \
-	static void ConstructLatentProperties(const UE::AnimNext::FTraitBinding& Binding) \
+	static void ConstructLatentProperties(const UE::UAF::FTraitBinding& Binding) \
 	{ \
 		/* Construct our latent properties */ \
-		const UE::AnimNext::FLatentPropertyHandle* LatentPropertyHandles = Binding.GetLatentPropertyHandles(); \
+		const UE::UAF::FLatentPropertyHandle* LatentPropertyHandles = Binding.GetLatentPropertyHandles(); \
 		EnumeratorMacro(ANIM_NEXT_IMPL_EXECUTE_LATENT_CONSTRUCTOR_FOR_PROPERTY) \
 	} \
 
@@ -39,7 +39,7 @@ namespace UE::AnimNext
 	/* @see GetLatentPropertyIndex below for details */ \
 	[&Binding, LatentPropertyHandles](int32 LatentPropertyIndex) \
 	{ \
-		const UE::AnimNext::FLatentPropertyHandle LatentPropertyHandle = LatentPropertyHandles[LatentPropertyIndex - 1]; \
+		const UE::UAF::FLatentPropertyHandle LatentPropertyHandle = LatentPropertyHandles[LatentPropertyIndex - 1]; \
 		if (LatentPropertyHandle.IsOffsetValid()) \
 		{ \
 			/* We remove the 'const' explicitly to avoid exposing a mutable version of the getter */ \
@@ -49,10 +49,10 @@ namespace UE::AnimNext
 
 #define ANIM_NEXT_IMPL_DEFINE_LATENT_DESTRUCTOR(EnumeratorMacro) \
 	/** @see FAnimNextTraitSharedData::DestructLatentProperties */ \
-	static void DestructLatentProperties(const UE::AnimNext::FTraitBinding& Binding) \
+	static void DestructLatentProperties(const UE::UAF::FTraitBinding& Binding) \
 	{ \
 		/* Destruct our latent properties */ \
-		const UE::AnimNext::FLatentPropertyHandle* LatentPropertyHandles = Binding.GetLatentPropertyHandles(); \
+		const UE::UAF::FLatentPropertyHandle* LatentPropertyHandles = Binding.GetLatentPropertyHandles(); \
 		EnumeratorMacro(ANIM_NEXT_IMPL_EXECUTE_LATENT_DESTRUCTOR_FOR_PROPERTY) \
 	} \
 
@@ -82,7 +82,7 @@ namespace UE::AnimNext
 	} \
 
 #define ANIM_NEXT_IMPL_DEFINE_LATENT_GETTER(PropertyName) \
-	const decltype(PropertyName)& Get##PropertyName(const UE::AnimNext::FTraitBinding& Binding) const \
+	const decltype(PropertyName)& Get##PropertyName(const UE::UAF::FTraitBinding& Binding) const \
 	{ \
 		/* We need a mapping of latent property name/offset to latent property index */ \
 		/* This can be built once at runtime using the UE reflection and cached on first call or using a constexpr function, see below */ \
@@ -90,8 +90,8 @@ namespace UE::AnimNext
 		constexpr int32 LatentPropertyIndex = GetLatentPropertyIndex(LatentPropertyOffset); \
 		/* A negative or zero property index means the property isn't latent, we can static assert since it isn't possible */ \
 		static_assert(LatentPropertyIndex > 0, "Property " #PropertyName " isn't latent"); \
-		const UE::AnimNext::FLatentPropertyHandle* LatentPropertyHandles = Binding.GetLatentPropertyHandles(); \
-		const UE::AnimNext::FLatentPropertyHandle LatentPropertyHandle = LatentPropertyHandles[LatentPropertyIndex - 1]; \
+		const UE::UAF::FLatentPropertyHandle* LatentPropertyHandles = Binding.GetLatentPropertyHandles(); \
+		const UE::UAF::FLatentPropertyHandle LatentPropertyHandle = LatentPropertyHandles[LatentPropertyIndex - 1]; \
 		/* An invalid offset means we are inline, otherwise we are cached */ \
 		if (!LatentPropertyHandle.IsOffsetValid()) \
 		{ \
@@ -135,12 +135,12 @@ struct FAnimNextTraitSharedData
 	/**
 	 * Constructs the latent properties on the bound trait instance
 	 */
-	static void ConstructLatentProperties(const UE::AnimNext::FTraitBinding& Binding) {}
+	static void ConstructLatentProperties(const UE::UAF::FTraitBinding& Binding) {}
 
 	/**
 	 * Destructs the latent properties on the bound trait instance
 	 */
-	static void DestructLatentProperties(const UE::AnimNext::FTraitBinding& Binding) {}
+	static void DestructLatentProperties(const UE::UAF::FTraitBinding& Binding) {}
 
 	/**
 	  * Returns the latent property index from a latent property offset
diff --git a/Engine/Plugins/Experimental/AnimNextAnimGraph/Source/AnimNextAnimGraph/Public/TraitInterfaces/IUpdate.h b/Engine/Plugins/Experimental/AnimNextAnimGraph/Source/AnimNextAnimGraph/Public/TraitInterfaces/IUpdate.h
--- a/Engine/Plugins/Experimental/AnimNextAnimGraph/Source/AnimNextAnimGraph/Public/TraitInterfaces/IUpdate.h
+++ b/Engine/Plugins/Experimental/AnimNextAnimGraph/Source/AnimNextAnimGraph/Public/TraitInterfaces/IUpdate.h
@@ -17,7 +17,7 @@ struct FRigUnit_AnimNextRunAnimationGraph_v1;
 struct FRigUnit_AnimNextRunAnimationGraph_v2;
 struct FAnimNode_AnimNextGraph;
 
-namespace UE::AnimNext
+namespace UE::UAF
 {
 	struct FUpdateTraversalQueue;
 
@@ -217,8 +217,8 @@ namespace UE::AnimNext
 		void ExecuteBookkeepingActions(Private::FUpdateEventBookkeepingList& BookkeepingList);
 
 		// The input and output event lists
-		UE::AnimNext::FTraitEventList* InputEventList = nullptr;
-		UE::AnimNext::FTraitEventList* OutputEventList = nullptr;
+		UE::UAF::FTraitEventList* InputEventList = nullptr;
+		UE::UAF::FTraitEventList* OutputEventList = nullptr;
 
 		// The currently executing entry
 		Private::FUpdateEntry* ExecutingEntry = nullptr;
diff --git a/Engine/Plugins/Experimental/AnimNextAnimGraph/Source/AnimNextAnimGraph/Public/Traits/Inertialization.h b/Engine/Plugins/Experimental/AnimNextAnimGraph/Source/AnimNextAnimGraph/Public/Traits/Inertialization.h
--- a/Engine/Plugins/Experimental/AnimNextAnimGraph/Source/AnimNextAnimGraph/Public/Traits/Inertialization.h
+++ b/Engine/Plugins/Experimental/AnimNextAnimGraph/Source/AnimNextAnimGraph/Public/Traits/Inertialization.h
@@ -10,7 +10,7 @@
 
 #include "Inertialization.generated.h"
 
-namespace UE::AnimNext
+namespace UE::UAF
 {
 	/**
 	 * FInertializationRequest
@@ -76,5 +76,5 @@ struct FAnimNextInertializationRequestEvent : public FAnimNextTraitEvent
 	GENERATED_BODY()
 	DECLARE_ANIM_TRAIT_EVENT(FAnimNextInertializationRequestEvent, FAnimNextTraitEvent)
 
-	UE::AnimNext::FInertializationRequest Request;
+	UE::UAF::FInertializationRequest Request;
 };
diff --git a/Engine/Plugins/Experimental/AnimNextAnimGraph/Source/AnimNextAnimGraphEditor/Private/AnimNextAnimGraphEditorModule.cpp b/Engine/Plugins/Experimental/AnimNextAnimGraph/Source/AnimNextAnimGraphEditor/Private/AnimNextAnimGraphEditorModule.cpp
--- a/Engine/Plugins/Experimental/AnimNextAnimGraph/Source/AnimNextAnimGraphEditor/Private/AnimNextAnimGraphEditorModule.cpp
+++ b/Engine/Plugins/Experimental/AnimNextAnimGraph/Source/AnimNextAnimGraphEditor/Private/AnimNextAnimGraphEditorModule.cpp
@@ -49,11 +49,11 @@
 namespace
 {
 	FAnimNextAnimGraphTraceModule GAnimNextAnimGraphTraceModule;
-	UE::AnimNextEditor::FEvaluationProgramTrackCreator GAnimNextModulesTrackCreator;
-	UE::AnimNextEditor::FSequenceInfoTrackCreator GSequenceInfoTrackCreator;
+	UE::UAF::Editor::FEvaluationProgramTrackCreator GAnimNextModulesTrackCreator;
+	UE::UAF::Editor::FSequenceInfoTrackCreator GSequenceInfoTrackCreator;
 }
 
-namespace UE::AnimNext::Editor
+namespace UE::UAF::Editor
 {
 
 class FAnimNextAnimGraphEditorModule : public IModuleInterface
@@ -77,7 +77,7 @@ class FAnimNextAnimGraphEditorModule : public IModuleInterface
 					
 		PropertyModule.RegisterCustomPropertyTypeLayout(
 							FAnimNextSequenceTraceInfo::StaticStruct()->GetFName(),
-							FOnGetPropertyTypeCustomizationInstance::CreateLambda([] { return MakeShared<AnimNextEditor::FAnimNextSequenceTraceInfoCustomization>(); }));
+							FOnGetPropertyTypeCustomizationInstance::CreateLambda([] { return MakeShared<UE::UAF::Editor::FAnimNextSequenceTraceInfoCustomization>(); }));
 
 		PropertyModule.RegisterCustomPropertyTypeLayout(
 			FAnimNextFactoryParams::StaticStruct()->GetFName(),
@@ -191,7 +191,7 @@ class FAnimNextAnimGraphEditorModule : public IModuleInterface
 			if (RigVMEdGraph->GetModel()->GetSchemaClass() == UAnimNextAnimationGraphSchema::StaticClass())
 			{
 				TMap<FAssetData, FAnimNextAssetRegistryExports> PerAssetManifestExports;
-				UE::AnimNext::UncookedOnly::FAnimGraphUtils::GetExportedManifestNodesFromAssetRegistry(PerAssetManifestExports);
+				UE::UAF::UncookedOnly::FAnimGraphUtils::GetExportedManifestNodesFromAssetRegistry(PerAssetManifestExports);
 
 				TArray<FAssetData> WorkspaceAssets;
 				if (TSharedPtr<UE::Workspace::IWorkspaceEditor> WorkspaceEditor = InWorkspaceEditorWeak.Pin())
@@ -226,7 +226,7 @@ class FAnimNextAnimGraphEditorModule : public IModuleInterface
 		{
 			if (ActionMenuContextData.RigVMSchema != nullptr && ActionMenuContextData.RigVMSchema->SupportsUnitFunction(ActionMenuContextData.RigVMController, Function))
 			{
-				UE::AnimNext::Editor::FUtils::AddSchemaRigUnitAction(UAnimNextTraitStackUnitNode::StaticClass(), Struct, *Function, InContextMenuBuilder);
+				UE::UAF::Editor::FUtils::AddSchemaRigUnitAction(UAnimNextTraitStackUnitNode::StaticClass(), Struct, *Function, InContextMenuBuilder);
 			}
 		}
 	}
@@ -245,5 +245,5 @@ class FAnimNextAnimGraphEditorModule : public IModuleInterface
 
 #undef LOCTEXT_NAMESPACE
 
-IMPLEMENT_MODULE(UE::AnimNext::Editor::FAnimNextAnimGraphEditorModule, AnimNextAnimGraphEditor);
+IMPLEMENT_MODULE(UE::UAF::Editor::FAnimNextAnimGraphEditorModule, AnimNextAnimGraphEditor);
 
diff --git a/Engine/Plugins/Experimental/AnimNextAnimGraph/Source/AnimNextAnimGraphEditor/Private/Graph/AnimNextAnimationGraphAssetDefinition.cpp b/Engine/Plugins/Experimental/AnimNextAnimGraph/Source/AnimNextAnimGraphEditor/Private/Graph/AnimNextAnimationGraphAssetDefinition.cpp
--- a/Engine/Plugins/Experimental/AnimNextAnimGraph/Source/AnimNextAnimGraphEditor/Private/Graph/AnimNextAnimationGraphAssetDefinition.cpp
+++ b/Engine/Plugins/Experimental/AnimNextAnimGraph/Source/AnimNextAnimGraphEditor/Private/Graph/AnimNextAnimationGraphAssetDefinition.cpp
@@ -14,7 +14,7 @@
 
 EAssetCommandResult UAssetDefinition_AnimNextAnimationGraph::OpenAssets(const FAssetOpenArgs& OpenArgs) const
 {
-	using namespace UE::AnimNext::Editor;
+	using namespace UE::UAF::Editor;
 	using namespace UE::Workspace;
 
 	for (UAnimNextAnimationGraph* Asset : OpenArgs.LoadObjects<UAnimNextAnimationGraph>())
diff --git a/Engine/Plugins/Experimental/AnimNextAnimGraph/Source/AnimNextAnimGraphEditor/Private/Graph/AnimNextAnimationGraphMenuExtensions.cpp b/Engine/Plugins/Experimental/AnimNextAnimGraph/Source/AnimNextAnimGraphEditor/Private/Graph/AnimNextAnimationGraphMenuExtensions.cpp
--- a/Engine/Plugins/Experimental/AnimNextAnimGraph/Source/AnimNextAnimGraphEditor/Private/Graph/AnimNextAnimationGraphMenuExtensions.cpp
+++ b/Engine/Plugins/Experimental/AnimNextAnimGraph/Source/AnimNextAnimGraphEditor/Private/Graph/AnimNextAnimationGraphMenuExtensions.cpp
@@ -25,12 +25,12 @@
 
 #define LOCTEXT_NAMESPACE "FAnimationGraphMenuExtensions"
 
-namespace UE::AnimNext::Editor::Private
+namespace UE::UAF::Editor::Private
 {
 	static const FLazyName VariablesTraitBaseName = TEXT("Variables");
 }
 
-namespace UE::AnimNext::Editor
+namespace UE::UAF::Editor
 {
 
 void FAnimationGraphMenuExtensions::RegisterMenus()
@@ -266,7 +266,7 @@ void FAnimationGraphMenuExtensions::AddManifestSection(UToolMenu* InMenu, const
 
 	const TWeakObjectPtr<const UAnimNextEdGraphNode> EdNodeWeak = AnimNextEdGraphNode;
 
-	if (!UE::AnimNext::UncookedOnly::FAnimGraphUtils::IsExposedToManifest(ModelNode))
+	if (!UE::UAF::UncookedOnly::FAnimGraphUtils::IsExposedToManifest(ModelNode))
 	{
 		Section.AddMenuEntry(
 			"AddTraitToManifest",
@@ -281,7 +281,7 @@ void FAnimationGraphMenuExtensions::AddManifestSection(UToolMenu* InMenu, const
 						{
 							UAnimNextController* Controller = CastChecked<UAnimNextController>(EdGraphNode->GetController());
 
-							if (ensure(UE::AnimNext::UncookedOnly::FAnimGraphUtils::IsExposedToManifest(ModelNode) == false))
+							if (ensure(UE::UAF::UncookedOnly::FAnimGraphUtils::IsExposedToManifest(ModelNode) == false))
 							{
 								FRigVMControllerCompileBracketScope CompileScope(Controller);
 
@@ -289,7 +289,7 @@ void FAnimationGraphMenuExtensions::AddManifestSection(UToolMenu* InMenu, const
 								
 								if (UAnimNextRigVMAssetEditorData* EditorData = EdGraphNode->GetGraph()->GetTypedOuter<UAnimNextRigVMAssetEditorData>())
 								{
-									UE::AnimNext::UncookedOnly::FAnimGraphUtils::RequestVMAutoRecompile(EditorData);  // required to force an asset tags update
+									UE::UAF::UncookedOnly::FAnimGraphUtils::RequestVMAutoRecompile(EditorData);  // required to force an asset tags update
 								}
 							}
 						}								
@@ -316,7 +316,7 @@ void FAnimationGraphMenuExtensions::AddManifestSection(UToolMenu* InMenu, const
 						
 						if (UAnimNextRigVMAssetEditorData* EditorData = EdGraphNode->GetGraph()->GetTypedOuter<UAnimNextRigVMAssetEditorData>())
 						{
-							UE::AnimNext::UncookedOnly::FAnimGraphUtils::RequestVMAutoRecompile(EditorData); // required to force an asset tags update
+							UE::UAF::UncookedOnly::FAnimGraphUtils::RequestVMAutoRecompile(EditorData); // required to force an asset tags update
 						}
 					}
 				}))
diff --git a/Engine/Plugins/Experimental/AnimNextAnimGraph/Source/AnimNextAnimGraphEditor/Private/Graph/AnimNextGraphPanelNodeFactory.cpp b/Engine/Plugins/Experimental/AnimNextAnimGraph/Source/AnimNextAnimGraphEditor/Private/Graph/AnimNextGraphPanelNodeFactory.cpp
--- a/Engine/Plugins/Experimental/AnimNextAnimGraph/Source/AnimNextAnimGraphEditor/Private/Graph/AnimNextGraphPanelNodeFactory.cpp
+++ b/Engine/Plugins/Experimental/AnimNextAnimGraph/Source/AnimNextAnimGraphEditor/Private/Graph/AnimNextGraphPanelNodeFactory.cpp
@@ -9,7 +9,7 @@ TSharedPtr<SGraphNode> FAnimNextGraphPanelNodeFactory::CreateNode(UEdGraphNode*
 {
 	if (UAnimNextEdGraphNode* AnimNextGraphNode = Cast<UAnimNextEdGraphNode>(Node))
 	{
-		if(UE::AnimNext::UncookedOnly::FAnimGraphUtils::IsTraitStackNode(AnimNextGraphNode->GetModelNode()))
+		if(UE::UAF::UncookedOnly::FAnimGraphUtils::IsTraitStackNode(AnimNextGraphNode->GetModelNode()))
 		{
 			TSharedPtr<SGraphNode> GraphNode =
 				SNew(SAnimNextGraphNode)
diff --git a/Engine/Plugins/Experimental/AnimNextAnimGraph/Source/AnimNextAnimGraphEditor/Private/Graph/STraitEditorView.cpp b/Engine/Plugins/Experimental/AnimNextAnimGraph/Source/AnimNextAnimGraphEditor/Private/Graph/STraitEditorView.cpp
--- a/Engine/Plugins/Experimental/AnimNextAnimGraph/Source/AnimNextAnimGraphEditor/Private/Graph/STraitEditorView.cpp
+++ b/Engine/Plugins/Experimental/AnimNextAnimGraph/Source/AnimNextAnimGraphEditor/Private/Graph/STraitEditorView.cpp
@@ -29,7 +29,7 @@
 
 #define LOCTEXT_NAMESPACE "TraitListEditor"
 
-namespace UE::AnimNext::Editor
+namespace UE::UAF::Editor
 {
 
 STraitEditorView::STraitEditorView()
@@ -94,7 +94,7 @@ void STraitEditorView::SetTraitData(const FTraitStackData& InTraitStackData)
 	UAnimNextEdGraphNode* AnimNextEdGraphNode = InTraitStackData.EdGraphNodeWeak.Get();
 
 	// If this is not a trait stack node, clear the shared data
-	if(AnimNextEdGraphNode && !UE::AnimNext::UncookedOnly::FAnimGraphUtils::IsTraitStackNode(AnimNextEdGraphNode->GetModelNode()))
+	if(AnimNextEdGraphNode && !UE::UAF::UncookedOnly::FAnimGraphUtils::IsTraitStackNode(AnimNextEdGraphNode->GetModelNode()))
 	{
 		TraitEditorSharedData->EdGraphNodeWeak = nullptr;
 	}
@@ -310,7 +310,7 @@ void STraitEditorView::RefreshTraitStackTraitsStatus()
 					// Open tab to display errors
 					if (TSharedPtr<UE::Workspace::IWorkspaceEditor> WorkspaceEditor = WorkspaceEditorWeak.Pin())
 					{
-						WorkspaceEditor->GetTabManager()->TryInvokeTab(FTabId(UE::AnimNext::Editor::CompilerResultsTabName));
+						WorkspaceEditor->GetTabManager()->TryInvokeTab(FTabId(UE::UAF::Editor::CompilerResultsTabName));
 					}
 				}
 			}
diff --git a/Engine/Plugins/Experimental/AnimNextAnimGraph/Source/AnimNextAnimGraphEditor/Private/Graph/STraitStackView.cpp b/Engine/Plugins/Experimental/AnimNextAnimGraph/Source/AnimNextAnimGraphEditor/Private/Graph/STraitStackView.cpp
--- a/Engine/Plugins/Experimental/AnimNextAnimGraph/Source/AnimNextAnimGraphEditor/Private/Graph/STraitStackView.cpp
+++ b/Engine/Plugins/Experimental/AnimNextAnimGraph/Source/AnimNextAnimGraphEditor/Private/Graph/STraitStackView.cpp
@@ -26,7 +26,7 @@
 
 #define LOCTEXT_NAMESPACE "TraitListEditor"
 
-namespace UE::AnimNext::Editor
+namespace UE::UAF::Editor
 {
 
 // --- FTraitStackViewEntry ---
@@ -151,7 +151,7 @@ void STraitStackView::Construct(const FArguments& InArgs, TSharedPtr<FTraitEdito
 																{
 																	if (UAnimNextRigVMAssetEditorData* EditorData = EdGraphNode->GetGraph()->GetTypedOuter<UAnimNextRigVMAssetEditorData>())
 																	{
-																		UE::AnimNext::UncookedOnly::FAnimGraphUtils::RequestVMAutoRecompile(EditorData);
+																		UE::UAF::UncookedOnly::FAnimGraphUtils::RequestVMAutoRecompile(EditorData);
 																	}
 																}
 															}
diff --git a/Engine/Plugins/Experimental/AnimNextAnimGraph/Source/AnimNextAnimGraphEditor/Private/Graph/TraitEditorDefs.cpp b/Engine/Plugins/Experimental/AnimNextAnimGraph/Source/AnimNextAnimGraphEditor/Private/Graph/TraitEditorDefs.cpp
--- a/Engine/Plugins/Experimental/AnimNextAnimGraph/Source/AnimNextAnimGraphEditor/Private/Graph/TraitEditorDefs.cpp
+++ b/Engine/Plugins/Experimental/AnimNextAnimGraph/Source/AnimNextAnimGraphEditor/Private/Graph/TraitEditorDefs.cpp
@@ -7,7 +7,7 @@
 
 #define LOCTEXT_NAMESPACE "TraitListEditor"
 
-namespace UE::AnimNext::Editor
+namespace UE::UAF::Editor
 {
 
 /*static*/ FSlateColor FTraitEditorUtils::GetTraitIconErrorDisplayColor(const FTraitStackTraitStatus& InTraitStatus)
@@ -389,6 +389,6 @@ TSharedPtr<SWidget> FTraitListDragDropBase::GetDefaultDecorator() const
 			];
 }
 
-} // namespace UE::AnimNext::Editor
+} // namespace UE::UAF::Editor
 
 #undef LOCTEXT_NAMESPACE
\ No newline at end of file
diff --git a/Engine/Plugins/Experimental/AnimNextAnimGraph/Source/AnimNextAnimGraphEditor/Private/Graph/TraitEditorTabSummoner.cpp b/Engine/Plugins/Experimental/AnimNextAnimGraph/Source/AnimNextAnimGraphEditor/Private/Graph/TraitEditorTabSummoner.cpp
--- a/Engine/Plugins/Experimental/AnimNextAnimGraph/Source/AnimNextAnimGraphEditor/Private/Graph/TraitEditorTabSummoner.cpp
+++ b/Engine/Plugins/Experimental/AnimNextAnimGraph/Source/AnimNextAnimGraphEditor/Private/Graph/TraitEditorTabSummoner.cpp
@@ -5,7 +5,7 @@
 
 #define LOCTEXT_NAMESPACE "WorkspaceTabSummoner"
 
-namespace UE::AnimNext::Editor
+namespace UE::UAF::Editor
 {
 
 const FLazyName TraitEditorTabName("TraitEditorTab");
@@ -32,6 +32,6 @@ FText FTraitEditorTabSummoner::GetTabToolTipText(const FWorkflowTabSpawnInfo& In
 	return ViewMenuTooltip;
 }
 
-} // end namespace UE::AnimNext::Graph
+} // end namespace UE::UAF::Graph
 
 #undef LOCTEXT_NAMESPACE
\ No newline at end of file
diff --git a/Engine/Plugins/Experimental/AnimNextAnimGraph/Source/AnimNextAnimGraphEditor/Private/RewindDebugger/EvaluationProgramTrack.cpp b/Engine/Plugins/Experimental/AnimNextAnimGraph/Source/AnimNextAnimGraphEditor/Private/RewindDebugger/EvaluationProgramTrack.cpp
--- a/Engine/Plugins/Experimental/AnimNextAnimGraph/Source/AnimNextAnimGraphEditor/Private/RewindDebugger/EvaluationProgramTrack.cpp
+++ b/Engine/Plugins/Experimental/AnimNextAnimGraph/Source/AnimNextAnimGraphEditor/Private/RewindDebugger/EvaluationProgramTrack.cpp
@@ -15,7 +15,7 @@
 
 #define LOCTEXT_NAMESPACE "EvaluationProgramTrack"
 
-namespace UE::AnimNextEditor
+namespace UE::UAF::Editor
 {
 
 static const FName AnimNextModulesName("AnimNextModules");
diff --git a/Engine/Plugins/Experimental/AnimNextAnimGraph/Source/AnimNextAnimGraphTestSuite/Private/AnimNextAnimGraphEditorVariablesTest.cpp b/Engine/Plugins/Experimental/AnimNextAnimGraph/Source/AnimNextAnimGraphTestSuite/Private/AnimNextAnimGraphEditorVariablesTest.cpp
--- a/Engine/Plugins/Experimental/AnimNextAnimGraph/Source/AnimNextAnimGraphTestSuite/Private/AnimNextAnimGraphEditorVariablesTest.cpp
+++ b/Engine/Plugins/Experimental/AnimNextAnimGraph/Source/AnimNextAnimGraphTestSuite/Private/AnimNextAnimGraphEditorVariablesTest.cpp
@@ -24,14 +24,14 @@
 
 #if WITH_DEV_AUTOMATION_TESTS && WITH_EDITOR
 
-namespace UE::AnimNext::Tests
+namespace UE::UAF::Tests
 {
 
 IMPLEMENT_SIMPLE_AUTOMATION_TEST(FEditor_AnimGraph_Variables, "Animation.AnimNext.Editor.AnimGraph.Variables", EAutomationTestFlags::EditorContext | EAutomationTestFlags::EngineFilter)
 
 bool FEditor_AnimGraph_Variables::RunTest(const FString& InParameters)
 {
-	using namespace UE::AnimNext;
+	using namespace UE::UAF;
 
 	struct FFactoryAndClass
 	{
@@ -130,7 +130,7 @@ IMPLEMENT_SIMPLE_AUTOMATION_TEST(FEditor_AnimGraph_Graphs, "Animation.AnimNext.E
 
 bool FEditor_AnimGraph_Graphs::RunTest(const FString& InParameters)
 {
-	using namespace UE::AnimNext;
+	using namespace UE::UAF;
 
 	struct FTestSettings
 	{
@@ -232,7 +232,7 @@ IMPLEMENT_SIMPLE_AUTOMATION_TEST(FEditor_AnimGraph_Variables_Python, "Animation.
 
 bool FEditor_AnimGraph_Variables_Python::RunTest(const FString& InParameters)
 {
-	using namespace UE::AnimNext;
+	using namespace UE::UAF;
 
 	const TCHAR* Script = TEXT(
 		"asset_tools = unreal.AssetToolsHelpers.get_asset_tools()\n"
diff --git a/Engine/Plugins/Experimental/AnimNextAnimGraph/Source/AnimNextAnimGraphTestSuite/Private/AnimNextAnimGraphTraitEventTest.cpp b/Engine/Plugins/Experimental/AnimNextAnimGraph/Source/AnimNextAnimGraphTestSuite/Private/AnimNextAnimGraphTraitEventTest.cpp
--- a/Engine/Plugins/Experimental/AnimNextAnimGraph/Source/AnimNextAnimGraphTestSuite/Private/AnimNextAnimGraphTraitEventTest.cpp
+++ b/Engine/Plugins/Experimental/AnimNextAnimGraph/Source/AnimNextAnimGraphTestSuite/Private/AnimNextAnimGraphTraitEventTest.cpp
@@ -24,7 +24,7 @@
 // AnimNext Runtime Trait Graph Event Tests
 //****************************************************************************
 
-namespace UE::AnimNext
+namespace UE::UAF
 {
 	struct FTraitGraphTest_EventAB_NoChildren : FBaseTrait
 	{
@@ -187,7 +187,7 @@ IMPLEMENT_SIMPLE_AUTOMATION_TEST(FAnimationAnimNextRuntimeTest_GraphTraitEvent,
 
 bool FAnimationAnimNextRuntimeTest_GraphTraitEvent::RunTest(const FString& InParameters)
 {
-	using namespace UE::AnimNext;
+	using namespace UE::UAF;
 
 	{
 		AUTO_REGISTER_ANIM_TRAIT(FTraitGraphTest_EventAB_NoChildren)
diff --git a/Engine/Plugins/Experimental/AnimNextAnimGraph/Source/AnimNextAnimGraphTestSuite/Private/AnimNextAnimGraphTraitEventTest.h b/Engine/Plugins/Experimental/AnimNextAnimGraph/Source/AnimNextAnimGraphTestSuite/Private/AnimNextAnimGraphTraitEventTest.h
--- a/Engine/Plugins/Experimental/AnimNextAnimGraph/Source/AnimNextAnimGraphTestSuite/Private/AnimNextAnimGraphTraitEventTest.h
+++ b/Engine/Plugins/Experimental/AnimNextAnimGraph/Source/AnimNextAnimGraphTestSuite/Private/AnimNextAnimGraphTraitEventTest.h
@@ -13,7 +13,7 @@ struct FTraitAnimGraphTest_EventA : public FAnimNextTraitEvent
 
 	bool bTestFlag = false;
 
-	TArray<UE::AnimNext::FTraitUID> VisitedTraits;
+	TArray<UE::UAF::FTraitUID> VisitedTraits;
 };
 
 struct FTraitAnimGraphTest_EventB : public FAnimNextTraitEvent
@@ -23,7 +23,7 @@ struct FTraitAnimGraphTest_EventB : public FAnimNextTraitEvent
 	bool bTestFlag0 = false;
 	bool bTestFlag1 = false;
 
-	TArray<UE::AnimNext::FTraitUID> VisitedTraits;
+	TArray<UE::UAF::FTraitUID> VisitedTraits;
 
 	FAnimNextTraitEventPtr ChildEvent;
 };
diff --git a/Engine/Plugins/Experimental/AnimNextAnimGraph/Source/AnimNextAnimGraphTestSuite/Private/AnimNextAnimGraphTraitGraphTest.cpp b/Engine/Plugins/Experimental/AnimNextAnimGraph/Source/AnimNextAnimGraphTestSuite/Private/AnimNextAnimGraphTraitGraphTest.cpp
--- a/Engine/Plugins/Experimental/AnimNextAnimGraph/Source/AnimNextAnimGraphTestSuite/Private/AnimNextAnimGraphTraitGraphTest.cpp
+++ b/Engine/Plugins/Experimental/AnimNextAnimGraph/Source/AnimNextAnimGraphTestSuite/Private/AnimNextAnimGraphTraitGraphTest.cpp
@@ -33,7 +33,7 @@
 // AnimNext Runtime Trait Graph Tests
 //****************************************************************************
 
-namespace UE::AnimNext
+namespace UE::UAF
 {
 	struct FTestTrait final : FBaseTrait, IEvaluate, IUpdate
 	{
@@ -100,7 +100,7 @@ IMPLEMENT_SIMPLE_AUTOMATION_TEST(FAnimationAnimNextEditorTest_GraphAddTrait, "An
 
 bool FAnimationAnimNextEditorTest_GraphAddTrait::RunTest(const FString& InParameters)
 {
-	using namespace UE::AnimNext;
+	using namespace UE::UAF;
 
 	{
 		AUTO_REGISTER_ANIM_TRAIT(FTestTrait)
@@ -210,7 +210,7 @@ IMPLEMENT_SIMPLE_AUTOMATION_TEST(FAnimationAnimNextEditorTest_GraphTraitOperatio
 
 bool FAnimationAnimNextEditorTest_GraphTraitOperations::RunTest(const FString& InParameters)
 {
-	using namespace UE::AnimNext;
+	using namespace UE::UAF;
 
 	{
 		AUTO_REGISTER_ANIM_TRAIT(FTestTrait)
@@ -381,7 +381,7 @@ IMPLEMENT_SIMPLE_AUTOMATION_TEST(FAnimationAnimNextEditorTest_GraphManifest, "An
 
 bool FAnimationAnimNextEditorTest_GraphManifest::RunTest(const FString& InParameters)
 {
-	using namespace UE::AnimNext;
+	using namespace UE::UAF;
 
 	{
 		FScopedClearNodeTemplateRegistry ScopedClearNodeTemplateRegistry;
@@ -409,17 +409,17 @@ bool FAnimationAnimNextEditorTest_GraphManifest::RunTest(const FString& InParame
 		// --- Add to Manifest ---
 		UE_RETURN_ON_ERROR(AnimNextController->AddNodeToManifestByName(TraitStackNode->GetFName()), TEXT("FAnimationAnimNextEditorTest_GraphManifest -> Add Node to Manifest Failed."));
 		UE_RETURN_ON_ERROR(TraitStackNode->IsExposedToManifest() == true, TEXT("FAnimationAnimNextEditorTest_GraphManifest -> Created node should be at Manifest."));
-		UE_RETURN_ON_ERROR(UE::AnimNext::UncookedOnly::FAnimGraphUtils::IsExposedToManifest(TraitStackNode) == true, TEXT("FAnimationAnimNextEditorTest_GraphManifest -> FAnimGraphUtils::IsExposedToManifest should return true."));
+		UE_RETURN_ON_ERROR(UE::UAF::UncookedOnly::FAnimGraphUtils::IsExposedToManifest(TraitStackNode) == true, TEXT("FAnimationAnimNextEditorTest_GraphManifest -> FAnimGraphUtils::IsExposedToManifest should return true."));
 
 		// --- Undo Add to Manifest ---
 		AnimNextController->Undo();
 		UE_RETURN_ON_ERROR(TraitStackNode->IsExposedToManifest() == false, TEXT("FAnimationAnimNextEditorTest_GraphManifest -> Created node should not be at Manifest after Undo."));
-		UE_RETURN_ON_ERROR(UE::AnimNext::UncookedOnly::FAnimGraphUtils::IsExposedToManifest(TraitStackNode) == false, TEXT("FAnimationAnimNextEditorTest_GraphManifest -> FAnimGraphUtils::IsExposedToManifest should return false after Undo."));
+		UE_RETURN_ON_ERROR(UE::UAF::UncookedOnly::FAnimGraphUtils::IsExposedToManifest(TraitStackNode) == false, TEXT("FAnimationAnimNextEditorTest_GraphManifest -> FAnimGraphUtils::IsExposedToManifest should return false after Undo."));
 
 		// --- Redo Add to Manifest ---
 		AnimNextController->Redo();
 		UE_RETURN_ON_ERROR(TraitStackNode->IsExposedToManifest() == true, TEXT("FAnimationAnimNextEditorTest_GraphManifest -> Created node should be at Manifest after Redo."));
-		UE_RETURN_ON_ERROR(UE::AnimNext::UncookedOnly::FAnimGraphUtils::IsExposedToManifest(TraitStackNode) == true, TEXT("FAnimationAnimNextEditorTest_GraphManifest -> FAnimGraphUtils::IsExposedToManifest should return true after Redo."));
+		UE_RETURN_ON_ERROR(UE::UAF::UncookedOnly::FAnimGraphUtils::IsExposedToManifest(TraitStackNode) == true, TEXT("FAnimationAnimNextEditorTest_GraphManifest -> FAnimGraphUtils::IsExposedToManifest should return true after Redo."));
 
 		// --- Check Registry Exports ---
 		FAnimNextAssetRegistryExports OutExports;
@@ -431,29 +431,29 @@ bool FAnimationAnimNextEditorTest_GraphManifest::RunTest(const FString& InParame
 			return NumNodes;
 		};
 		
-		UE::AnimNext::UncookedOnly::FAnimGraphUtils::GetAssetManifestNodesRegistryExports(EditorData, OutExports);
+		UE::UAF::UncookedOnly::FAnimGraphUtils::GetAssetManifestNodesRegistryExports(EditorData, OutExports);
 		UE_RETURN_ON_ERROR(CountNumberOfManifestNodes() == 1, TEXT("FAnimationAnimNextEditorTest_GraphManifest -> Asset Registry Exports Num should be 1."));
 		OutExports.Exports.Reset();
 
 		// --- Undo Add to Manifest ---
 		AnimNextController->Undo();
 		UE_RETURN_ON_ERROR(TraitStackNode->IsExposedToManifest() == false, TEXT("FAnimationAnimNextEditorTest_GraphManifest -> Created node should not be at Manifest after Undo."));
-		UE_RETURN_ON_ERROR(UE::AnimNext::UncookedOnly::FAnimGraphUtils::IsExposedToManifest(TraitStackNode) == false, TEXT("FAnimationAnimNextEditorTest_GraphManifest -> FAnimGraphUtils::IsExposedToManifest should return false after Undo."));
+		UE_RETURN_ON_ERROR(UE::UAF::UncookedOnly::FAnimGraphUtils::IsExposedToManifest(TraitStackNode) == false, TEXT("FAnimationAnimNextEditorTest_GraphManifest -> FAnimGraphUtils::IsExposedToManifest should return false after Undo."));
 		OutExports.Exports.Reset();
 
 		// --- Check Registry Exports ---
-		UE::AnimNext::UncookedOnly::FAnimGraphUtils::GetAssetManifestNodesRegistryExports(EditorData, OutExports);
+		UE::UAF::UncookedOnly::FAnimGraphUtils::GetAssetManifestNodesRegistryExports(EditorData, OutExports);
 		UE_RETURN_ON_ERROR(CountNumberOfManifestNodes() == 0, TEXT("FAnimationAnimNextEditorTest_GraphManifest -> Asset Registry Exports Num should be 0 after Undo."));
 		OutExports.Exports.Reset();
 
 		// --- Redo Add to Manifest ---
 		AnimNextController->Redo();
 		UE_RETURN_ON_ERROR(TraitStackNode->IsExposedToManifest() == true, TEXT("FAnimationAnimNextEditorTest_GraphManifest -> Created node should be at Manifest after Redo."));
-		UE_RETURN_ON_ERROR(UE::AnimNext::UncookedOnly::FAnimGraphUtils::IsExposedToManifest(TraitStackNode) == true, TEXT("FAnimationAnimNextEditorTest_GraphManifest -> FAnimGraphUtils::IsExposedToManifest should return true after Redo."));
+		UE_RETURN_ON_ERROR(UE::UAF::UncookedOnly::FAnimGraphUtils::IsExposedToManifest(TraitStackNode) == true, TEXT("FAnimationAnimNextEditorTest_GraphManifest -> FAnimGraphUtils::IsExposedToManifest should return true after Redo."));
 		OutExports.Exports.Reset();
 
 		// --- Check Registry Exports ---
-		UE::AnimNext::UncookedOnly::FAnimGraphUtils::GetAssetManifestNodesRegistryExports(EditorData, OutExports);
+		UE::UAF::UncookedOnly::FAnimGraphUtils::GetAssetManifestNodesRegistryExports(EditorData, OutExports);
 		UE_RETURN_ON_ERROR(CountNumberOfManifestNodes() == 1, TEXT("FAnimationAnimNextEditorTest_GraphManifest -> Asset Registry Exports Num should be 1 after Redo."));
 		OutExports.Exports.Reset();
 	}
@@ -467,7 +467,7 @@ IMPLEMENT_SIMPLE_AUTOMATION_TEST(FAnimationAnimNextRuntimeTest_GraphExecute, "An
 
 bool FAnimationAnimNextRuntimeTest_GraphExecute::RunTest(const FString& InParameters)
 {
-	using namespace UE::AnimNext;
+	using namespace UE::UAF;
 
 	{
 		AUTO_REGISTER_ANIM_TRAIT(FTestTrait)
@@ -556,9 +556,9 @@ bool FAnimationAnimNextRuntimeTest_GraphExecute::RunTest(const FString& InParame
 
 		{
 			FUpdateGraphContext UpdateGraphContext(*GraphInstance.Get(), 1.0f / 30.0f);
-			UE::AnimNext::UpdateGraph(UpdateGraphContext);
-			const UE::AnimNext::FEvaluateGraphContext EvaluateGraphContext(*GraphInstance.Get(), UE::AnimNext::FReferencePose(), 0);
-			(void)UE::AnimNext::EvaluateGraph(EvaluateGraphContext);
+			UE::UAF::UpdateGraph(UpdateGraphContext);
+			const UE::UAF::FEvaluateGraphContext EvaluateGraphContext(*GraphInstance.Get(), UE::UAF::FReferencePose(), 0);
+			(void)UE::UAF::EvaluateGraph(EvaluateGraphContext);
 		}
 
 		AddErrorIfFalse(Messages.Num() == 7, "FAnimationAnimNextRuntimeTest_GraphExecute - > Unexpected message count");
@@ -580,7 +580,7 @@ IMPLEMENT_SIMPLE_AUTOMATION_TEST(FAnimationAnimNextRuntimeTest_GraphExecuteLaten
 
 bool FAnimationAnimNextRuntimeTest_GraphExecuteLatent::RunTest(const FString& InParameters)
 {
-	using namespace UE::AnimNext;
+	using namespace UE::UAF;
 
 	{
 		AUTO_REGISTER_ANIM_TRAIT(FTestTrait)
@@ -694,9 +694,9 @@ bool FAnimationAnimNextRuntimeTest_GraphExecuteLatent::RunTest(const FString& In
 		
 		{
 			FUpdateGraphContext UpdateGraphContext(*GraphInstance.Get(), 1.0f / 30.0f);
-			UE::AnimNext::UpdateGraph(UpdateGraphContext);
-			const UE::AnimNext::FEvaluateGraphContext EvaluateGraphContext(*GraphInstance.Get(), UE::AnimNext::FReferencePose(), 0);
-			(void)UE::AnimNext::EvaluateGraph(EvaluateGraphContext);
+			UE::UAF::UpdateGraph(UpdateGraphContext);
+			const UE::UAF::FEvaluateGraphContext EvaluateGraphContext(*GraphInstance.Get(), UE::UAF::FReferencePose(), 0);
+			(void)UE::UAF::EvaluateGraph(EvaluateGraphContext);
 		}
 
 		AddErrorIfFalse(Messages.Num() == 7, "FAnimationAnimNextRuntimeTest_GraphExecuteLatent - > Unexpected message count");
@@ -719,7 +719,7 @@ IMPLEMENT_SIMPLE_AUTOMATION_TEST(FAnimationAnimNextRuntimeTest_Variables, "Anima
 
 bool FAnimationAnimNextRuntimeTest_Variables::RunTest(const FString& InParameters)
 {
-	using namespace UE::AnimNext;
+	using namespace UE::UAF;
 
 	{
 		UFactory* GraphFactory = NewObject<UAnimNextAnimationGraphFactory>();
diff --git a/Engine/Plugins/Experimental/AnimNextAnimGraph/Source/AnimNextAnimGraphTestSuite/Private/AnimNextAnimGraphTraitInterfacesTest.cpp b/Engine/Plugins/Experimental/AnimNextAnimGraph/Source/AnimNextAnimGraphTestSuite/Private/AnimNextAnimGraphTraitInterfacesTest.cpp
--- a/Engine/Plugins/Experimental/AnimNextAnimGraph/Source/AnimNextAnimGraphTestSuite/Private/AnimNextAnimGraphTraitInterfacesTest.cpp
+++ b/Engine/Plugins/Experimental/AnimNextAnimGraph/Source/AnimNextAnimGraphTestSuite/Private/AnimNextAnimGraphTraitInterfacesTest.cpp
@@ -32,7 +32,7 @@
 // AnimNext Runtime TraitInterfaces Tests
 //****************************************************************************
 
-namespace UE::AnimNext
+namespace UE::UAF
 {
 	namespace Private
 	{
@@ -348,7 +348,7 @@ namespace UE::AnimNext
 	#endif // WITH_EDITOR
 	};
 
-} // end namespace UE::AnimNext
+} // end namespace UE::UAF
 
 // --- Runtime Test Trait Interface Registry ---
 
@@ -356,7 +356,7 @@ IMPLEMENT_SIMPLE_AUTOMATION_TEST(FAnimationAnimNextRuntimeTest_TraitInterfaceReg
 
 bool FAnimationAnimNextRuntimeTest_TraitInterfaceRegistry::RunTest(const FString& InParameters)
 {
-	using namespace UE::AnimNext;
+	using namespace UE::UAF;
 
 	{
 		FTraitInterfaceRegistry& Registry = FTraitInterfaceRegistry::Get();
@@ -419,7 +419,7 @@ IMPLEMENT_SIMPLE_AUTOMATION_TEST(FAnimationAnimNextRuntimeTest_IHierarchy, "Anim
 
 bool FAnimationAnimNextRuntimeTest_IHierarchy::RunTest(const FString& InParameters)
 {
-	using namespace UE::AnimNext;
+	using namespace UE::UAF;
 
 	{
 		AUTO_REGISTER_ANIM_TRAIT(FTraitWithNoChildren)
@@ -642,7 +642,7 @@ IMPLEMENT_SIMPLE_AUTOMATION_TEST(FAnimationAnimNextRuntimeTest_IUpdate, "Animati
 
 bool FAnimationAnimNextRuntimeTest_IUpdate::RunTest(const FString& InParameters)
 {
-	using namespace UE::AnimNext;
+	using namespace UE::UAF;
 	{
 		AUTO_REGISTER_ANIM_TRAIT(FTraitWithNoChildren)
 		AUTO_REGISTER_ANIM_TRAIT(FTraitWithOneChild)
@@ -784,7 +784,7 @@ IMPLEMENT_SIMPLE_AUTOMATION_TEST(FAnimationAnimNextRuntimeTest_IEvaluate, "Anima
 
 bool FAnimationAnimNextRuntimeTest_IEvaluate::RunTest(const FString& InParameters)
 {
-	using namespace UE::AnimNext;
+	using namespace UE::UAF;
 	{
 		AUTO_REGISTER_ANIM_TRAIT(FTraitWithNoChildren)
 		AUTO_REGISTER_ANIM_TRAIT(FTraitWithOneChild)
@@ -894,7 +894,7 @@ bool FAnimationAnimNextRuntimeTest_IEvaluate::RunTest(const FString& InParameter
 			Private::EvaluatedTraits = &EvaluatedTraits;
 
 			// Call pre/post evaluate on our graph
-			UE::AnimNext::FEvaluateGraphContext EvaluateGraphContext(*GraphInstance.Get(), UE::AnimNext::FReferencePose(), 0);
+			UE::UAF::FEvaluateGraphContext EvaluateGraphContext(*GraphInstance.Get(), UE::UAF::FReferencePose(), 0);
 			(void)EvaluateGraph(EvaluateGraphContext);
 
 			AddErrorIfFalse(EvaluatedTraits.Num() == 6, "FAnimationAnimNextRuntimeTest_IEvaluate -> Expected 6 nodes to have been visited during the evaluate traversal");
@@ -925,7 +925,7 @@ IMPLEMENT_SIMPLE_AUTOMATION_TEST(FAnimationAnimNextRuntimeTest_IScopedInterface,
 
 bool FAnimationAnimNextRuntimeTest_IScopedInterface::RunTest(const FString& InParameters)
 {
-	using namespace UE::AnimNext;
+	using namespace UE::UAF;
 	{
 		AUTO_REGISTER_ANIM_TRAIT(FTraitWithOneChild)
 		AUTO_REGISTER_ANIM_TRAIT(FScopedTagTrait)
diff --git a/Engine/Plugins/Experimental/AnimNextAnimGraph/Source/AnimNextAnimGraphTestSuite/Private/AnimNextTraitBaseTest.cpp b/Engine/Plugins/Experimental/AnimNextAnimGraph/Source/AnimNextAnimGraphTestSuite/Private/AnimNextTraitBaseTest.cpp
--- a/Engine/Plugins/Experimental/AnimNextAnimGraph/Source/AnimNextAnimGraphTestSuite/Private/AnimNextTraitBaseTest.cpp
+++ b/Engine/Plugins/Experimental/AnimNextAnimGraph/Source/AnimNextAnimGraphTestSuite/Private/AnimNextTraitBaseTest.cpp
@@ -29,7 +29,7 @@
 // AnimNext Runtime TraitBase Tests
 //****************************************************************************
 
-namespace UE::AnimNext
+namespace UE::UAF
 {
 	namespace Private
 	{
@@ -317,7 +317,7 @@ IMPLEMENT_SIMPLE_AUTOMATION_TEST(FAnimationAnimNextRuntimeTest_TraitRegistry, "A
 
 bool FAnimationAnimNextRuntimeTest_TraitRegistry::RunTest(const FString& InParameters)
 {
-	using namespace UE::AnimNext;
+	using namespace UE::UAF;
 
 	FTraitRegistry& Registry = FTraitRegistry::Get();
 
@@ -437,7 +437,7 @@ IMPLEMENT_SIMPLE_AUTOMATION_TEST(FAnimationAnimNextRuntimeTest_NodeTemplateRegis
 
 bool FAnimationAnimNextRuntimeTest_NodeTemplateRegistry::RunTest(const FString& InParameters)
 {
-	using namespace UE::AnimNext;
+	using namespace UE::UAF;
 
 	AUTO_REGISTER_ANIM_TRAIT(FTraitA_Base)
 	AUTO_REGISTER_ANIM_TRAIT(FTraitAB_Add)
@@ -518,7 +518,7 @@ IMPLEMENT_SIMPLE_AUTOMATION_TEST(FAnimationAnimNextRuntimeTest_NodeLifetime, "An
 
 bool FAnimationAnimNextRuntimeTest_NodeLifetime::RunTest(const FString& InParameters)
 {
-	using namespace UE::AnimNext;
+	using namespace UE::UAF;
 	{
 		AUTO_REGISTER_ANIM_TRAIT(FTraitA_Base)
 		AUTO_REGISTER_ANIM_TRAIT(FTraitAB_Add)
@@ -727,7 +727,7 @@ IMPLEMENT_SIMPLE_AUTOMATION_TEST(FAnimationAnimNextRuntimeTest_TraitStackBinding
 
 bool FAnimationAnimNextRuntimeTest_TraitStackBinding::RunTest(const FString& InParameters)
 {
-	using namespace UE::AnimNext;
+	using namespace UE::UAF;
 
 	{
 		AUTO_REGISTER_ANIM_TRAIT(FTraitA_Base)
@@ -1149,7 +1149,7 @@ IMPLEMENT_SIMPLE_AUTOMATION_TEST(FAnimationAnimNextRuntimeTest_TraitSerializatio
 
 bool FAnimationAnimNextRuntimeTest_TraitSerialization::RunTest(const FString& InParameters)
 {
-	using namespace UE::AnimNext;
+	using namespace UE::UAF;
 
 	{
 		FScopedClearNodeTemplateRegistry ScopedClearNodeTemplateRegistry;
@@ -1505,7 +1505,7 @@ IMPLEMENT_SIMPLE_AUTOMATION_TEST(FAnimationAnimNextRuntimeTest_SimpleGraphBuilde
 
 bool FAnimationAnimNextRuntimeTest_SimpleGraphBuilder::RunTest(const FString& InParameters)
 {
-	using namespace UE::AnimNext;
+	using namespace UE::UAF;
 
 	{
 		FScopedClearNodeTemplateRegistry ScopedClearNodeTemplateRegistry;
@@ -1654,21 +1654,21 @@ bool FAnimationAnimNextRuntimeTest_SimpleGraphBuilder::RunTest(const FString& In
 
 FTraitA_BaseSharedData::FTraitA_BaseSharedData()
 #if WITH_DEV_AUTOMATION_TESTS
-	: TraitUID(UE::AnimNext::FTraitA_Base::TraitUID.GetUID())
+	: TraitUID(UE::UAF::FTraitA_Base::TraitUID.GetUID())
 #endif
 {
 }
 
 FTraitAB_AddSharedData::FTraitAB_AddSharedData()
 #if WITH_DEV_AUTOMATION_TESTS
-	: TraitUID(UE::AnimNext::FTraitAB_Add::TraitUID.GetUID())
+	: TraitUID(UE::UAF::FTraitAB_Add::TraitUID.GetUID())
 #endif
 {
 }
 
 FTraitAC_AddSharedData::FTraitAC_AddSharedData()
 #if WITH_DEV_AUTOMATION_TESTS
-	: TraitUID(UE::AnimNext::FTraitAC_Add::TraitUID.GetUID())
+	: TraitUID(UE::UAF::FTraitAC_Add::TraitUID.GetUID())
 #endif
 {
 }
diff --git a/Engine/Plugins/Experimental/AnimNextAnimGraph/Source/AnimNextAnimGraphTestSuite/Private/AnimNextTraitEventTest.cpp b/Engine/Plugins/Experimental/AnimNextAnimGraph/Source/AnimNextAnimGraphTestSuite/Private/AnimNextTraitEventTest.cpp
--- a/Engine/Plugins/Experimental/AnimNextAnimGraph/Source/AnimNextAnimGraphTestSuite/Private/AnimNextTraitEventTest.cpp
+++ b/Engine/Plugins/Experimental/AnimNextAnimGraph/Source/AnimNextAnimGraphTestSuite/Private/AnimNextTraitEventTest.cpp
@@ -25,7 +25,7 @@
 // AnimNext Runtime TraitEvent Tests
 //****************************************************************************
 
-namespace UE::AnimNext
+namespace UE::UAF
 {
 	struct FTraitCoreTest_EventAB_Base : FBaseTrait
 	{
@@ -120,7 +120,7 @@ IMPLEMENT_SIMPLE_AUTOMATION_TEST(FAnimationAnimNextRuntimeTest_TraitEventLifetim
 
 bool FAnimationAnimNextRuntimeTest_TraitEventLifetime::RunTest(const FString& InParameters)
 {
-	using namespace UE::AnimNext;
+	using namespace UE::UAF;
 
 	{
 		FTraitEventLifetime Lifetime;
@@ -187,7 +187,7 @@ IMPLEMENT_SIMPLE_AUTOMATION_TEST(FAnimationAnimNextRuntimeTest_TraitEvent, "Anim
 
 bool FAnimationAnimNextRuntimeTest_TraitEvent::RunTest(const FString& InParameters)
 {
-	using namespace UE::AnimNext;
+	using namespace UE::UAF;
 
 	{
 		FAnimNextTraitEvent Event;
@@ -233,7 +233,7 @@ IMPLEMENT_SIMPLE_AUTOMATION_TEST(FAnimationAnimNextRuntimeTest_TraitEventRaising
 
 bool FAnimationAnimNextRuntimeTest_TraitEventRaising::RunTest(const FString& InParameters)
 {
-	using namespace UE::AnimNext;
+	using namespace UE::UAF;
 
 	// Test event raising on a node trait stack
 	{
diff --git a/Engine/Plugins/Experimental/AnimNextAnimGraph/Source/AnimNextAnimGraphTestSuite/Private/AnimNextTraitEventTest.h b/Engine/Plugins/Experimental/AnimNextAnimGraph/Source/AnimNextAnimGraphTestSuite/Private/AnimNextTraitEventTest.h
--- a/Engine/Plugins/Experimental/AnimNextAnimGraph/Source/AnimNextAnimGraphTestSuite/Private/AnimNextTraitEventTest.h
+++ b/Engine/Plugins/Experimental/AnimNextAnimGraph/Source/AnimNextAnimGraphTestSuite/Private/AnimNextTraitEventTest.h
@@ -13,12 +13,12 @@ struct FTraitCoreTest_EventA : public FAnimNextTraitEvent
 
 	bool bAlwaysForwardToBase = true;
 
-	TArray<UE::AnimNext::FTraitUID> VisitedTraits;
+	TArray<UE::UAF::FTraitUID> VisitedTraits;
 };
 
 struct FTraitCoreTest_EventB : public FAnimNextTraitEvent
 {
 	DECLARE_ANIM_TRAIT_EVENT(FTraitCoreTest_EventB, FAnimNextTraitEvent)
 
-	TArray<UE::AnimNext::FTraitUID> VisitedTraits;
+	TArray<UE::UAF::FTraitUID> VisitedTraits;
 };
diff --git a/Engine/Plugins/Experimental/AnimNextAnimGraph/Source/AnimNextAnimGraphTestSuite/Private/Module.cpp b/Engine/Plugins/Experimental/AnimNextAnimGraph/Source/AnimNextAnimGraphTestSuite/Private/Module.cpp
--- a/Engine/Plugins/Experimental/AnimNextAnimGraph/Source/AnimNextAnimGraphTestSuite/Private/Module.cpp
+++ b/Engine/Plugins/Experimental/AnimNextAnimGraph/Source/AnimNextAnimGraphTestSuite/Private/Module.cpp
@@ -4,11 +4,11 @@
 #include "Modules/ModuleInterface.h"
 #include "Modules/ModuleManager.h"
 
-namespace UE::AnimNext::AnimGraph::Tests
+namespace UE::UAF::AnimGraph::Tests
 {
 	class FModule : public IModuleInterface
 	{
 	};
 }
 
-IMPLEMENT_MODULE(UE::AnimNext::AnimGraph::Tests::FModule, AnimNextAnimGraphTestSuite)
+IMPLEMENT_MODULE(UE::UAF::AnimGraph::Tests::FModule, AnimNextAnimGraphTestSuite)
diff --git a/Engine/Plugins/Experimental/AnimNextAnimGraph/Source/AnimNextAnimGraphUncookedOnly/Internal/Entries/AnimNextAnimationGraphEntry.h b/Engine/Plugins/Experimental/AnimNextAnimGraph/Source/AnimNextAnimGraphUncookedOnly/Internal/Entries/AnimNextAnimationGraphEntry.h
--- a/Engine/Plugins/Experimental/AnimNextAnimGraph/Source/AnimNextAnimGraphUncookedOnly/Internal/Entries/AnimNextAnimationGraphEntry.h
+++ b/Engine/Plugins/Experimental/AnimNextAnimGraph/Source/AnimNextAnimGraphUncookedOnly/Internal/Entries/AnimNextAnimationGraphEntry.h
@@ -13,13 +13,13 @@ class UAnimNextRigVMAssetEditorData;
 class UAnimNextEdGraph;
 enum class ERigVMGraphNotifType : uint8;
 
-namespace UE::AnimNext::Editor
+namespace UE::UAF::Editor
 {
 	struct FUtils;
 	class FAnimNextAnimGraphEditorModule;
 }
 
-namespace UE::AnimNext::Tests
+namespace UE::UAF::Tests
 {
 	class FEditor_Graphs;
 	class FEditor_AnimGraph_Graphs;
@@ -33,10 +33,10 @@ class UAnimNextAnimationGraphEntry : public UAnimNextRigVMAssetEntry, public IAn
 	GENERATED_BODY()
 
 	friend class UAnimNextAnimationGraph_EditorData;
-	friend struct UE::AnimNext::Editor::FUtils;
-	friend class UE::AnimNext::Editor::FAnimNextAnimGraphEditorModule;
-	friend class UE::AnimNext::Tests::FEditor_Graphs;
-	friend class UE::AnimNext::Tests::FEditor_AnimGraph_Graphs;
+	friend struct UE::UAF::Editor::FUtils;
+	friend class UE::UAF::Editor::FAnimNextAnimGraphEditorModule;
+	friend class UE::UAF::Tests::FEditor_Graphs;
+	friend class UE::UAF::Tests::FEditor_AnimGraph_Graphs;
 
 	// IAnimNextRigVMExportInterface interface
 	virtual FAnimNextParamType GetExportType() const override;
diff --git a/Engine/Plugins/Experimental/AnimNextAnimGraph/Source/AnimNextAnimGraphUncookedOnly/Internal/Graph/AnimGraphEditorSchemaActions.cpp b/Engine/Plugins/Experimental/AnimNextAnimGraph/Source/AnimNextAnimGraphUncookedOnly/Internal/Graph/AnimGraphEditorSchemaActions.cpp
--- a/Engine/Plugins/Experimental/AnimNextAnimGraph/Source/AnimNextAnimGraphUncookedOnly/Internal/Graph/AnimGraphEditorSchemaActions.cpp
+++ b/Engine/Plugins/Experimental/AnimNextAnimGraph/Source/AnimNextAnimGraphUncookedOnly/Internal/Graph/AnimGraphEditorSchemaActions.cpp
@@ -47,7 +47,7 @@ UEdGraphNode* FAnimNextSchemaAction_AddManifestNode::PerformAction(UEdGraph* Par
 	{
 		if (UAnimNextController* Controller = Cast<UAnimNextController>(EdGraph->GetController()))
 		{
-			FName ValidName = UE::AnimNext::Editor::FUtils::ValidateName(Cast<UObject>(Host), NodeName);
+			FName ValidName = UE::UAF::Editor::FUtils::ValidateName(Cast<UObject>(Host), NodeName);
 
 			Controller->OpenUndoBracket(FString::Printf(TEXT("Add '%s' Manifest Node"), *ValidName.ToString()));
 
diff --git a/Engine/Plugins/Experimental/AnimNextAnimGraph/Source/AnimNextAnimGraphUncookedOnly/Internal/Graph/AnimNextAnimationGraph_EditorData.h b/Engine/Plugins/Experimental/AnimNextAnimGraph/Source/AnimNextAnimGraphUncookedOnly/Internal/Graph/AnimNextAnimationGraph_EditorData.h
--- a/Engine/Plugins/Experimental/AnimNextAnimGraph/Source/AnimNextAnimGraphUncookedOnly/Internal/Graph/AnimNextAnimationGraph_EditorData.h
+++ b/Engine/Plugins/Experimental/AnimNextAnimGraph/Source/AnimNextAnimGraphUncookedOnly/Internal/Graph/AnimNextAnimationGraph_EditorData.h
@@ -15,13 +15,13 @@ class FAnimationAnimNextRuntimeTest_GraphAddTrait;
 class FAnimationAnimNextRuntimeTest_GraphExecute;
 class FAnimationAnimNextRuntimeTest_GraphExecuteLatent;
 
-namespace UE::AnimNext::UncookedOnly
+namespace UE::UAF::UncookedOnly
 {
 	struct FUtils;
 	struct FAnimGraphUtils;
 }
 
-namespace UE::AnimNext::Editor
+namespace UE::UAF::Editor
 {
 	class FModuleEditor;
 	class SAnimNextGraphView;
@@ -48,11 +48,11 @@ class UAnimNextAnimationGraph_EditorData : public UAnimNextSharedVariables_Edito
 
 	friend class UAnimNextAnimationGraphFactory;
 	friend class UAnimNextEdGraph;
-	friend struct UE::AnimNext::UncookedOnly::FUtils;
-	friend struct UE::AnimNext::Editor::FUtils;
-	friend class UE::AnimNext::Editor::FModuleEditor;
-	friend class UE::AnimNext::Editor::SAnimNextGraphView;
-	friend struct UE::AnimNext::UncookedOnly::FAnimGraphUtils;
+	friend struct UE::UAF::UncookedOnly::FUtils;
+	friend struct UE::UAF::Editor::FUtils;
+	friend class UE::UAF::Editor::FModuleEditor;
+	friend class UE::UAF::Editor::SAnimNextGraphView;
+	friend struct UE::UAF::UncookedOnly::FAnimGraphUtils;
 	friend struct FAnimNextGraphSchemaAction_RigUnit;
 	friend struct FAnimNextGraphSchemaAction_DispatchFactory;
 	friend class FAnimationAnimNextEditorTest_GraphAddTrait;
diff --git a/Engine/Plugins/Experimental/AnimNextAnimGraph/Source/AnimNextAnimGraphUncookedOnly/Private/AnimGraphUncookedOnlyUtils.cpp b/Engine/Plugins/Experimental/AnimNextAnimGraph/Source/AnimNextAnimGraphUncookedOnly/Private/AnimGraphUncookedOnlyUtils.cpp
--- a/Engine/Plugins/Experimental/AnimNextAnimGraph/Source/AnimNextAnimGraphUncookedOnly/Private/AnimGraphUncookedOnlyUtils.cpp
+++ b/Engine/Plugins/Experimental/AnimNextAnimGraph/Source/AnimNextAnimGraphUncookedOnly/Private/AnimGraphUncookedOnlyUtils.cpp
@@ -11,7 +11,7 @@
 #include "Graph/AnimNextAnimationGraph_EditorData.h"
 #include "AnimNextTraitStackUnitNode.h"
 
-namespace UE::AnimNext::UncookedOnly
+namespace UE::UAF::UncookedOnly
 {
 
 bool FAnimGraphUtils::IsTraitStackNode(const URigVMNode* InModelNode)
@@ -65,7 +65,7 @@ void FAnimGraphUtils::GetAssetManifestNodesRegistryExports(const UAnimNextRigVMA
 		{
 			if (const URigVMNode* ModelNode = EdNode->GetModelNode())
 			{
-				if (UE::AnimNext::UncookedOnly::FAnimGraphUtils::IsExposedToManifest(ModelNode))
+				if (UE::UAF::UncookedOnly::FAnimGraphUtils::IsExposedToManifest(ModelNode))
 				{
 					const UPackage* Package = EditorData->GetOuter()->GetPackage();
 					FString PackageName;
diff --git a/Engine/Plugins/Experimental/AnimNextAnimGraph/Source/AnimNextAnimGraphUncookedOnly/Private/AnimNextAnimGraphUncookedOnlyModule.cpp b/Engine/Plugins/Experimental/AnimNextAnimGraph/Source/AnimNextAnimGraphUncookedOnly/Private/AnimNextAnimGraphUncookedOnlyModule.cpp
--- a/Engine/Plugins/Experimental/AnimNextAnimGraph/Source/AnimNextAnimGraphUncookedOnly/Private/AnimNextAnimGraphUncookedOnlyModule.cpp
+++ b/Engine/Plugins/Experimental/AnimNextAnimGraph/Source/AnimNextAnimGraphUncookedOnly/Private/AnimNextAnimGraphUncookedOnlyModule.cpp
@@ -7,14 +7,14 @@
 
 #define LOCTEXT_NAMESPACE "FAnimNextAnimGraphUncookedOnlyModule"
 
-namespace UE::AnimNext::AnimGraph::UncookedOnly
+namespace UE::UAF::AnimGraph::UncookedOnly
 {
 
 class FAnimNextAnimGraphUncookedOnlyModule : public IModuleInterface
 {
 	virtual void StartupModule() override
 	{
-		using namespace UE::AnimNext::UncookedOnly;
+		using namespace UE::UAF::UncookedOnly;
 
 		FPublicVariablesImpl::Register();
 	}
@@ -24,4 +24,4 @@ class FAnimNextAnimGraphUncookedOnlyModule : public IModuleInterface
 
 #undef LOCTEXT_NAMESPACE
 
-IMPLEMENT_MODULE(UE::AnimNext::AnimGraph::UncookedOnly::FAnimNextAnimGraphUncookedOnlyModule, AnimNextAnimGraphUncookedOnly);
+IMPLEMENT_MODULE(UE::UAF::AnimGraph::UncookedOnly::FAnimNextAnimGraphUncookedOnlyModule, AnimNextAnimGraphUncookedOnly);
diff --git a/Engine/Plugins/Experimental/AnimNextAnimGraph/Source/AnimNextAnimGraphUncookedOnly/Private/AnimNextAnimGraphWorkspaceAssetUserData.cpp b/Engine/Plugins/Experimental/AnimNextAnimGraph/Source/AnimNextAnimGraphUncookedOnly/Private/AnimNextAnimGraphWorkspaceAssetUserData.cpp
--- a/Engine/Plugins/Experimental/AnimNextAnimGraph/Source/AnimNextAnimGraphUncookedOnly/Private/AnimNextAnimGraphWorkspaceAssetUserData.cpp
+++ b/Engine/Plugins/Experimental/AnimNextAnimGraph/Source/AnimNextAnimGraphUncookedOnly/Private/AnimNextAnimGraphWorkspaceAssetUserData.cpp
@@ -15,7 +15,7 @@ void UAnimNextAnimGraphWorkspaceAssetUserData::GetAssetRegistryTags(FAssetRegist
 		CachedExports.Exports.Reset();
 		{
 			UAnimNextAnimationGraph* Asset = CastChecked<UAnimNextAnimationGraph>(GetOuter());
-			const UAnimNextAnimationGraph_EditorData* EditorData = UE::AnimNext::UncookedOnly::FUtils::GetEditorData<UAnimNextAnimationGraph_EditorData>(Asset);
+			const UAnimNextAnimationGraph_EditorData* EditorData = UE::UAF::UncookedOnly::FUtils::GetEditorData<UAnimNextAnimationGraph_EditorData>(Asset);
 			{
 				FWorkspaceOutlinerItemExport& RootAssetExport = CachedExports.Exports.Add_GetRef(FWorkspaceOutlinerItemExport(Asset->GetFName(), Asset));
 				RootAssetExport.GetData().InitializeAsScriptStruct(FAnimNextAnimationGraphOutlinerData::StaticStruct());
@@ -23,7 +23,7 @@ void UAnimNextAnimGraphWorkspaceAssetUserData::GetAssetRegistryTags(FAssetRegist
 				Data.SoftAssetPtr = Asset;
 			}
 
-			UE::AnimNext::UncookedOnly::FUtils::GetAssetWorkspaceExports(EditorData, CachedExports, Context);
+			UE::UAF::UncookedOnly::FUtils::GetAssetWorkspaceExports(EditorData, CachedExports, Context);
 		}
 	}
 
diff --git a/Engine/Plugins/Experimental/AnimNextAnimGraph/Source/AnimNextAnimGraphUncookedOnly/Private/AnimNextController.cpp b/Engine/Plugins/Experimental/AnimNextAnimGraph/Source/AnimNextAnimGraphUncookedOnly/Private/AnimNextController.cpp
--- a/Engine/Plugins/Experimental/AnimNextAnimGraph/Source/AnimNextAnimGraphUncookedOnly/Private/AnimNextController.cpp
+++ b/Engine/Plugins/Experimental/AnimNextAnimGraph/Source/AnimNextAnimGraphUncookedOnly/Private/AnimNextController.cpp
@@ -41,9 +41,9 @@ FName UAnimNextController::AddTraitByName(FName InNodeName, FName InNewTraitType
 		return NAME_None;
 	}
 
-	UE::AnimNext::FTraitRegistry& TraitRegistry = UE::AnimNext::FTraitRegistry::Get();
+	UE::UAF::FTraitRegistry& TraitRegistry = UE::UAF::FTraitRegistry::Get();
 
-	const UE::AnimNext::FTrait* Trait = TraitRegistry.Find(InNewTraitTypeName);
+	const UE::UAF::FTrait* Trait = TraitRegistry.Find(InNewTraitTypeName);
 	if (Trait == nullptr)
 	{
 		ReportError(TEXT("Unknown Trait Type."));
@@ -177,8 +177,8 @@ bool UAnimNextController::SetTraitPinIndex(FName InNodeName, FName InTraitInstan
 		const FRigVMTrait* Trait = (FRigVMTrait*)ScopedTrait->GetStructMemory();
 		if (const UScriptStruct* TraitSharedInstanceData = Trait->GetTraitSharedDataStruct())
 		{
-			UE::AnimNext::FTraitRegistry& TraitRegistry = UE::AnimNext::FTraitRegistry::Get();
-			if (const UE::AnimNext::FTrait* AnimNextTrait = TraitRegistry.Find(TraitSharedInstanceData))
+			UE::UAF::FTraitRegistry& TraitRegistry = UE::UAF::FTraitRegistry::Get();
+			if (const UE::UAF::FTrait* AnimNextTrait = TraitRegistry.Find(TraitSharedInstanceData))
 			{
 				// Avoid multiple VM recompilations, for each operation
 				FRigVMControllerCompileBracketScope CompileScope(this);
@@ -265,7 +265,7 @@ bool UAnimNextController::AddNodeToManifest(URigVMNode* InModelNode, bool bSetup
 		return false;
 	}
 
-	if (!ensure(UE::AnimNext::UncookedOnly::FAnimGraphUtils::IsExposedToManifest(InModelNode) == false))
+	if (!ensure(UE::UAF::UncookedOnly::FAnimGraphUtils::IsExposedToManifest(InModelNode) == false))
 	{
 		return false;
 	}
@@ -347,7 +347,7 @@ bool UAnimNextController::RemoveNodeFromManifest(URigVMNode* InModelNode, bool b
 		return false;
 	}
 
-	if (!ensure(UE::AnimNext::UncookedOnly::FAnimGraphUtils::IsExposedToManifest(InModelNode)))
+	if (!ensure(UE::UAF::UncookedOnly::FAnimGraphUtils::IsExposedToManifest(InModelNode)))
 	{
 		return false;
 	}
diff --git a/Engine/Plugins/Experimental/AnimNextAnimGraph/Source/AnimNextAnimGraphUncookedOnly/Private/AnimNextControllerActions.cpp b/Engine/Plugins/Experimental/AnimNextAnimGraph/Source/AnimNextAnimGraphUncookedOnly/Private/AnimNextControllerActions.cpp
--- a/Engine/Plugins/Experimental/AnimNextAnimGraph/Source/AnimNextAnimGraphUncookedOnly/Private/AnimNextControllerActions.cpp
+++ b/Engine/Plugins/Experimental/AnimNextAnimGraph/Source/AnimNextAnimGraphUncookedOnly/Private/AnimNextControllerActions.cpp
@@ -11,7 +11,7 @@ FAnimNextManifestAction::FAnimNextManifestAction()
 FAnimNextManifestAction::FAnimNextManifestAction(UAnimNextController* InController, URigVMNode* InNode, bool bInNewIncludeInManifestState)
 	: FAnimNextBaseAction(InController)
 	, NodePath(InNode->GetNodePath())
-	, bOldManifestState(UE::AnimNext::UncookedOnly::FAnimGraphUtils::IsExposedToManifest(InNode))
+	, bOldManifestState(UE::UAF::UncookedOnly::FAnimGraphUtils::IsExposedToManifest(InNode))
 	, bIncludeInManifestState(bInNewIncludeInManifestState)
 {
 }
diff --git a/Engine/Plugins/Experimental/AnimNextAnimGraph/Source/AnimNextAnimGraphUncookedOnly/Private/Graph/AnimNextAnimationGraphSchema.cpp b/Engine/Plugins/Experimental/AnimNextAnimGraph/Source/AnimNextAnimGraphUncookedOnly/Private/Graph/AnimNextAnimationGraphSchema.cpp
--- a/Engine/Plugins/Experimental/AnimNextAnimGraph/Source/AnimNextAnimGraphUncookedOnly/Private/Graph/AnimNextAnimationGraphSchema.cpp
+++ b/Engine/Plugins/Experimental/AnimNextAnimGraph/Source/AnimNextAnimGraphUncookedOnly/Private/Graph/AnimNextAnimationGraphSchema.cpp
@@ -7,7 +7,7 @@ bool UAnimNextAnimationGraphSchema::CanSetNodeTitle(URigVMController* InControll
 {
 	if (InNode != nullptr)
 	{
-		if (UE::AnimNext::UncookedOnly::FAnimGraphUtils::IsTraitStackNode(InNode))
+		if (UE::UAF::UncookedOnly::FAnimGraphUtils::IsTraitStackNode(InNode))
 		{
 			return true;
 		}
diff --git a/Engine/Plugins/Experimental/AnimNextAnimGraph/Source/AnimNextAnimGraphUncookedOnly/Private/Graph/AnimNextAnimationGraph_EditorData.cpp b/Engine/Plugins/Experimental/AnimNextAnimGraph/Source/AnimNextAnimGraphUncookedOnly/Private/Graph/AnimNextAnimationGraph_EditorData.cpp
--- a/Engine/Plugins/Experimental/AnimNextAnimGraph/Source/AnimNextAnimGraphUncookedOnly/Private/Graph/AnimNextAnimationGraph_EditorData.cpp
+++ b/Engine/Plugins/Experimental/AnimNextAnimGraph/Source/AnimNextAnimGraphUncookedOnly/Private/Graph/AnimNextAnimationGraph_EditorData.cpp
@@ -33,7 +33,7 @@
 
 #define LOCTEXT_NAMESPACE "AnimNextAnimationGraph_EditorData"
 
-namespace UE::AnimNext::UncookedOnly::Private
+namespace UE::UAF::UncookedOnly::Private
 {
 	// Represents a trait entry on a node
 	struct FTraitEntryMapping
@@ -278,7 +278,7 @@ namespace UE::AnimNext::UncookedOnly::Private
 
 			FString DefaultValue;
 			{
-				constexpr UE::AnimNext::FTraitUID ReferencePoseTraitUID = UE::AnimNext::FTraitUID::MakeUID(TEXT("FReferencePoseTrait"));	// Trait header is private, reference by UID directly
+				constexpr UE::UAF::FTraitUID ReferencePoseTraitUID = UE::UAF::FTraitUID::MakeUID(TEXT("FReferencePoseTrait"));	// Trait header is private, reference by UID directly
 				const FTrait* Trait = FTraitRegistry::Get().Find(ReferencePoseTraitUID);
 				check(Trait != nullptr);
 
@@ -382,8 +382,8 @@ namespace UE::AnimNext::UncookedOnly::Private
 			URigVMUnitNode* VMNode = VMController->AddUnitNode(FRigUnit_AnimNextTraitStack::StaticStruct(), FRigVMStruct::ExecuteName, FVector2D(0.0f, 0.0f), FString(), false);
 
 			UAnimNextController* AnimNextController = CastChecked<UAnimNextController>(VMController);
-			constexpr UE::AnimNext::FTraitUID ReferencePoseTraitUID = UE::AnimNext::FTraitUID::MakeUID(TEXT("FReferencePoseTrait"));	// Trait header is private, reference by UID directly
-			const FName RigVMTraitName =  AnimNextController->AddTraitByName(VMNode->GetFName(), *UE::AnimNext::FTraitRegistry::Get().Find(ReferencePoseTraitUID)->GetTraitName(), INDEX_NONE, TEXT(""), false);
+			constexpr UE::UAF::FTraitUID ReferencePoseTraitUID = UE::UAF::FTraitUID::MakeUID(TEXT("FReferencePoseTrait"));	// Trait header is private, reference by UID directly
+			const FName RigVMTraitName =  AnimNextController->AddTraitByName(VMNode->GetFName(), *UE::UAF::FTraitRegistry::Get().Find(ReferencePoseTraitUID)->GetTraitName(), INDEX_NONE, TEXT(""), false);
 		
 			check(RigVMTraitName != NAME_None);
 
@@ -493,7 +493,7 @@ namespace UE::AnimNext::UncookedOnly::Private
 
 void UAnimNextAnimationGraph_EditorData::OnPreCompileAsset(FRigVMCompileSettings& InSettings)
 {
-	using namespace UE::AnimNext::UncookedOnly;
+	using namespace UE::UAF::UncookedOnly;
 
 	InSettings.ASTSettings.bSetupTraits = false; // disable the default implementation of decorators for now
 
@@ -511,8 +511,8 @@ void UAnimNextAnimationGraph_EditorData::OnPreCompileAsset(FRigVMCompileSettings
 
 void UAnimNextAnimationGraph_EditorData::OnPreCompileGetProgrammaticFunctionHeaders(const FRigVMCompileSettings& InSettings, FAnimNextGetFunctionHeaderCompileContext& OutCompileContext)
 {
-	using namespace UE::AnimNext::UncookedOnly;
-	using namespace UE::AnimNext::UncookedOnly::Private;
+	using namespace UE::UAF::UncookedOnly;
+	using namespace UE::UAF::UncookedOnly::Private;
 
 	Super::OnPreCompileGetProgrammaticFunctionHeaders(InSettings, OutCompileContext);
 
@@ -568,8 +568,8 @@ void UAnimNextAnimationGraph_EditorData::OnPreCompileGetProgrammaticFunctionHead
 
 void UAnimNextAnimationGraph_EditorData::OnPreCompileProcessGraphs(const FRigVMCompileSettings& InSettings, FAnimNextProcessGraphCompileContext& OutCompileContext)
 {
-	using namespace UE::AnimNext::UncookedOnly;
-	using namespace UE::AnimNext::UncookedOnly::Private;
+	using namespace UE::UAF::UncookedOnly;
+	using namespace UE::UAF::UncookedOnly::Private;
 
 	FRigVMClient* VMClient = GetRigVMClient();
 	UAnimNextAnimationGraph* AnimationGraph = FUtils::GetAsset<UAnimNextAnimationGraph>(this);
@@ -605,7 +605,7 @@ void UAnimNextAnimationGraph_EditorData::OnPreCompileProcessGraphs(const FRigVMC
 	{
 		UAnimNextController* TempController = CastChecked<UAnimNextController>(VMClient->GetOrCreateController(AnimGraphs[0]));
 
-		UE::AnimNext::FTraitWriter TraitWriter;
+		UE::UAF::FTraitWriter TraitWriter;
 
 		FRigVMPinInfoArray LatentPins;
 		TMap<FName, URigVMPin*> LatentPinMapping;
@@ -703,7 +703,7 @@ void UAnimNextAnimationGraph_EditorData::OnPreCompileProcessGraphs(const FRigVMC
 
 void UAnimNextAnimationGraph_EditorData::OnCompileJobStarted()
 {
-	using namespace UE::AnimNext::UncookedOnly;
+	using namespace UE::UAF::UncookedOnly;
 
 	UAnimNextAnimationGraph* AnimationGraph = FUtils::GetAsset<UAnimNextAnimationGraph>(this);
 
@@ -714,7 +714,7 @@ void UAnimNextAnimationGraph_EditorData::OnCompileJobStarted()
 
 void UAnimNextAnimationGraph_EditorData::OnCompileJobFinished()
 {
-	using namespace UE::AnimNext::UncookedOnly;
+	using namespace UE::UAF::UncookedOnly;
 
 	UAnimNextAnimationGraph* AnimationGraph = FUtils::GetAsset<UAnimNextAnimationGraph>(this);
 
@@ -761,7 +761,7 @@ bool UAnimNextAnimationGraph_EditorData::CanAddNewEntry(TSubclassOf<UAnimNextRig
 
 UAnimNextAnimationGraphEntry* UAnimNextAnimationGraphLibrary::AddAnimationGraph(UAnimNextAnimationGraph* InAsset, FName InName, bool bSetupUndoRedo, bool bPrintPythonCommand)
 {
-	return UE::AnimNext::UncookedOnly::FUtils::GetEditorData<UAnimNextAnimationGraph_EditorData>(InAsset)->AddAnimationGraph(InName, bSetupUndoRedo, bPrintPythonCommand);
+	return UE::UAF::UncookedOnly::FUtils::GetEditorData<UAnimNextAnimationGraph_EditorData>(InAsset)->AddAnimationGraph(InName, bSetupUndoRedo, bPrintPythonCommand);
 }
 
 UAnimNextAnimationGraphEntry* UAnimNextAnimationGraph_EditorData::AddAnimationGraph(FName InName, bool bSetupUndoRedo, bool bPrintPythonCommand)
@@ -824,7 +824,7 @@ UAnimNextAnimationGraphEntry* UAnimNextAnimationGraph_EditorData::AddAnimationGr
 			RefreshExternalModels();
 			RigVMClient.AddModel(NewRigVMGraphModel, true);
 			URigVMController* Controller = RigVMClient.GetController(NewRigVMGraphModel);
-			UE::AnimNext::UncookedOnly::FAnimGraphUtils::SetupAnimGraph(NewEntry->GetEntryName(), Controller);
+			UE::UAF::UncookedOnly::FAnimGraphUtils::SetupAnimGraph(NewEntry->GetEntryName(), Controller);
 		}
 	}
 
@@ -867,7 +867,7 @@ void UAnimNextAnimationGraph_EditorData::InitializeAssetUserData()
 
 void UAnimNextAnimationGraph_EditorData::GetAnimNextAssetRegistryTags(FAssetRegistryTagsContext& Context, FAnimNextAssetRegistryExports& OutExports) const
 {
-	UE::AnimNext::UncookedOnly::FAnimGraphUtils::GetAssetManifestNodesRegistryExports(this, OutExports);
+	UE::UAF::UncookedOnly::FAnimGraphUtils::GetAssetManifestNodesRegistryExports(this, OutExports);
 	UAnimNextRigVMAssetEditorData::GetAnimNextAssetRegistryTags(Context, OutExports);
 }
 
diff --git a/Engine/Plugins/Experimental/AnimNextAnimGraph/Source/AnimNextAnimGraphUncookedOnly/Private/Graph/RigVMTrait_AnimNextPublicVariablesUncookedOnly.cpp b/Engine/Plugins/Experimental/AnimNextAnimGraph/Source/AnimNextAnimGraphUncookedOnly/Private/Graph/RigVMTrait_AnimNextPublicVariablesUncookedOnly.cpp
--- a/Engine/Plugins/Experimental/AnimNextAnimGraph/Source/AnimNextAnimGraphUncookedOnly/Private/Graph/RigVMTrait_AnimNextPublicVariablesUncookedOnly.cpp
+++ b/Engine/Plugins/Experimental/AnimNextAnimGraph/Source/AnimNextAnimGraphUncookedOnly/Private/Graph/RigVMTrait_AnimNextPublicVariablesUncookedOnly.cpp
@@ -9,7 +9,7 @@
 #include "RigVMModel/RigVMPin.h"
 #include "Variables/AnimNextSharedVariables.h"
 
-namespace UE::AnimNext::UncookedOnly
+namespace UE::UAF::UncookedOnly
 {
 
 void FPublicVariablesImpl::Register()
@@ -33,7 +33,7 @@ void FPublicVariablesImpl::GetProgrammaticPins(const FRigVMTrait_AnimNextPublicV
 		return;
 	}
 
-	UAnimNextRigVMAssetEditorData* EditorData = UE::AnimNext::UncookedOnly::FUtils::GetEditorData<UAnimNextRigVMAssetEditorData>(InTrait.InternalAsset.Get());
+	UAnimNextRigVMAssetEditorData* EditorData = UE::UAF::UncookedOnly::FUtils::GetEditorData<UAnimNextRigVMAssetEditorData>(InTrait.InternalAsset.Get());
 	if(EditorData == nullptr)
 	{
 		return;
diff --git a/Engine/Plugins/Experimental/AnimNextChooser/Source/AnimNextChooser/Private/ChooserParameters.cpp b/Engine/Plugins/Experimental/AnimNextChooser/Source/AnimNextChooser/Private/ChooserParameters.cpp
--- a/Engine/Plugins/Experimental/AnimNextChooser/Source/AnimNextChooser/Private/ChooserParameters.cpp
+++ b/Engine/Plugins/Experimental/AnimNextChooser/Source/AnimNextChooser/Private/ChooserParameters.cpp
@@ -5,7 +5,7 @@
 #include "StructUtils/PropertyBag.h"
 #include "UObject/FortniteMainBranchObjectVersion.h"
 
-namespace UE::AnimNext::Private
+namespace UE::UAF::Private
 {
 
 // Helper function to find the first data interface instance in the context 
@@ -26,7 +26,7 @@ FUAFAssetInstance* GetFirstDataInterfaceInstance(FChooserEvaluationContext& Cont
 
 bool FBoolAnimProperty::GetValue(FChooserEvaluationContext& Context, bool& OutResult) const
 {
-	if(FUAFAssetInstance* Instance = UE::AnimNext::Private::GetFirstDataInterfaceInstance(Context))
+	if(FUAFAssetInstance* Instance = UE::UAF::Private::GetFirstDataInterfaceInstance(Context))
 	{
 		return Instance->GetVariable(Variable, OutResult) == EPropertyBagResult::Success;
 	}
@@ -35,7 +35,7 @@ bool FBoolAnimProperty::GetValue(FChooserEvaluationContext& Context, bool& OutRe
 
 bool FBoolAnimProperty::SetValue(FChooserEvaluationContext& Context, bool InValue) const
 {
-	if(FUAFAssetInstance* Instance = UE::AnimNext::Private::GetFirstDataInterfaceInstance(Context))
+	if(FUAFAssetInstance* Instance = UE::UAF::Private::GetFirstDataInterfaceInstance(Context))
 	{
 		return Instance->SetVariable(Variable, InValue) == EPropertyBagResult::Success;
 	}
@@ -75,7 +75,7 @@ void FBoolAnimProperty::PostSerialize(const FArchive& Ar)
 
 bool FFloatAnimProperty::GetValue(FChooserEvaluationContext& Context, double& OutResult) const
 {
-	if(FUAFAssetInstance* Instance = UE::AnimNext::Private::GetFirstDataInterfaceInstance(Context))
+	if(FUAFAssetInstance* Instance = UE::UAF::Private::GetFirstDataInterfaceInstance(Context))
 	{
 		return Instance->GetVariable(Variable, OutResult) == EPropertyBagResult::Success;
 	}
@@ -84,7 +84,7 @@ bool FFloatAnimProperty::GetValue(FChooserEvaluationContext& Context, double& Ou
 
 bool FFloatAnimProperty::SetValue(FChooserEvaluationContext& Context, double InValue) const
 {
-	if(FUAFAssetInstance* Instance = UE::AnimNext::Private::GetFirstDataInterfaceInstance(Context))
+	if(FUAFAssetInstance* Instance = UE::UAF::Private::GetFirstDataInterfaceInstance(Context))
  	{
  		return Instance->SetVariable(Variable, InValue) == EPropertyBagResult::Success;
  	}
@@ -124,7 +124,7 @@ void FFloatAnimProperty::PostSerialize(const FArchive& Ar)
 
 bool FEnumAnimProperty::GetValue(FChooserEvaluationContext& Context, uint8& OutResult) const
 {
-	if(FUAFAssetInstance* Instance = UE::AnimNext::Private::GetFirstDataInterfaceInstance(Context))
+	if(FUAFAssetInstance* Instance = UE::UAF::Private::GetFirstDataInterfaceInstance(Context))
 	{
 		return Instance->GetVariable(Variable, OutResult) == EPropertyBagResult::Success;
 	}
@@ -133,7 +133,7 @@ bool FEnumAnimProperty::GetValue(FChooserEvaluationContext& Context, uint8& OutR
 
 bool FEnumAnimProperty::SetValue(FChooserEvaluationContext& Context, uint8 InValue) const
 {
-	if(FUAFAssetInstance* Instance = UE::AnimNext::Private::GetFirstDataInterfaceInstance(Context))
+	if(FUAFAssetInstance* Instance = UE::UAF::Private::GetFirstDataInterfaceInstance(Context))
 	{
 		return Instance->SetVariable(Variable, InValue) == EPropertyBagResult::Success;
 	}
diff --git a/Engine/Plugins/Experimental/AnimNextChooser/Source/AnimNextChooser/Private/Module.cpp b/Engine/Plugins/Experimental/AnimNextChooser/Source/AnimNextChooser/Private/Module.cpp
--- a/Engine/Plugins/Experimental/AnimNextChooser/Source/AnimNextChooser/Private/Module.cpp
+++ b/Engine/Plugins/Experimental/AnimNextChooser/Source/AnimNextChooser/Private/Module.cpp
@@ -8,7 +8,7 @@
 #include "Animation/AnimSequence.h"
 #include "Chooser.h"
 
-namespace UE::AnimNext::Chooser
+namespace UE::UAF::Chooser
 {
 
 class FModule : public IModuleInterface
@@ -32,4 +32,4 @@ public:
 
 }
 
-IMPLEMENT_MODULE(UE::AnimNext::Chooser::FModule, AnimNextChooser)
+IMPLEMENT_MODULE(UE::UAF::Chooser::FModule, AnimNextChooser)
diff --git a/Engine/Plugins/Experimental/AnimNextChooser/Source/AnimNextChooser/Private/RigUnit_EvaluateChooser.h b/Engine/Plugins/Experimental/AnimNextChooser/Source/AnimNextChooser/Private/RigUnit_EvaluateChooser.h
--- a/Engine/Plugins/Experimental/AnimNextChooser/Source/AnimNextChooser/Private/RigUnit_EvaluateChooser.h
+++ b/Engine/Plugins/Experimental/AnimNextChooser/Source/AnimNextChooser/Private/RigUnit_EvaluateChooser.h
@@ -25,7 +25,7 @@ struct FRigVMDispatch_EvaluateChooser : public FRigVMDispatchFactory
 	FRigVMDispatch_EvaluateChooser();
 	
 private:
-	friend struct UE::AnimNext::UncookedOnly::FUtils;
+	friend struct UE::UAF::UncookedOnly::FUtils;
 
 	virtual UScriptStruct* GetExecuteContextStruct() const { return FRigVMExecuteContext::StaticStruct(); }
 	virtual FName GetArgumentNameForOperandIndex(int32 InOperandIndex, int32 InTotalOperands, FRigVMRegistryHandle& InRegistry) const override;
diff --git a/Engine/Plugins/Experimental/AnimNextControlRig/Source/AnimNextControlRig/Internal/AnimNextControlRigHierarchyMappings.h b/Engine/Plugins/Experimental/AnimNextControlRig/Source/AnimNextControlRig/Internal/AnimNextControlRigHierarchyMappings.h
--- a/Engine/Plugins/Experimental/AnimNextControlRig/Source/AnimNextControlRig/Internal/AnimNextControlRigHierarchyMappings.h
+++ b/Engine/Plugins/Experimental/AnimNextControlRig/Source/AnimNextControlRig/Internal/AnimNextControlRigHierarchyMappings.h
@@ -15,12 +15,12 @@ struct FPoseContext;
 struct FControlRigIOSettings;
 struct FCompactPose;
 
-namespace UE::AnimNext
+namespace UE::UAF
 {
 struct FKeyframeState;
 }
 
-namespace UE::AnimNext::ControlRig
+namespace UE::UAF::ControlRig
 {
 
 class FAnimNextControlRigPoseAdapter;
@@ -46,13 +46,13 @@ struct ANIMNEXTCONTROLRIG_API FAnimNextControlRigHierarchyMappings
 	void UpdateControlRigRefPoseIfNeeded(UControlRig* ControlRig
 		, const UObject* InstanceObject
 		, const USkeletalMeshComponent* SkeletalMeshComponent
-		, const UE::AnimNext::FReferencePose& InRefPose
+		, const UE::UAF::FReferencePose& InRefPose
 		, bool bInSetRefPoseFromSkeleton
 		, bool bIncludePoseInHash);
 
 	void UpdateInputOutputMappingIfRequired(UControlRig* InControlRig
 		, URigHierarchy* InHierarchy
-		, const UE::AnimNext::FReferencePose& InRefPose
+		, const UE::UAF::FReferencePose& InRefPose
 		, int32 InCurrentLOD
 		, const TArray<FBoneReference>& InInputBonesToTransfer
 		, const TArray<FBoneReference>& InOutputBonesToTransfer
@@ -61,7 +61,7 @@ struct ANIMNEXTCONTROLRIG_API FAnimNextControlRigHierarchyMappings
 		, bool bResetInputPoseToInitial);
 
 	void UpdateInput(UControlRig* ControlRig
-		, UE::AnimNext::FKeyframeState& InOutput
+		, UE::UAF::FKeyframeState& InOutput
 		, const FControlRigIOSettings& InInputSettings
 		, const FControlRigIOSettings& InOutputSettings
 		, TWeakObjectPtr<UNodeMappingContainer>& InNodeMappingContainer
@@ -72,7 +72,7 @@ struct ANIMNEXTCONTROLRIG_API FAnimNextControlRigHierarchyMappings
 		, bool bInTransferInputCurves);
 
 	void UpdateOutput(UControlRig* ControlRig
-		, UE::AnimNext::FKeyframeState& InOutput
+		, UE::UAF::FKeyframeState& InOutput
 		, const FControlRigIOSettings& InOutputSettings
 		, TWeakObjectPtr<UNodeMappingContainer>& InNodeMappingContainer
 		, bool bInExecute
@@ -104,7 +104,7 @@ struct ANIMNEXTCONTROLRIG_API FAnimNextControlRigHierarchyMappings
 
 	void PerformUpdateToDate(UControlRig* ControlRig
 		, URigHierarchy* InHierarchy
-		, const UE::AnimNext::FReferencePose& InRefPose
+		, const UE::UAF::FReferencePose& InRefPose
 		, int32 InCurrentLOD
 		, const TWeakObjectPtr<UNodeMappingContainer>& InNodeMappingContainer
 		, bool bInTransferPoseInGlobalSpace
@@ -133,4 +133,4 @@ protected:
 	TOptional<int32> RefPoseSetterHash;
 };
 
-} // end namespace UE::AnimNext
+} // end namespace UE::UAF
diff --git a/Engine/Plugins/Experimental/AnimNextControlRig/Source/AnimNextControlRig/Internal/AnimNextControlRigModule.h b/Engine/Plugins/Experimental/AnimNextControlRig/Source/AnimNextControlRig/Internal/AnimNextControlRigModule.h
--- a/Engine/Plugins/Experimental/AnimNextControlRig/Source/AnimNextControlRig/Internal/AnimNextControlRigModule.h
+++ b/Engine/Plugins/Experimental/AnimNextControlRig/Source/AnimNextControlRig/Internal/AnimNextControlRigModule.h
@@ -5,7 +5,7 @@
 #include "Modules/ModuleInterface.h"
 #include "UObject/UObjectGlobals.h"
 
-namespace UE::AnimNext::ControlRig
+namespace UE::UAF::ControlRig
 {
 
 class ANIMNEXTCONTROLRIG_API FAnimNextControlRigModule : public IModuleInterface
diff --git a/Engine/Plugins/Experimental/AnimNextControlRig/Source/AnimNextControlRig/Internal/AnimNextControlRigPoseAdapter.h b/Engine/Plugins/Experimental/AnimNextControlRig/Source/AnimNextControlRig/Internal/AnimNextControlRigPoseAdapter.h
--- a/Engine/Plugins/Experimental/AnimNextControlRig/Source/AnimNextControlRig/Internal/AnimNextControlRigPoseAdapter.h
+++ b/Engine/Plugins/Experimental/AnimNextControlRig/Source/AnimNextControlRig/Internal/AnimNextControlRigPoseAdapter.h
@@ -12,12 +12,12 @@ class UNodeMappingContainer;
 class UControlRig;
 class URigHierarchy;
 
-namespace UE::AnimNext
+namespace UE::UAF
 {
 struct FKeyframeState;
 }
 
-namespace UE::AnimNext::ControlRig
+namespace UE::UAF::ControlRig
 {
 
 class ANIMNEXTCONTROLRIG_API FAnimNextControlRigPoseAdapter : public FControlRigPoseAdapter
@@ -32,13 +32,13 @@ public:
 
 	void UpdateInputOutputMappingIfRequired(UControlRig* InControlRig
 		, URigHierarchy* InHierarchy
-		, const UE::AnimNext::FReferencePose& InRefPose
+		, const UE::UAF::FReferencePose& InRefPose
 		, const int32 InCurentLOD
 		, const TWeakObjectPtr<UNodeMappingContainer>& InNodeMappingContainer
 		, bool bInTransferPoseInGlobalSpace
 		, bool bInResetInputPoseToInitial);
 
-	static const FReferenceSkeleton* GetReferenceSkeleton(const UE::AnimNext::FReferencePose& InRefPose);
+	static const FReferenceSkeleton* GetReferenceSkeleton(const UE::UAF::FReferencePose& InRefPose);
 };
 
-} // end namespace UE::AnimNext
+} // end namespace UE::UAF
